Horde: Add name field to service accounts and use that in audit logs

Previously, name of service accounts resolved to "unknown".

diff --git a/Engine/Source/Programs/Horde/HordeServer.Tests/Authentication/ServiceAccountAuthTest.cs b/Engine/Source/Programs/Horde/HordeServer.Tests/Authentication/ServiceAccountAuthTest.cs
--- a/Engine/Source/Programs/Horde/HordeServer.Tests/Authentication/ServiceAccountAuthTest.cs
+++ b/Engine/Source/Programs/Horde/HordeServer.Tests/Authentication/ServiceAccountAuthTest.cs
@@ -35,7 +35,7 @@ namespace HordeServer.Tests.Authentication
 		public async Task InitializeAsync()
 		{
 			IServiceAccountCollection serviceAccountCollection = ServiceProvider.GetRequiredService<IServiceAccountCollection>();
-			(_, _serviceAccountToken) = await serviceAccountCollection.CreateAsync(new CreateServiceAccountOptions("myDesc",
+			(_, _serviceAccountToken) = await serviceAccountCollection.CreateAsync(new CreateServiceAccountOptions("myName", "myDesc",
 				Claims: new List<IUserClaim> { new UserClaim("myClaim", "myValue"), new UserClaim("foo", "bar") })
 				);
 		}
diff --git a/Engine/Source/Programs/Horde/HordeServer.Tests/ServiceAccounts/ServiceAccountCollectionTests.cs b/Engine/Source/Programs/Horde/HordeServer.Tests/ServiceAccounts/ServiceAccountCollectionTests.cs
--- a/Engine/Source/Programs/Horde/HordeServer.Tests/ServiceAccounts/ServiceAccountCollectionTests.cs
+++ b/Engine/Source/Programs/Horde/HordeServer.Tests/ServiceAccounts/ServiceAccountCollectionTests.cs
@@ -2,13 +2,12 @@
 
 using System.Collections.Generic;
 using System.Threading.Tasks;
-using HordeServer.Accounts;
 using HordeServer.Server;
 using HordeServer.ServiceAccounts;
 using HordeServer.Users;
 using Microsoft.VisualStudio.TestTools.UnitTesting;
 
-namespace HordeServer.Tests.Accounts
+namespace HordeServer.Tests.ServiceAccounts
 {
 	[TestClass]
 	public class HordeAccountCollectionTests : DatabaseIntegrationTest
@@ -21,20 +20,19 @@ namespace HordeServer.Tests.Accounts
 		{
 			IMongoService mongoService = GetMongoServiceSingleton();
 			_serviceAccountCollection = new ServiceAccountCollection(mongoService);
-			(_serviceAccount, _token) = _serviceAccountCollection.CreateAsync(new CreateServiceAccountOptions(Description: "myDesc",
-				Claims: new List<IUserClaim> { new UserClaim("myClaim", "myValue") })
-				).Result;
+			CreateServiceAccountOptions options = new (Name: "myName", Description: "myDesc", Claims: new List<IUserClaim> { new UserClaim("myClaim", "myValue") });
+			(_serviceAccount, _token) = _serviceAccountCollection.CreateAsync(options).Result;
 		}
 
 		[TestMethod]
 		public async Task AddAsync()
 		{
-			(IServiceAccount sa, _) = await _serviceAccountCollection.CreateAsync(new CreateServiceAccountOptions(Description: "myDesc",
-				Claims: new List<IUserClaim> { new UserClaim("myClaim", "myValue") }
-				));
+			CreateServiceAccountOptions options = new (Name: "myName", Description: "myDesc", Claims: new List<IUserClaim> { new UserClaim("myClaim", "myValue") });
+			(IServiceAccount sa, _) = await _serviceAccountCollection.CreateAsync(options);
 			Assert.AreEqual(1, sa.Claims.Count);
 			Assert.AreEqual("myValue", sa.Claims[0].Value);
 			Assert.IsTrue(sa.Enabled);
+			Assert.AreEqual("myName", sa.Name);
 			Assert.AreEqual("myDesc", sa.Description);
 		}
 
@@ -58,12 +56,13 @@ namespace HordeServer.Tests.Accounts
 			IServiceAccount? account = await _serviceAccountCollection.GetAsync(_serviceAccount.Id);
 			Assert.IsNotNull(account);
 
-			List<UserClaim> newClaims = new() { new UserClaim("newClaim1", "newValue1"), new UserClaim("newClaim2", "newValue2") };
+			List<UserClaim> newClaims = [new UserClaim("newClaim1", "newValue1"), new UserClaim("newClaim2", "newValue2")];
 			(_, string? newToken) = await account.UpdateAsync(new UpdateServiceAccountOptions
 			{
 				Claims = newClaims,
 				ResetToken = true,
 				Enabled = false,
+				Name = "newName",
 				Description = "newDesc"
 			});
 
@@ -81,6 +80,7 @@ namespace HordeServer.Tests.Accounts
 			Assert.AreEqual("newValue1", sa.Claims[0].Value);
 			Assert.AreEqual("newValue2", sa.Claims[1].Value);
 			Assert.AreEqual(false, sa.Enabled);
+			Assert.AreEqual("newName", sa.Name);
 			Assert.AreEqual("newDesc", sa.Description);
 		}
 
diff --git a/Engine/Source/Programs/Horde/HordeServer/Authentication/ServiceAccountAuthHandler.cs b/Engine/Source/Programs/Horde/HordeServer/Authentication/ServiceAccountAuthHandler.cs
--- a/Engine/Source/Programs/Horde/HordeServer/Authentication/ServiceAccountAuthHandler.cs
+++ b/Engine/Source/Programs/Horde/HordeServer/Authentication/ServiceAccountAuthHandler.cs
@@ -7,6 +7,7 @@ using System.Security.Claims;
 using System.Text.Encodings.Web;
 using System.Threading.Tasks;
 using HordeServer.ServiceAccounts;
+using HordeServer.Utilities;
 using Microsoft.AspNetCore.Authentication;
 using Microsoft.Extensions.Logging;
 using Microsoft.Extensions.Options;
@@ -59,7 +60,8 @@ namespace HordeServer.Authentication
 			}
 
 			List<Claim> claims = new List<Claim>(10);
-			claims.Add(new Claim(ClaimTypes.Name, AuthenticationScheme));
+			claims.Add(new Claim(ClaimTypes.Name, serviceAccount.Name));
+			claims.Add(new Claim(HordeClaimTypes.User, serviceAccount.Name));
 			claims.AddRange(serviceAccount.Claims.Select(claimPair => new Claim(claimPair.Type, claimPair.Value)));
 
 			ClaimsIdentity identity = new ClaimsIdentity(claims, Scheme.Name);
diff --git a/Engine/Source/Programs/Horde/HordeServer/ServiceAccounts/IServiceAccount.cs b/Engine/Source/Programs/Horde/HordeServer/ServiceAccounts/IServiceAccount.cs
--- a/Engine/Source/Programs/Horde/HordeServer/ServiceAccounts/IServiceAccount.cs
+++ b/Engine/Source/Programs/Horde/HordeServer/ServiceAccounts/IServiceAccount.cs
@@ -21,6 +21,11 @@ namespace HordeServer.ServiceAccounts
 		/// Unique internal ID for this Horde account
 		/// </summary>
 		ServiceAccountId Id { get; }
+		
+		/// <summary>
+		/// Human-readable name for identifying this account
+		/// </summary>
+		string Name { get; }
 
 		/// <summary>
 		/// Description of the account (who is it for, is there an owner etc)
@@ -57,11 +62,13 @@ namespace HordeServer.ServiceAccounts
 	/// <summary>
 	/// Options for updating an account
 	/// </summary>
-	/// <param name="Claims">If set, claims to update</param>
+	/// <param name="Name">If set, name to update</param>
 	/// <param name="Description">If set, description to update</param>
+	/// <param name="Claims">If set, claims to update</param>
 	/// <param name="ResetToken">Whether to reset the secret token for this account</param>
 	/// <param name="Enabled">If set, enabled flag to update</param>
-	public record class UpdateServiceAccountOptions(
+	public record UpdateServiceAccountOptions(
+		string? Name = null,
 		string? Description = null,
 		IReadOnlyList<IUserClaim>? Claims = null,
 		bool? ResetToken = null,
diff --git a/Engine/Source/Programs/Horde/HordeServer/ServiceAccounts/IServiceAccountCollection.cs b/Engine/Source/Programs/Horde/HordeServer/ServiceAccounts/IServiceAccountCollection.cs
--- a/Engine/Source/Programs/Horde/HordeServer/ServiceAccounts/IServiceAccountCollection.cs
+++ b/Engine/Source/Programs/Horde/HordeServer/ServiceAccounts/IServiceAccountCollection.cs
@@ -57,8 +57,9 @@ namespace HordeServer.ServiceAccounts
 	/// <summary>
 	/// Options for creating a new service account
 	/// </summary>
+	/// <param name="Name">Name of the account</param>
 	/// <param name="Description">Description for the account</param>
 	/// <param name="Claims">Optional list of claims</param>
 	/// <param name="Enabled">Whether the account should be enabled</param>
-	public record class CreateServiceAccountOptions(string Description, IReadOnlyList<IUserClaim>? Claims = null, bool? Enabled = null);
+	public record CreateServiceAccountOptions(string Name, string Description, IReadOnlyList<IUserClaim>? Claims = null, bool? Enabled = null);
 }
diff --git a/Engine/Source/Programs/Horde/HordeServer/ServiceAccounts/ServiceAccountCollection.cs b/Engine/Source/Programs/Horde/HordeServer/ServiceAccounts/ServiceAccountCollection.cs
--- a/Engine/Source/Programs/Horde/HordeServer/ServiceAccounts/ServiceAccountCollection.cs
+++ b/Engine/Source/Programs/Horde/HordeServer/ServiceAccounts/ServiceAccountCollection.cs
@@ -44,6 +44,8 @@ namespace HordeServer.ServiceAccounts
 			public ServiceAccountId Id { get; set; }
 
 			/// <inheritdoc/>
+			public string Name { get; set; } = "<not-set>";
+			
 			public string SecretToken { get; set; } = String.Empty;
 
 			[BsonElement("Claims2")]
@@ -65,9 +67,10 @@ namespace HordeServer.ServiceAccounts
 			{
 			}
 
-			public ServiceAccountDocument(ServiceAccountId id, string description)
+			public ServiceAccountDocument(ServiceAccountId id, string name, string description)
 			{
 				Id = id;
+				Name = name;
 				Description = description;
 			}
 
@@ -98,16 +101,13 @@ namespace HordeServer.ServiceAccounts
 			=> StringUtils.FormatHexString(RandomNumberGenerator.GetBytes(32));
 
 		/// <inheritdoc/>
-		public async Task<(IServiceAccount, string)> CreateAsync(
-			CreateServiceAccountOptions options,
-			CancellationToken cancellationToken = default)
+		public async Task<(IServiceAccount, string)> CreateAsync(CreateServiceAccountOptions options, CancellationToken cancellationToken = default)
 		{
 			string newToken = CreateToken();
-
-			ServiceAccountDocument account = new ServiceAccountDocument(new ServiceAccountId(BinaryIdUtils.CreateNew()), options.Description ?? "")
+			ServiceAccountDocument account = new (new ServiceAccountId(BinaryIdUtils.CreateNew()), options.Name, options.Description)
 			{
 				SecretToken = newToken,
-				Enabled = options.Enabled ?? true
+				Enabled = options.Enabled ?? true,
 			};
 
 			if (options.Claims != null)
@@ -155,7 +155,11 @@ namespace HordeServer.ServiceAccounts
 		{
 			string? newToken = null;
 			UpdateDefinition<ServiceAccountDocument> update = Builders<ServiceAccountDocument>.Update.Set(x => x.UpdateIndex, document.UpdateIndex + 1);
-
+			
+			if (options.Name != null)
+			{
+				update = update.Set(x => x.Name, options.Name);
+			}
 			if (options.Description != null)
 			{
 				update = update.Set(x => x.Description, options.Description);
diff --git a/Engine/Source/Programs/Horde/HordeServer/ServiceAccounts/ServiceAccountsController.cs b/Engine/Source/Programs/Horde/HordeServer/ServiceAccounts/ServiceAccountsController.cs
--- a/Engine/Source/Programs/Horde/HordeServer/ServiceAccounts/ServiceAccountsController.cs
+++ b/Engine/Source/Programs/Horde/HordeServer/ServiceAccounts/ServiceAccountsController.cs
@@ -48,7 +48,8 @@ namespace HordeServer.ServiceAccounts
 			}
 
 			List<IUserClaim> claims = request.Claims.ConvertAll<IUserClaim>(x => new UserClaim(x.Type, x.Value));
-			(IServiceAccount account, string secretToken) = await _serviceAccountCollection.CreateAsync(new CreateServiceAccountOptions(request.Description, claims, request.Enabled), cancellationToken);
+			CreateServiceAccountOptions options = new (request.Name, request.Description, claims, request.Enabled);
+			(IServiceAccount account, string secretToken) = await _serviceAccountCollection.CreateAsync(options, cancellationToken);
 			return new CreateServiceAccountResponse(account.Id, secretToken);
 		}
 
@@ -150,8 +151,7 @@ namespace HordeServer.ServiceAccounts
 				claims = request.Claims.ConvertAll(x => new UserClaim(x.Type, x.Value));
 			}
 
-			UpdateServiceAccountOptions options = new UpdateServiceAccountOptions(request.Description, claims, request.ResetToken, request.Enabled);
-
+			UpdateServiceAccountOptions options = new (request.Name, request.Description, claims, request.ResetToken, request.Enabled);
 			(IServiceAccount? newAccount, string? newToken) = await account.UpdateAsync(options, cancellationToken);
 			if (newAccount == null)
 			{
diff --git a/Engine/Source/Programs/Horde/Plugins/Compute/HordeServer.Compute/Agents/AgentsController.cs b/Engine/Source/Programs/Horde/Plugins/Compute/HordeServer.Compute/Agents/AgentsController.cs
--- a/Engine/Source/Programs/Horde/Plugins/Compute/HordeServer.Compute/Agents/AgentsController.cs
+++ b/Engine/Source/Programs/Horde/Plugins/Compute/HordeServer.Compute/Agents/AgentsController.cs
@@ -218,10 +218,7 @@ namespace HordeServer.Agents
 				return Forbid(AgentAclAction.UpdateAgent, agentId);
 			}
 
-			List<PoolId>? updatePools = update.Pools?.ConvertAll(x => new PoolId(x));
-
-			IUser? user = await _userCollection.GetUserAsync(User, HttpContext.RequestAborted);
-			string userName = user?.Name ?? "Unknown";
+			string userName = User.GetUser() ?? "Unknown";
 
 			for (; ; )
 			{
diff --git a/Engine/Source/Programs/Horde/Plugins/Tools/HordeServer.Tools.Tests/Tools/ToolsControllerTests.cs b/Engine/Source/Programs/Horde/Plugins/Tools/HordeServer.Tools.Tests/Tools/ToolsControllerTests.cs
--- a/Engine/Source/Programs/Horde/Plugins/Tools/HordeServer.Tools.Tests/Tools/ToolsControllerTests.cs
+++ b/Engine/Source/Programs/Horde/Plugins/Tools/HordeServer.Tools.Tests/Tools/ToolsControllerTests.cs
@@ -58,7 +58,7 @@ public class ToolsControllerTests
 		configService.OverrideConfig(globalConfig);
 
 		List<IUserClaim> claims = [new UserClaim("http://epicgames.com/ue/horde/role", "agent")];
-		(IServiceAccount _, string token) = await serviceAccounts.CreateAsync(new CreateServiceAccountOptions("myDesc", claims));
+		(IServiceAccount _, string token) = await serviceAccounts.CreateAsync(new CreateServiceAccountOptions("myName", "myDesc", claims));
 
 		// Create tool and deployment
 		using MemoryStream ms = new(await ToolTests.CreateZipFileDataAsync("foo.txt", "foo content"));
diff --git a/Engine/Source/Programs/Shared/EpicGames.Horde/ServiceAccounts/ServiceAccountMessages.cs b/Engine/Source/Programs/Shared/EpicGames.Horde/ServiceAccounts/ServiceAccountMessages.cs
--- a/Engine/Source/Programs/Shared/EpicGames.Horde/ServiceAccounts/ServiceAccountMessages.cs
+++ b/Engine/Source/Programs/Shared/EpicGames.Horde/ServiceAccounts/ServiceAccountMessages.cs
@@ -10,31 +10,33 @@ namespace EpicGames.Horde.ServiceAccounts
 	/// <summary>
 	/// Creates a new user account
 	/// </summary>
+	/// <param name="Name">Name of the account</param>
 	/// <param name="Description">Description for the account</param>
 	/// <param name="Claims">Claims for the user</param>
 	/// <param name="Enabled">Whether the account is enabled</param>
-	public record class CreateServiceAccountRequest(string Description, List<AccountClaimMessage> Claims, bool? Enabled = true);
+	public record CreateServiceAccountRequest(string Name, string Description, List<AccountClaimMessage> Claims, bool? Enabled = true);
 
 	/// <summary>
 	/// Response from the request to create a new user account
 	/// </summary>
 	/// <param name="Id">The created account id</param>
 	/// <param name="SecretToken">Secret used to auth with this account</param>
-	public record class CreateServiceAccountResponse(ServiceAccountId Id, string SecretToken);
+	public record CreateServiceAccountResponse(ServiceAccountId Id, string SecretToken);
 
 	/// <summary>
 	/// Update request for a user account
 	/// </summary>
+	/// <param name="Name">Name of the account</param>
 	/// <param name="Description">Description for the account</param>
 	/// <param name="Claims">Claims for the user</param>
 	/// <param name="ResetToken">Request that the token get reset</param>
 	/// <param name="Enabled">Whether the account is enabled</param>
-	public record class UpdateServiceAccountRequest(string? Description = null, List<AccountClaimMessage>? Claims = null, bool? ResetToken = null, bool? Enabled = null);
+	public record UpdateServiceAccountRequest(string? Name = null, string? Description = null, List<AccountClaimMessage>? Claims = null, bool? ResetToken = null, bool? Enabled = null);
 
 	/// <summary>
 	/// Response from updating a user account
 	/// </summary>
-	public record class UpdateServiceAccountResponse(string? NewSecretToken = null);
+	public record UpdateServiceAccountResponse(string? NewSecretToken = null);
 
 	/// <summary>
 	/// Creates a new user account
@@ -43,5 +45,5 @@ namespace EpicGames.Horde.ServiceAccounts
 	/// <param name="Claims">Claims for the user</param>
 	/// <param name="Description">Description for the account</param>
 	/// <param name="Enabled">Whether the account is enabled</param>
-	public record class GetServiceAccountResponse(ServiceAccountId Id, List<AccountClaimMessage> Claims, string Description, bool Enabled);
+	public record GetServiceAccountResponse(ServiceAccountId Id, List<AccountClaimMessage> Claims, string Description, bool Enabled);
 }
