Horde: Fix broken ServiceAccountAuth test

diff --git a/Engine/Source/Programs/Horde/HordeServer.Tests/Authentication/ServiceAccountAuthTest.cs b/Engine/Source/Programs/Horde/HordeServer.Tests/Authentication/ServiceAccountAuthTest.cs
--- a/Engine/Source/Programs/Horde/HordeServer.Tests/Authentication/ServiceAccountAuthTest.cs
+++ b/Engine/Source/Programs/Horde/HordeServer.Tests/Authentication/ServiceAccountAuthTest.cs
@@ -7,6 +7,7 @@ using System.Threading.Tasks;
 using HordeServer.Authentication;
 using HordeServer.ServiceAccounts;
 using HordeServer.Users;
+using HordeServer.Utilities;
 using Microsoft.AspNetCore.Authentication;
 using Microsoft.AspNetCore.Http;
 using Microsoft.Extensions.DependencyInjection;
@@ -106,8 +107,9 @@ namespace HordeServer.Tests.Authentication
 			ServiceAccountAuthHandler handler = await GetAuthHandlerAsync($"ServiceAccount {_serviceAccountToken}");
 			AuthenticateResult result = await handler.AuthenticateAsync();
 			Assert.IsTrue(result.Succeeded);
-			Assert.AreEqual(3, result.Ticket!.Principal.Claims.Count());
-			Assert.AreEqual(ServiceAccountAuthHandler.AuthenticationScheme, result.Ticket.Principal.FindFirst(ClaimTypes.Name)!.Value);
+			Assert.AreEqual(4, result.Ticket!.Principal.Claims.Count());
+			Assert.AreEqual("myName", result.Ticket.Principal.FindFirst(ClaimTypes.Name)!.Value);
+			Assert.AreEqual("myName", result.Ticket.Principal.FindFirst(HordeClaimTypes.User)!.Value);
 			Assert.AreEqual("myValue", result.Ticket!.Principal.FindFirst("myClaim")!.Value);
 			Assert.AreEqual("bar", result.Ticket!.Principal.FindFirst("foo")!.Value);
 		}
