Horde: Add clear conform to agent context menu
#skipci

diff --git a/Engine/Source/Programs/Horde/HordeDashboard/src/components/AgentView.tsx b/Engine/Source/Programs/Horde/HordeDashboard/src/components/AgentView.tsx
--- a/Engine/Source/Programs/Horde/HordeDashboard/src/components/AgentView.tsx
+++ b/Engine/Source/Programs/Horde/HordeDashboard/src/components/AgentView.tsx
@@ -1260,6 +1260,9 @@ const agentSelectedProps: IContextualMenuProps = {
          else if (item.key === "requestFullConform") {
             localState.requestBuilderUpdate(false, false, true);
          }
+         else if (item.key === "clearConform") {
+            localState.requestBuilderUpdate(false, false, false);
+         }            
          else if (item.key === "restart") {
             localState.setRestartBuilderDialogOpen(true);
          }
@@ -1301,6 +1304,10 @@ const agentSelectedProps: IContextualMenuProps = {
          key: 'cancelleases',
          text: 'Cancel Leases',
       },
+      {
+         key: 'clearConform',
+         text: 'Clear Pending Conform'
+      },
       {
          key: 'requestConform',
          text: 'Request Conform'
@@ -1376,6 +1383,11 @@ const agentContextMenuProps: IContextualMenuItem[] = [
       text: 'Cancel Leases',
       onClick: () => localState.setCancelLeasesDialogOpen(true)
    },
+   {
+      key: 'clearConform',
+      text: 'Clear Pending Conform',
+      onClick: () => localState.requestBuilderUpdate(false, false, false)
+   },
    {
       key: 'requestConform',
       text: 'Request Conform',
