Horde: Move agent functionality from tray app into a plugin.

#rnx

diff --git a/Engine/Source/Programs/Horde/HordeTrayApp/AgentPlugin.cs b/Engine/Source/Programs/Horde/HordeTrayApp/AgentPlugin.cs
new file mode 100644
--- /dev/null
+++ b/Engine/Source/Programs/Horde/HordeTrayApp/AgentPlugin.cs
@@ -0,0 +1,517 @@
+// Copyright Epic Games, Inc. All Rights Reserved.
+
+using System;
+using System.Collections.Generic;
+using System.Diagnostics;
+using System.IO.Pipes;
+using System.Linq;
+using System.Runtime.InteropServices;
+using System.Text;
+using System.Text.Json;
+using System.Threading.Tasks;
+using EpicGames.Core;
+using HordeTrayApp.Forms;
+using HordeAgent;
+using Microsoft.Win32;
+
+namespace HordeTrayApp
+{
+	record struct IdleStat(string Name, long Value, long MinValue);
+
+	class AgentPlugin : TrayAppPluginBase
+	{
+		readonly ITrayAppHost _host;
+
+		readonly BackgroundTask _clientTask;
+		readonly BackgroundTask _tickPauseStateTask;
+
+		readonly ToolStripMenuItem _enrollMenuItem;
+		readonly ToolStripMenuItem _statusEnabled;
+		readonly ToolStripMenuItem _statusDisabled;
+		readonly ToolStripMenuItem _statusWhenIdle;
+
+		readonly ToolStripMenuItem _statusMenuItem;
+		readonly ToolStripMenuItem _logsMenuItem;
+
+		readonly Settings _settings;
+
+		AgentSettingsMessage? _agentSettings;
+		IdleForm? _idleForm;
+
+		void EnrollWithServer()
+		{
+			Uri? serverUrl = _agentSettings?.ServerUrl;
+			if (serverUrl != null)
+			{
+				Process.Start(new ProcessStartInfo(new Uri(serverUrl, "agents/registration").ToString()) { UseShellExecute = true });
+			}
+		}
+
+		public AgentPlugin(ITrayAppHost host)
+		{
+			_host = host;
+			_settings = LoadSettings();
+
+			_enrollMenuItem = new ToolStripMenuItem("Enroll with Server...");
+			_enrollMenuItem.Click += (s, e) => EnrollWithServer();
+
+			_statusEnabled = new ToolStripMenuItem("Enabled");
+			_statusEnabled.Click += (s, e) => SetUserStatus(UserStatus.Enabled);
+
+			_statusDisabled = new ToolStripMenuItem("Disabled");
+			_statusDisabled.Click += (s, e) => SetUserStatus(UserStatus.Disabled);
+
+			_statusWhenIdle = new ToolStripMenuItem("When Idle");
+			_statusWhenIdle.Click += (s, e) => SetUserStatus(UserStatus.WhenIdle);
+
+			ToolStripMenuItem showStatsMenuItem = new ToolStripMenuItem("Stats...");
+			showStatsMenuItem.Click += Status_Stats_OnClick;
+
+			_statusMenuItem = new ToolStripMenuItem("Status");
+			_statusMenuItem.DropDownItems.Add(_statusEnabled);
+			_statusMenuItem.DropDownItems.Add(_statusDisabled);
+			_statusMenuItem.DropDownItems.Add(_statusWhenIdle);
+			_statusMenuItem.DropDownItems.Add(new ToolStripSeparator());
+			_statusMenuItem.DropDownItems.Add(showStatsMenuItem);
+
+			_logsMenuItem = new ToolStripMenuItem("Open logs dir");
+			_logsMenuItem.Click += OnOpenLogs;
+
+			_clientTask = BackgroundTask.StartNew(StatusTaskAsync);
+			_tickPauseStateTask = BackgroundTask.StartNew(ctx => TickPauseStateAsync(ctx));
+		}
+
+		private static Settings LoadSettings()
+		{
+			Settings? result = null;
+
+			DirectoryReference? settingsRoot = DirectoryReference.GetSpecialFolder(Environment.SpecialFolder.CommonApplicationData);
+			if (settingsRoot != null)
+			{
+				FileReference settingsPath = FileReference.Combine(settingsRoot, "Epic", "Horde", "TrayApp", "Settings.json");
+				if (FileReference.Exists(settingsPath))
+				{
+					try
+					{
+						using FileStream stream = FileReference.Open(settingsPath, FileMode.Open, FileAccess.Read);
+						result = JsonSerializer.Deserialize<Settings>(stream);
+					}
+					catch (Exception)
+					{
+					}
+				}
+				else
+				{
+					// File not found, create a file containing the default settings
+					DirectoryReference.CreateDirectory(settingsPath.Directory);
+					using FileStream stream = FileReference.Open(settingsPath, FileMode.OpenOrCreate, FileAccess.Write);
+					JsonSerializer.Serialize(stream, new Settings(), new JsonSerializerOptions() { WriteIndented = true });
+				}
+			}
+
+			return result ?? new Settings();
+		}
+
+		public override void PopulateMenu(ContextMenuStrip contextMenu)
+		{
+			base.PopulateMenu(contextMenu);
+
+			contextMenu.Items.Add(_enrollMenuItem);
+			contextMenu.Items.Add(_statusMenuItem);
+			contextMenu.Items.Add(new ToolStripSeparator());
+			contextMenu.Items.Add(_logsMenuItem);
+		}
+
+		public override void UpdateMenu()
+		{
+			base.UpdateMenu();
+
+			UserStatus status = GetUserStatus();
+			_statusEnabled.Checked = (status == UserStatus.Enabled);
+			_statusDisabled.Checked = (status == UserStatus.Disabled);
+			_statusWhenIdle.Checked = (status == UserStatus.WhenIdle);
+		}
+
+		public override async ValueTask DisposeAsync()
+		{
+			await base.DisposeAsync();
+
+			if (_idleForm != null)
+			{
+				_idleForm.Dispose();
+				_idleForm = null;
+			}
+
+			_enrollMenuItem.Dispose();
+			_statusEnabled.Dispose();
+			_statusDisabled.Dispose();
+			_statusWhenIdle.Dispose();
+
+			_logsMenuItem.Dispose();
+			_statusMenuItem.Dispose();
+
+			await _tickPauseStateTask.DisposeAsync();
+			await _clientTask.DisposeAsync();
+		}
+
+		private void OnOpenLogs(object? sender, EventArgs e)
+		{
+			DirectoryReference? programDataDir = DirectoryReference.GetSpecialFolder(Environment.SpecialFolder.CommonApplicationData);
+			if (programDataDir != null)
+			{
+				DirectoryReference logsDir = DirectoryReference.Combine(programDataDir, "Epic", "Horde", "Agent");
+				if (DirectoryReference.Exists(logsDir))
+				{
+					Process.Start(new ProcessStartInfo { FileName = logsDir.FullName, UseShellExecute = true });
+				}
+				else
+				{
+					MessageBox.Show("Unable to open logs dir " + logsDir.FullName, "Horde Tray App", MessageBoxButtons.OK, MessageBoxIcon.Error);
+				}
+			}
+		}
+
+		enum UserStatus
+		{
+			Enabled = 0,
+			Disabled = 1,
+			WhenIdle = 2,
+		}
+
+		const string RegistryKey = "HKEY_CURRENT_USER\\Software\\Epic Games\\Horde\\TrayApp";
+		const string RegistryStatusValue = "Status";
+
+		private static UserStatus GetUserStatus()
+		{
+			return (UserStatus)((Registry.GetValue(RegistryKey, RegistryStatusValue, null) as int?) ?? 0);
+		}
+
+		private void SetUserStatus(UserStatus status)
+		{
+			Registry.SetValue(RegistryKey, RegistryStatusValue, (int)status);
+			_statusChangedEvent.Set();
+		}
+
+		private void Status_Stats_OnClick(object? sender, EventArgs e)
+		{
+			if (_idleForm == null)
+			{
+				_idleForm = new IdleForm();
+				_idleForm.FormClosed += (s, e) =>
+				{
+					_idleForm.Dispose();
+					_idleForm = null;
+				};
+				_idleForm.Show();
+			}
+		}
+
+		TrayAppPluginStatus? _status;
+
+		void SetStatus(AgentStatusMessage status)
+		{
+			if (!status.Healthy)
+			{
+				string message = String.IsNullOrEmpty(status.Detail) ? "Error. Check logs." : (status.Detail.Length > 100) ? status.Detail.Substring(0, 100) : status.Detail;
+				_status = new TrayAppPluginStatus(TrayAppPluginState.Error, message);
+			}
+			else if (status.NumLeases > 0)
+			{
+				string message = (status.NumLeases == 1) ? "Currently handling 1 lease" : $"Currently handling {status.NumLeases} leases";
+				_status = new TrayAppPluginStatus(TrayAppPluginState.Busy, message);
+			}
+			else if (_enabled)
+			{
+				string message = "Agent is operating normally";
+				_status = new TrayAppPluginStatus(TrayAppPluginState.Ok, message);
+			}
+			else
+			{
+				string message = "Agent is paused";
+				_status = new TrayAppPluginStatus(TrayAppPluginState.Paused, message);
+			}
+			_host.UpdateStatus();
+		}
+
+		public override TrayAppPluginStatus GetStatus()
+			=> _status ?? base.GetStatus();
+
+		async Task StatusTaskAsync(CancellationToken cancellationToken)
+		{
+			SetStatus(AgentStatusMessage.Starting);
+			for (; ; )
+			{
+				try
+				{
+					await PollForStatusUpdatesAsync(cancellationToken);
+				}
+				catch (OperationCanceledException) when (cancellationToken.IsCancellationRequested)
+				{
+					break;
+				}
+				catch
+				{
+					SetStatus(new AgentStatusMessage(false, 0, "Unable to connect to Agent. Check the service is running."));
+					await Task.Delay(TimeSpan.FromSeconds(5.0), cancellationToken);
+				}
+			}
+		}
+
+#pragma warning disable IDE1006
+		[StructLayout(LayoutKind.Sequential)]
+		struct LASTINPUTINFO
+		{
+			public int cbSize;
+			public uint dwTime;
+		}
+
+		[DllImport("user32.dll")]
+		static extern bool GetLastInputInfo(ref LASTINPUTINFO plii);
+
+		[DllImport("kernel32.dll")]
+		static extern uint GetTickCount();
+
+		[StructLayout(LayoutKind.Sequential)]
+		struct FILETIME
+		{
+			public uint dwLowDateTime;
+			public uint dwHighDateTime;
+
+			public readonly ulong Total => dwLowDateTime | ((ulong)dwHighDateTime << 32);
+		};
+
+		[StructLayout(LayoutKind.Sequential)]
+		struct MEMORYSTATUSEX
+		{
+			public int dwLength;
+			public uint dwMemoryLoad;
+			public ulong ullTotalPhys;
+			public ulong ullAvailPhys;
+			public ulong ullTotalPageFile;
+			public ulong ullAvailPageFile;
+			public ulong ullTotalVirtual;
+			public ulong ullAvailVirtual;
+			public ulong ullAvailExtendedVirtual;
+		}
+
+		[DllImport("kernel32.dll", CharSet = CharSet.Auto, SetLastError = true)]
+		static extern bool GlobalMemoryStatusEx(ref MEMORYSTATUSEX lpBuffer);
+
+		[DllImport("kernel32.dll", SetLastError = true)]
+		static extern bool GetSystemTimes(out FILETIME lpIdleTime, out FILETIME lpKernelTime, out FILETIME lpUserTime);
+#pragma warning restore IDE1006
+
+		bool _enabled;
+		readonly AsyncEvent _statusChangedEvent = new AsyncEvent();
+		readonly AsyncEvent _enabledChangedEvent = new AsyncEvent();
+
+		async Task TickPauseStateAsync(CancellationToken cancellationToken)
+		{
+			await using BackgroundTask cpuStatsTask = BackgroundTask.StartNew(ctx => TickCpuStatsAsync(ctx));
+			await using BackgroundTask criticalProcessTask = BackgroundTask.StartNew(ctx => TickCriticalProcessAsync(ctx));
+
+			TimeSpan pollInterval = TimeSpan.FromSeconds(0.25);
+
+			Stopwatch stateChangeTimer = Stopwatch.StartNew();
+			while (!cancellationToken.IsCancellationRequested)
+			{
+				Task statusChangedTask = _statusChangedEvent.Task;
+
+				UserStatus userStatus = GetUserStatus();
+				if (userStatus == UserStatus.Enabled)
+				{
+					if (!_enabled)
+					{
+						_enabled = true;
+						_enabledChangedEvent.Set();
+					}
+				}
+				else if (userStatus == UserStatus.Disabled)
+				{
+					if (_enabled)
+					{
+						_enabled = false;
+						_enabledChangedEvent.Set();
+					}
+				}
+
+				DateTime utcNow = DateTime.UtcNow;
+				IEnumerable<IdleStat> idleStats = GetIdleStats();
+
+				bool idle = idleStats.All(x => x.Value >= x.MinValue);
+				if (idle == _enabled)
+				{
+					stateChangeTimer.Restart();
+				}
+
+				const int WakeTimeSecs = 2;
+				const int IdleTimeSecs = 30;
+				int stateChangeTime = (int)stateChangeTimer.Elapsed.TotalSeconds;
+				int stateChangeMaxTime = _enabled ? WakeTimeSecs : IdleTimeSecs;
+				_idleForm?.TickStats(_enabled, stateChangeTime, stateChangeMaxTime, idleStats);
+
+				if (userStatus == UserStatus.WhenIdle && stateChangeTime >= stateChangeMaxTime)
+				{
+					_enabled ^= true;
+					_enabledChangedEvent.Set();
+					stateChangeTimer.Restart();
+				}
+
+				await Task.WhenAny(statusChangedTask, Task.Delay(pollInterval, cancellationToken));
+			}
+		}
+
+		IEnumerable<IdleStat> GetIdleStats()
+		{
+			// Check there has been no input for a while
+			LASTINPUTINFO lastInputInfo = new LASTINPUTINFO();
+			lastInputInfo.cbSize = Marshal.SizeOf<LASTINPUTINFO>();
+
+			if (GetLastInputInfo(ref lastInputInfo))
+			{
+				yield return new IdleStat("LastInputTime", (GetTickCount() - lastInputInfo.dwTime) / 1000, _settings.Idle.MinIdleTimeSecs);
+			}
+
+			// Check that no critical processes are running
+			if (_settings.Idle.CriticalProcesses.Any())
+			{
+				yield return new IdleStat("CriticalProcCount", -_idleCriticalProcessCount, 0);
+			}
+
+			// Only look at memory/CPU usage if we're not paused; executing jobs will increase them
+			if (!_enabled)
+			{
+				// Check the CPU usage doesn't exceed the limit
+				yield return new IdleStat("IdleCpuPct", _idleCpuPct, _settings.Idle.MinIdleCpuPct);
+
+				// Check there's enough available virtual memory 
+				MEMORYSTATUSEX memoryStatus = new MEMORYSTATUSEX();
+				memoryStatus.dwLength = Marshal.SizeOf<MEMORYSTATUSEX>();
+
+				if (GlobalMemoryStatusEx(ref memoryStatus))
+				{
+					yield return new IdleStat("VirtualMemMb", (long)(memoryStatus.ullAvailPhys + memoryStatus.ullAvailPageFile) / (1024 * 1024), _settings.Idle.MinFreeVirtualMemMb);
+				}
+			}
+		}
+
+		int _idleCpuPct = 0;
+		int _idleCriticalProcessCount = 0;
+
+		async Task TickCpuStatsAsync(CancellationToken cancellationToken)
+		{
+			const int NumSamples = 10;
+			TimeSpan sampleInterval = TimeSpan.FromSeconds(0.2);
+			(ulong IdleTime, ulong TotalTime)[] samples = new (ulong IdleTime, ulong TotalTime)[NumSamples];
+
+			int sampleIdx = 0;
+			for (; ; )
+			{
+				if (GetSystemTimes(out FILETIME idleTime, out FILETIME kernelTime, out FILETIME userTime))
+				{
+					(ulong prevIdleTime, ulong prevTotalTime) = samples[sampleIdx];
+					(ulong nextIdleTime, ulong nextTotalTime) = (idleTime.Total, kernelTime.Total + userTime.Total);
+
+					samples[sampleIdx] = (nextIdleTime, nextTotalTime);
+					sampleIdx = (sampleIdx + 1) % NumSamples;
+
+					if (prevTotalTime > 0 && nextTotalTime > prevTotalTime)
+					{
+						_idleCpuPct = (int)(((nextIdleTime - prevIdleTime) * 100) / (nextTotalTime - prevTotalTime));
+					}
+				}
+				await Task.Delay(sampleInterval, cancellationToken);
+			}
+		}
+
+		async Task TickCriticalProcessAsync(CancellationToken cancellationToken)
+		{
+			TimeSpan sampleInterval = TimeSpan.FromSeconds(1.0);
+
+			for (; ; )
+			{
+				try
+				{
+					if (RuntimeInformation.IsOSPlatform(OSPlatform.Windows) && _settings.Idle.CriticalProcesses.Any())
+					{
+						IEnumerable<int> hordeProcessIds = Process.GetProcessesByName("HordeAgent").Select(x => x.Id);
+						IEnumerable<Process> criticalProcesses = _settings.Idle.CriticalProcesses
+							.Select(x => Path.GetFileNameWithoutExtension(x).ToUpperInvariant())
+							.Distinct()
+							.SelectMany(x => Process.GetProcessesByName(x))
+							.Where(x => !x.HasExited);
+
+						// Ignore processes that are descendants of HordeAgent
+						if (hordeProcessIds.Any() && criticalProcesses.Any())
+						{
+							criticalProcesses = criticalProcesses
+								.Where(x => !ProcessUtils.GetAncestorProcesses(x)
+								.Select(x => x.Id).Intersect(hordeProcessIds).Any())
+								.Where(x => !x.HasExited);
+						}
+
+						_idleCriticalProcessCount = criticalProcesses.Count();
+					}
+				}
+				catch (InvalidOperationException)
+				{
+					// If a process stops running Process.Id will throw an exception
+				}
+				await Task.Delay(sampleInterval, cancellationToken);
+			}
+		}
+
+		async Task PollForStatusUpdatesAsync(CancellationToken cancellationToken)
+		{
+			AgentMessageBuffer message = new AgentMessageBuffer();
+			using (NamedPipeClientStream pipeClient = new NamedPipeClientStream(".", AgentMessagePipe.PipeName, PipeDirection.InOut))
+			{
+				SetStatus(new AgentStatusMessage(false, 0, "Connecting to agent..."));
+				await pipeClient.ConnectAsync(cancellationToken);
+
+				SetStatus(new AgentStatusMessage(false, 0, "Waiting for status update."));
+				for (; ; )
+				{
+					Task idleChangeTask = _enabledChangedEvent.Task;
+
+					bool enabled = _enabled;
+					message.Set(AgentMessageType.SetEnabledRequest, new AgentEnabledMessage(enabled));
+					await message.SendAsync(pipeClient, cancellationToken);
+
+					if (_agentSettings == null)
+					{
+						message.Set(AgentMessageType.GetSettingsRequest);
+						await message.SendAsync(pipeClient, cancellationToken);
+
+						if (!await message.TryReadAsync(pipeClient, cancellationToken))
+						{
+							break;
+						}
+						if (message.Type == AgentMessageType.GetSettingsResponse)
+						{
+							_agentSettings = message.Parse<AgentSettingsMessage>();
+						}
+					}
+
+					message.Set(AgentMessageType.GetStatusRequest);
+					await message.SendAsync(pipeClient, cancellationToken);
+
+					if (!await message.TryReadAsync(pipeClient, cancellationToken))
+					{
+						break;
+					}
+
+					switch (message.Type)
+					{
+						case AgentMessageType.GetStatusResponse:
+							AgentStatusMessage status = message.Parse<AgentStatusMessage>();
+							SetStatus(status);
+							break;
+					}
+
+					await Task.WhenAny(idleChangeTask, Task.Delay(TimeSpan.FromSeconds(5.0), cancellationToken));
+				}
+			}
+		}
+	}
+}
diff --git a/Engine/Source/Programs/Horde/HordeTrayApp/Forms/IdleForm.cs b/Engine/Source/Programs/Horde/HordeTrayApp/Forms/IdleForm.cs
--- a/Engine/Source/Programs/Horde/HordeTrayApp/Forms/IdleForm.cs
+++ b/Engine/Source/Programs/Horde/HordeTrayApp/Forms/IdleForm.cs
@@ -1,6 +1,6 @@
 // Copyright Epic Games, Inc. All Rights Reserved.
 
-namespace Horde.Agent.TrayApp.Forms
+namespace HordeTrayApp.Forms
 {
 	partial class IdleForm : Form
 	{
diff --git a/Engine/Source/Programs/Horde/HordeTrayApp/ITrayAppHost.cs b/Engine/Source/Programs/Horde/HordeTrayApp/ITrayAppHost.cs
new file mode 100644
--- /dev/null
+++ b/Engine/Source/Programs/Horde/HordeTrayApp/ITrayAppHost.cs
@@ -0,0 +1,21 @@
+// Copyright Epic Games, Inc. All Rights Reserved.
+
+using System;
+using System.Collections.Generic;
+using System.Linq;
+using System.Text;
+using System.Threading.Tasks;
+
+namespace HordeTrayApp
+{
+	/// <summary>
+	/// Interface for the tray app host application
+	/// </summary>
+	interface ITrayAppHost
+	{
+		/// <summary>
+		/// Notifies the host that a status change has ocurred
+		/// </summary>
+		void UpdateStatus();
+	}
+}
diff --git a/Engine/Source/Programs/Horde/HordeTrayApp/ITrayAppPlugin.cs b/Engine/Source/Programs/Horde/HordeTrayApp/ITrayAppPlugin.cs
new file mode 100644
--- /dev/null
+++ b/Engine/Source/Programs/Horde/HordeTrayApp/ITrayAppPlugin.cs
@@ -0,0 +1,73 @@
+// Copyright Epic Games, Inc. All Rights Reserved.
+
+using System;
+using System.Collections.Generic;
+using System.Linq;
+using System.Text;
+using System.Threading.Tasks;
+
+namespace HordeTrayApp
+{
+	enum TrayAppPluginState
+	{
+		Undefined,
+		Ok,
+		Busy,
+		Paused,
+		Error
+	}
+
+	/// <summary>
+	/// Reported status of a plugin
+	/// </summary>
+	record class TrayAppPluginStatus(TrayAppPluginState State, string? Message);
+
+	/// <summary>
+	/// Plugin for the tray app
+	/// </summary>
+	interface ITrayAppPlugin : IAsyncDisposable
+	{
+		/// <summary>
+		/// Get the current status of this plugin
+		/// </summary>
+		TrayAppPluginStatus GetStatus();
+
+		/// <summary>
+		/// Allow the plugin to customize the tray icon context menu
+		/// </summary>
+		void PopulateMenu(ContextMenuStrip contextMenu);
+
+		/// <summary>
+		/// Update the menu state before it's shown
+		/// </summary>
+		void UpdateMenu();
+	}
+
+	/// <summary>
+	/// Plugin for the tray app
+	/// </summary>
+	abstract class TrayAppPluginBase : ITrayAppPlugin
+	{
+		/// <inheritdoc/>
+		public virtual TrayAppPluginStatus GetStatus()
+		{
+			return new TrayAppPluginStatus(TrayAppPluginState.Undefined, null);
+		}
+
+		/// <inheritdoc/>
+		public virtual void PopulateMenu(ContextMenuStrip contextMenu)
+		{
+		}
+
+		/// <inheritdoc/>
+		public virtual void UpdateMenu()
+		{
+		}
+
+		/// <inheritdoc/>
+		public virtual ValueTask DisposeAsync()
+		{
+			return default;
+		}
+	}
+}
diff --git a/Engine/Source/Programs/Horde/HordeTrayApp/Program.cs b/Engine/Source/Programs/Horde/HordeTrayApp/Program.cs
--- a/Engine/Source/Programs/Horde/HordeTrayApp/Program.cs
+++ b/Engine/Source/Programs/Horde/HordeTrayApp/Program.cs
@@ -7,14 +7,12 @@ using System.Runtime.InteropServices;
 using System.Text.Json;
 using EpicGames.Core;
 using HordeAgent;
-using Horde.Agent.TrayApp.Forms;
-using Horde.Agent.TrayApp.Properties;
+using HordeTrayApp.Forms;
+using HordeTrayApp.Properties;
 using Microsoft.Win32;
 
-namespace Horde.Agent.TrayApp
+namespace HordeTrayApp
 {
-	record struct IdleStat(string Name, long Value, long MinValue);
-
 	static class Program
 	{
 		const string MutexName = "Horde.Agent.TrayApp-Mutex";
@@ -141,77 +139,33 @@ namespace Horde.Agent.TrayApp
 		}
 	}
 
-	class CustomApplicationContext : ApplicationContext, IAsyncDisposable
+	class CustomApplicationContext : ApplicationContext, ITrayAppHost, IAsyncDisposable
 	{
 		readonly NotifyIcon _trayIcon;
-		readonly BackgroundTask _clientTask;
-		readonly BackgroundTask _tickPauseStateTask;
-		readonly BackgroundTask _waitForExitTask;
 		readonly Control _mainThreadInvokeTarget;
+		readonly List<ITrayAppPlugin> _plugins = new List<ITrayAppPlugin>();
+		readonly BackgroundTask _waitForExitTask;
 
-		readonly ToolStripMenuItem _enrollMenuItem;
-		readonly ToolStripMenuItem _statusEnabled;
-		readonly ToolStripMenuItem _statusDisabled;
-		readonly ToolStripMenuItem _statusWhenIdle;
-
-		readonly Settings _settings;
-
-		AgentSettingsMessage? _agentSettings;
-		IdleForm? _idleForm;
 		bool _disposed;
 
-		void EnrollWithServer()
-		{
-			Uri? serverUrl = _agentSettings?.ServerUrl;
-			if (serverUrl != null)
-			{
-				Process.Start(new ProcessStartInfo(new Uri(serverUrl, "agents/registration").ToString()) { UseShellExecute = true });
-			}
-		}
-
 		public CustomApplicationContext(EventWaitHandle eventHandle)
 		{
-			_settings = LoadSettings();
-
-			_enrollMenuItem = new ToolStripMenuItem("Enroll with Server...");
-			_enrollMenuItem.Click += (s, e) => EnrollWithServer();
-
-			_statusEnabled = new ToolStripMenuItem("Enabled");
-			_statusEnabled.Click += (s, e) => SetUserStatus(UserStatus.Enabled);
-
-			_statusDisabled = new ToolStripMenuItem("Disabled");
-			_statusDisabled.Click += (s, e) => SetUserStatus(UserStatus.Disabled);
+			_plugins.Add(new AgentPlugin(this));
 
-			_statusWhenIdle = new ToolStripMenuItem("When Idle");
-			_statusWhenIdle.Click += (s, e) => SetUserStatus(UserStatus.WhenIdle);
-
-			ToolStripMenuItem showStatsMenuItem = new ToolStripMenuItem("Stats...");
-			showStatsMenuItem.Click += Status_Stats_OnClick;
-
-			ToolStripMenuItem statusMenuItem = new ToolStripMenuItem("Status");
-			statusMenuItem.DropDownItems.Add(_statusEnabled);
-			statusMenuItem.DropDownItems.Add(_statusDisabled);
-			statusMenuItem.DropDownItems.Add(_statusWhenIdle);
-			statusMenuItem.DropDownItems.Add(new ToolStripSeparator());
-			statusMenuItem.DropDownItems.Add(showStatsMenuItem);
-
-			ToolStripMenuItem logsMenuItem = new ToolStripMenuItem("Open logs dir");
-			logsMenuItem.Click += OnOpenLogs;
+			_mainThreadInvokeTarget = new Control();
+			_mainThreadInvokeTarget.CreateControl();
 
 			ToolStripMenuItem exitMenuItem = new ToolStripMenuItem("Exit");
 			exitMenuItem.Click += OnExit;
 
 			ContextMenuStrip menu = new ContextMenuStrip();
-			menu.Items.Add(_enrollMenuItem);
-			menu.Items.Add(statusMenuItem);
-			menu.Items.Add(new ToolStripSeparator());
-			menu.Items.Add(logsMenuItem);
+			foreach (ITrayAppPlugin plugin in _plugins)
+			{
+				plugin.PopulateMenu(menu);
+			}
 			menu.Items.Add(new ToolStripSeparator());
 			menu.Items.Add(exitMenuItem);
 
-			_mainThreadInvokeTarget = new Control();
-			_mainThreadInvokeTarget.CreateControl();
-
 			_trayIcon = new NotifyIcon()
 			{
 				Icon = Resources.StatusNormal,
@@ -220,133 +174,40 @@ namespace Horde.Agent.TrayApp
 			};
 			_trayIcon.Click += TrayIcon_Click;
 
-			_clientTask = BackgroundTask.StartNew(StatusTaskAsync);
-			_tickPauseStateTask = BackgroundTask.StartNew(ctx => TickPauseStateAsync(ctx));
 			_waitForExitTask = BackgroundTask.StartNew(ctx => WaitForExitAsync(eventHandle, ctx));
 		}
 
-		private static Settings LoadSettings()
-		{
-			Settings? result = null;
-
-			DirectoryReference? settingsRoot = DirectoryReference.GetSpecialFolder(Environment.SpecialFolder.CommonApplicationData);
-			if (settingsRoot != null)
-			{
-				FileReference settingsPath = FileReference.Combine(settingsRoot, "Epic", "Horde", "TrayApp", "Settings.json");
-				if (FileReference.Exists(settingsPath))
-				{
-					try
-					{
-						using FileStream stream = FileReference.Open(settingsPath, FileMode.Open, FileAccess.Read);
-						result = JsonSerializer.Deserialize<Settings>(stream);
-					}
-					catch (Exception)
-					{
-					}
-				}
-				else
-				{
-					// File not found, create a file containing the default settings
-					DirectoryReference.CreateDirectory(settingsPath.Directory);
-					using FileStream stream = FileReference.Open(settingsPath, FileMode.OpenOrCreate, FileAccess.Write);
-					JsonSerializer.Serialize(stream, new Settings(), new JsonSerializerOptions() { WriteIndented = true });
-				}
-			}
-
-			return result ?? new Settings();
-		}
-
-		private void TrayIcon_Click(object? sender, EventArgs e)
-		{
-			UserStatus status = GetUserStatus();
-			_statusEnabled.Checked = (status == UserStatus.Enabled);
-			_statusDisabled.Checked = (status == UserStatus.Disabled);
-			_statusWhenIdle.Checked = (status == UserStatus.WhenIdle);
-		}
-
 		protected override void Dispose(bool disposing)
 		{
 			base.Dispose(disposing);
 
 			if (disposing)
 			{
-				if (_idleForm != null)
-				{
-					_idleForm.Dispose();
-					_idleForm = null;
-				}
-
-				_enrollMenuItem.Dispose();
-				_statusEnabled.Dispose();
-				_statusDisabled.Dispose();
-				_statusWhenIdle.Dispose();
-
 				_mainThreadInvokeTarget.Dispose();
 				_trayIcon.Dispose();
 				_disposed = true;
 			}
 		}
 
-		public async ValueTask DisposeAsync()
-		{
-			await _waitForExitTask.DisposeAsync();
-			await _tickPauseStateTask.DisposeAsync();
-			await _clientTask.DisposeAsync();
-
-			Dispose();
-			GC.SuppressFinalize(this);
-		}
-
-		private void OnOpenLogs(object? sender, EventArgs e)
+		private void TrayIcon_Click(object? sender, EventArgs e)
 		{
-			DirectoryReference? programDataDir = DirectoryReference.GetSpecialFolder(Environment.SpecialFolder.CommonApplicationData);
-			if (programDataDir != null)
+			foreach (AgentPlugin plugin in _plugins)
 			{
-				DirectoryReference logsDir = DirectoryReference.Combine(programDataDir, "Epic", "Horde", "Agent");
-				if (DirectoryReference.Exists(logsDir))
-				{
-					Process.Start(new ProcessStartInfo { FileName = logsDir.FullName, UseShellExecute = true });
-				}
-				else
-				{
-					MessageBox.Show("Unable to open logs dir " + logsDir.FullName, "Horde Tray App", MessageBoxButtons.OK, MessageBoxIcon.Error);
-				}
+				plugin.UpdateMenu();
 			}
 		}
 
-		enum UserStatus
-		{
-			Enabled = 0,
-			Disabled = 1,
-			WhenIdle = 2,
-		}
-
-		const string RegistryKey = "HKEY_CURRENT_USER\\Software\\Epic Games\\Horde\\TrayApp";
-		const string RegistryStatusValue = "Status";
-
-		private static UserStatus GetUserStatus()
-		{
-			return (UserStatus)((Registry.GetValue(RegistryKey, RegistryStatusValue, null) as int?) ?? 0);
-		}
-
-		private void SetUserStatus(UserStatus status)
+		public async ValueTask DisposeAsync()
 		{
-			Registry.SetValue(RegistryKey, RegistryStatusValue, (int)status);
-			_statusChangedEvent.Set();
-		}
+			await _waitForExitTask.DisposeAsync();
 
-		private void Status_Stats_OnClick(object? sender, EventArgs e)
-		{
-			if (_idleForm == null)
+			foreach (AgentPlugin plugin in _plugins)
 			{
-				_idleForm = new IdleForm();
-				_idleForm.FormClosed += (s, e) =>
-				{
-					_idleForm.Dispose();
-					_idleForm = null;
-				};
-				_idleForm.Show();
+				await plugin.DisposeAsync();
 			}
+
+			Dispose();
+			GC.SuppressFinalize(this);
 		}
 
 		private void OnExit(object? sender, EventArgs e)
@@ -354,320 +215,50 @@ namespace Horde.Agent.TrayApp
 			ExitThread();
 		}
 
-		void SetStatus(AgentStatusMessage status)
+		public void UpdateStatus()
 		{
-			_mainThreadInvokeTarget.BeginInvoke(() => SetStatus_MainThread(status));
+			_mainThreadInvokeTarget.BeginInvoke(() => SetStatus_MainThread());
 		}
 
-		void SetStatus_MainThread(AgentStatusMessage status)
+		void SetStatus_MainThread()
 		{
 			if (!_disposed)
 			{
-				if (!status.Healthy)
-				{
-					_trayIcon.Icon = Resources.StatusError;
-					_trayIcon.Text = String.IsNullOrEmpty(status.Detail) ? "Error. Check logs." : (status.Detail.Length > 100) ? status.Detail.Substring(0, 100) : status.Detail;
-				}
-				else if (status.NumLeases > 0)
-				{
-					_trayIcon.Icon = Resources.StatusBusy;
-					_trayIcon.Text = (status.NumLeases == 1) ? "Currently handling 1 lease" : $"Currently handling {status.NumLeases} leases";
-				}
-				else if (_enabled)
-				{
-					_trayIcon.Icon = Resources.StatusNormal;
-					_trayIcon.Text = "Agent is operating normally";
-				}
-				else
-				{
-					_trayIcon.Icon = Resources.StatusPaused;
-					_trayIcon.Text = "Agent is paused";
-				}
-			}
-		}
+				TrayAppPluginState state = TrayAppPluginState.Undefined;
+				List<string> messages = new List<string>();
 
-		async Task StatusTaskAsync(CancellationToken cancellationToken)
-		{
-			SetStatus(AgentStatusMessage.Starting);
-			for (; ; )
-			{
-				try
-				{
-					await PollForStatusUpdatesAsync(cancellationToken);
-				}
-				catch (OperationCanceledException) when (cancellationToken.IsCancellationRequested)
+				foreach (ITrayAppPlugin plugin in _plugins)
 				{
-					break;
-				}
-				catch
-				{
-					SetStatus(new AgentStatusMessage(false, 0, "Unable to connect to Agent. Check the service is running."));
-					await Task.Delay(TimeSpan.FromSeconds(5.0), cancellationToken);
-				}
-			}
-		}
-
-		async Task WaitForExitAsync(EventWaitHandle eventHandle, CancellationToken cancellationToken)
-		{
-			await eventHandle.WaitOneAsync(cancellationToken);
-			_mainThreadInvokeTarget.BeginInvoke(() => Exit_MainThread());
-		}
-
-#pragma warning disable IDE1006
-		[StructLayout(LayoutKind.Sequential)]
-		struct LASTINPUTINFO
-		{
-			public int cbSize;
-			public uint dwTime;
-		}
-
-		[DllImport("user32.dll")]
-		static extern bool GetLastInputInfo(ref LASTINPUTINFO plii);
-
-		[DllImport("kernel32.dll")]
-		static extern uint GetTickCount();
-
-		[StructLayout(LayoutKind.Sequential)]
-		struct FILETIME
-		{
-			public uint dwLowDateTime;
-			public uint dwHighDateTime;
-
-			public readonly ulong Total => dwLowDateTime | ((ulong)dwHighDateTime << 32);
-		};
-
-		[StructLayout(LayoutKind.Sequential)]
-		struct MEMORYSTATUSEX
-		{
-			public int dwLength;
-			public uint dwMemoryLoad;
-			public ulong ullTotalPhys;
-			public ulong ullAvailPhys;
-			public ulong ullTotalPageFile;
-			public ulong ullAvailPageFile;
-			public ulong ullTotalVirtual;
-			public ulong ullAvailVirtual;
-			public ulong ullAvailExtendedVirtual;
-		}
-
-		[DllImport("kernel32.dll", CharSet = CharSet.Auto, SetLastError = true)]
-		static extern bool GlobalMemoryStatusEx(ref MEMORYSTATUSEX lpBuffer);
-
-		[DllImport("kernel32.dll", SetLastError = true)]
-		static extern bool GetSystemTimes(out FILETIME lpIdleTime, out FILETIME lpKernelTime, out FILETIME lpUserTime);
-#pragma warning restore IDE1006
-
-		bool _enabled;
-		readonly AsyncEvent _statusChangedEvent = new AsyncEvent();
-		readonly AsyncEvent _enabledChangedEvent = new AsyncEvent();
-
-		async Task TickPauseStateAsync(CancellationToken cancellationToken)
-		{
-			await using BackgroundTask cpuStatsTask = BackgroundTask.StartNew(ctx => TickCpuStatsAsync(ctx));
-			await using BackgroundTask criticalProcessTask = BackgroundTask.StartNew(ctx => TickCriticalProcessAsync(ctx));
-
-			TimeSpan pollInterval = TimeSpan.FromSeconds(0.25);
-
-			Stopwatch stateChangeTimer = Stopwatch.StartNew();
-			while (!cancellationToken.IsCancellationRequested)
-			{
-				Task statusChangedTask = _statusChangedEvent.Task;
-
-				UserStatus userStatus = GetUserStatus();
-				if (userStatus == UserStatus.Enabled)
-				{
-					if (!_enabled)
+					TrayAppPluginStatus status = plugin.GetStatus();
+					if (status.State >= state)
 					{
-						_enabled = true;
-						_enabledChangedEvent.Set();
-					}
-				}
-				else if (userStatus == UserStatus.Disabled)
-				{
-					if (_enabled)
-					{
-						_enabled = false;
-						_enabledChangedEvent.Set();
-					}
-				}
-
-				DateTime utcNow = DateTime.UtcNow;
-				IEnumerable<IdleStat> idleStats = GetIdleStats();
-
-				bool idle = idleStats.All(x => x.Value >= x.MinValue);
-				if (idle == _enabled)
-				{
-					stateChangeTimer.Restart();
-				}
-
-				const int WakeTimeSecs = 2;
-				const int IdleTimeSecs = 30;
-				int stateChangeTime = (int)stateChangeTimer.Elapsed.TotalSeconds;
-				int stateChangeMaxTime = _enabled ? WakeTimeSecs : IdleTimeSecs;
-				_idleForm?.TickStats(_enabled, stateChangeTime, stateChangeMaxTime, idleStats);
-
-				if (userStatus == UserStatus.WhenIdle && stateChangeTime >= stateChangeMaxTime)
-				{
-					_enabled ^= true;
-					_enabledChangedEvent.Set();
-					stateChangeTimer.Restart();
-				}
-
-				await Task.WhenAny(statusChangedTask, Task.Delay(pollInterval, cancellationToken));
-			}
-		}
-
-		IEnumerable<IdleStat> GetIdleStats()
-		{
-			// Check there has been no input for a while
-			LASTINPUTINFO lastInputInfo = new LASTINPUTINFO();
-			lastInputInfo.cbSize = Marshal.SizeOf<LASTINPUTINFO>();
-
-			if (GetLastInputInfo(ref lastInputInfo))
-			{
-				yield return new IdleStat("LastInputTime", (GetTickCount() - lastInputInfo.dwTime) / 1000, _settings.Idle.MinIdleTimeSecs);
-			}
-
-			// Check that no critical processes are running
-			if (_settings.Idle.CriticalProcesses.Any())
-			{
-				yield return new IdleStat("CriticalProcCount", -_idleCriticalProcessCount, 0);
-			}
-
-			// Only look at memory/CPU usage if we're not paused; executing jobs will increase them
-			if (!_enabled)
-			{
-				// Check the CPU usage doesn't exceed the limit
-				yield return new IdleStat("IdleCpuPct", _idleCpuPct, _settings.Idle.MinIdleCpuPct);
-
-				// Check there's enough available virtual memory 
-				MEMORYSTATUSEX memoryStatus = new MEMORYSTATUSEX();
-				memoryStatus.dwLength = Marshal.SizeOf<MEMORYSTATUSEX>();
-
-				if (GlobalMemoryStatusEx(ref memoryStatus))
-				{
-					yield return new IdleStat("VirtualMemMb", (long)(memoryStatus.ullAvailPhys + memoryStatus.ullAvailPageFile) / (1024 * 1024), _settings.Idle.MinFreeVirtualMemMb);
-				}
-			}
-		}
-
-		int _idleCpuPct = 0;
-		int _idleCriticalProcessCount = 0;
-
-		async Task TickCpuStatsAsync(CancellationToken cancellationToken)
-		{
-			const int NumSamples = 10;
-			TimeSpan sampleInterval = TimeSpan.FromSeconds(0.2);
-			(ulong IdleTime, ulong TotalTime)[] samples = new (ulong IdleTime, ulong TotalTime)[NumSamples];
-
-			int sampleIdx = 0;
-			for (; ; )
-			{
-				if (GetSystemTimes(out FILETIME idleTime, out FILETIME kernelTime, out FILETIME userTime))
-				{
-					(ulong prevIdleTime, ulong prevTotalTime) = samples[sampleIdx];
-					(ulong nextIdleTime, ulong nextTotalTime) = (idleTime.Total, kernelTime.Total + userTime.Total);
-
-					samples[sampleIdx] = (nextIdleTime, nextTotalTime);
-					sampleIdx = (sampleIdx + 1) % NumSamples;
-
-					if (prevTotalTime > 0 && nextTotalTime > prevTotalTime)
-					{
-						_idleCpuPct = (int)(((nextIdleTime - prevIdleTime) * 100) / (nextTotalTime - prevTotalTime));
-					}
-				}
-				await Task.Delay(sampleInterval, cancellationToken);
-			}
-		}
-
-		async Task TickCriticalProcessAsync(CancellationToken cancellationToken)
-		{
-			TimeSpan sampleInterval = TimeSpan.FromSeconds(1.0);
-
-			for (; ; )
-			{
-				try
-				{
-					if (RuntimeInformation.IsOSPlatform(OSPlatform.Windows) && _settings.Idle.CriticalProcesses.Any())
-					{
-						IEnumerable<int> hordeProcessIds = Process.GetProcessesByName("HordeAgent").Select(x => x.Id);
-						IEnumerable<Process> criticalProcesses = _settings.Idle.CriticalProcesses
-							.Select(x => Path.GetFileNameWithoutExtension(x).ToUpperInvariant())
-							.Distinct()
-							.SelectMany(x => Process.GetProcessesByName(x))
-							.Where(x => !x.HasExited);
-
-						// Ignore processes that are descendants of HordeAgent
-						if (hordeProcessIds.Any() && criticalProcesses.Any())
+						if (status.State > state)
 						{
-							criticalProcesses = criticalProcesses
-								.Where(x => !ProcessUtils.GetAncestorProcesses(x)
-								.Select(x => x.Id).Intersect(hordeProcessIds).Any())
-								.Where(x => !x.HasExited);
+							messages.Clear();
+							state = status.State;
+						}
+						if (status.Message != null)
+						{
+							messages.Add(status.Message);
 						}
-
-						_idleCriticalProcessCount = criticalProcesses.Count();
 					}
 				}
-				catch (InvalidOperationException)
+
+				_trayIcon.Icon = state switch
 				{
-					// If a process stops running Process.Id will throw an exception
-				}
-				await Task.Delay(sampleInterval, cancellationToken);
+					TrayAppPluginState.Busy => Resources.StatusBusy,
+					TrayAppPluginState.Paused => Resources.StatusPaused,
+					TrayAppPluginState.Error => Resources.StatusError,
+					_ => Resources.StatusNormal
+				};
+				_trayIcon.Text = String.Join("\n", messages);
 			}
 		}
 
-		async Task PollForStatusUpdatesAsync(CancellationToken cancellationToken)
+		async Task WaitForExitAsync(EventWaitHandle eventHandle, CancellationToken cancellationToken)
 		{
-			AgentMessageBuffer message = new AgentMessageBuffer();
-			using (NamedPipeClientStream pipeClient = new NamedPipeClientStream(".", AgentMessagePipe.PipeName, PipeDirection.InOut))
-			{
-				SetStatus(new AgentStatusMessage(false, 0, "Connecting to agent..."));
-				await pipeClient.ConnectAsync(cancellationToken);
-
-				SetStatus(new AgentStatusMessage(false, 0, "Waiting for status update."));
-				for (; ; )
-				{
-					Task idleChangeTask = _enabledChangedEvent.Task;
-
-					bool enabled = _enabled;
-					message.Set(AgentMessageType.SetEnabledRequest, new AgentEnabledMessage(enabled));
-					await message.SendAsync(pipeClient, cancellationToken);
-
-					if (_agentSettings == null)
-					{
-						message.Set(AgentMessageType.GetSettingsRequest);
-						await message.SendAsync(pipeClient, cancellationToken);
-
-						if (!await message.TryReadAsync(pipeClient, cancellationToken))
-						{
-							break;
-						}
-						if (message.Type == AgentMessageType.GetSettingsResponse)
-						{
-							_agentSettings = message.Parse<AgentSettingsMessage>();
-						}
-					}
-
-					message.Set(AgentMessageType.GetStatusRequest);
-					await message.SendAsync(pipeClient, cancellationToken);
-
-					if (!await message.TryReadAsync(pipeClient, cancellationToken))
-					{
-						break;
-					}
-
-					switch (message.Type)
-					{
-						case AgentMessageType.GetStatusResponse:
-							AgentStatusMessage status = message.Parse<AgentStatusMessage>();
-							SetStatus(status);
-							break;
-					}
-
-					await Task.WhenAny(idleChangeTask, Task.Delay(TimeSpan.FromSeconds(5.0), cancellationToken));
-				}
-			}
+			await eventHandle.WaitOneAsync(cancellationToken);
+			_mainThreadInvokeTarget.BeginInvoke(() => Exit_MainThread());
 		}
 
 		void Exit_MainThread()
diff --git a/Engine/Source/Programs/Horde/HordeTrayApp/Settings.cs b/Engine/Source/Programs/Horde/HordeTrayApp/Settings.cs
--- a/Engine/Source/Programs/Horde/HordeTrayApp/Settings.cs
+++ b/Engine/Source/Programs/Horde/HordeTrayApp/Settings.cs
@@ -2,7 +2,7 @@
 
 #pragma warning disable CS1591
 
-namespace Horde.Agent.TrayApp
+namespace HordeTrayApp
 {
 	[Serializable]
 	public class Settings
