Horde: Fix static analysis warnings.

#rnx

diff --git a/Engine/Source/Programs/Horde/HordeServer/ServerApp.cs b/Engine/Source/Programs/Horde/HordeServer/ServerApp.cs
--- a/Engine/Source/Programs/Horde/HordeServer/ServerApp.cs
+++ b/Engine/Source/Programs/Horde/HordeServer/ServerApp.cs
@@ -14,7 +14,6 @@ using HordeServer.Commands;
 using HordeServer.Plugins;
 using HordeServer.Storage;
 using HordeServer.Utilities;
-using Microsoft.Extensions.Caching.Memory;
 using Microsoft.Extensions.Configuration;
 using Microsoft.Extensions.Configuration.Memory;
 using Microsoft.Extensions.DependencyInjection;
diff --git a/Engine/Source/Programs/Horde/HordeTrayApp/AgentPlugin.cs b/Engine/Source/Programs/Horde/HordeTrayApp/AgentPlugin.cs
--- a/Engine/Source/Programs/Horde/HordeTrayApp/AgentPlugin.cs
+++ b/Engine/Source/Programs/Horde/HordeTrayApp/AgentPlugin.cs
@@ -1,17 +1,12 @@
 // Copyright Epic Games, Inc. All Rights Reserved.
 
-using System;
-using System.Collections.Generic;
 using System.Diagnostics;
 using System.IO.Pipes;
-using System.Linq;
 using System.Runtime.InteropServices;
-using System.Text;
 using System.Text.Json;
-using System.Threading.Tasks;
 using EpicGames.Core;
-using HordeTrayApp.Forms;
 using HordeAgent;
+using HordeTrayApp.Forms;
 using Microsoft.Win32;
 
 namespace HordeTrayApp
diff --git a/Engine/Source/Programs/Horde/HordeTrayApp/ITrayAppHost.cs b/Engine/Source/Programs/Horde/HordeTrayApp/ITrayAppHost.cs
--- a/Engine/Source/Programs/Horde/HordeTrayApp/ITrayAppHost.cs
+++ b/Engine/Source/Programs/Horde/HordeTrayApp/ITrayAppHost.cs
@@ -1,11 +1,5 @@
 // Copyright Epic Games, Inc. All Rights Reserved.
 
-using System;
-using System.Collections.Generic;
-using System.Linq;
-using System.Text;
-using System.Threading.Tasks;
-
 namespace HordeTrayApp
 {
 	/// <summary>
diff --git a/Engine/Source/Programs/Horde/HordeTrayApp/ITrayAppPlugin.cs b/Engine/Source/Programs/Horde/HordeTrayApp/ITrayAppPlugin.cs
--- a/Engine/Source/Programs/Horde/HordeTrayApp/ITrayAppPlugin.cs
+++ b/Engine/Source/Programs/Horde/HordeTrayApp/ITrayAppPlugin.cs
@@ -1,11 +1,5 @@
 // Copyright Epic Games, Inc. All Rights Reserved.
 
-using System;
-using System.Collections.Generic;
-using System.Linq;
-using System.Text;
-using System.Threading.Tasks;
-
 namespace HordeTrayApp
 {
 	enum TrayAppPluginState
diff --git a/Engine/Source/Programs/Horde/HordeTrayApp/Program.cs b/Engine/Source/Programs/Horde/HordeTrayApp/Program.cs
--- a/Engine/Source/Programs/Horde/HordeTrayApp/Program.cs
+++ b/Engine/Source/Programs/Horde/HordeTrayApp/Program.cs
@@ -1,15 +1,9 @@
 // Copyright Epic Games, Inc. All Rights Reserved.
 
 using System.Diagnostics;
-using System.IO.Pipes;
 using System.Reflection;
-using System.Runtime.InteropServices;
-using System.Text.Json;
 using EpicGames.Core;
-using HordeAgent;
-using HordeTrayApp.Forms;
 using HordeTrayApp.Properties;
-using Microsoft.Win32;
 
 namespace HordeTrayApp
 {
diff --git a/Engine/Source/Programs/Horde/Plugins/Build/HordeServer.Build/Notifications/Sinks/SlackNotificationSink.cs b/Engine/Source/Programs/Horde/Plugins/Build/HordeServer.Build/Notifications/Sinks/SlackNotificationSink.cs
--- a/Engine/Source/Programs/Horde/Plugins/Build/HordeServer.Build/Notifications/Sinks/SlackNotificationSink.cs
+++ b/Engine/Source/Programs/Horde/Plugins/Build/HordeServer.Build/Notifications/Sinks/SlackNotificationSink.cs
@@ -38,7 +38,6 @@ using HordeServer.Streams;
 using HordeServer.Users;
 using HordeServer.Utilities;
 using Microsoft.AspNetCore.Hosting;
-using Microsoft.CodeAnalysis;
 using Microsoft.Extensions.Caching.Memory;
 using Microsoft.Extensions.Hosting;
 using Microsoft.Extensions.Logging;
@@ -911,7 +910,7 @@ namespace HordeServer.Notifications.Sinks
 		{
 			if (issue.FixCommitId != null)
 			{
-				foreach(IGrouping<StreamId, IIssueSpan> group in spans.GroupBy(x => x.StreamId))
+				foreach (IGrouping<StreamId, IIssueSpan> group in spans.GroupBy(x => x.StreamId))
 				{
 					IIssueStream? stream = issue.Streams.FirstOrDefault(x => x.StreamId == group.Key);
 					if (stream != null && (stream.ContainsFix ?? false))
@@ -935,7 +934,7 @@ namespace HordeServer.Notifications.Sinks
 			if (issue.FixCommitId != null)
 			{
 				HashSet<StreamId> originStreams = new HashSet<StreamId>(issue.Streams.Where(x => x.MergeOrigin ?? false).Select(x => x.StreamId));
-				foreach(IGrouping<StreamId, IIssueSpan> group in spans.GroupBy(x => x.StreamId).OrderBy(x => x.Key.ToString()))
+				foreach (IGrouping<StreamId, IIssueSpan> group in spans.GroupBy(x => x.StreamId).OrderBy(x => x.Key.ToString()))
 				{
 					if (originStreams.Contains(group.Key))
 					{
diff --git a/Engine/Source/Programs/Horde/Plugins/Compute/HordeServer.Compute/Agents/AgentScheduler.cs b/Engine/Source/Programs/Horde/Plugins/Compute/HordeServer.Compute/Agents/AgentScheduler.cs
--- a/Engine/Source/Programs/Horde/Plugins/Compute/HordeServer.Compute/Agents/AgentScheduler.cs
+++ b/Engine/Source/Programs/Horde/Plugins/Compute/HordeServer.Compute/Agents/AgentScheduler.cs
@@ -466,7 +466,6 @@ namespace HordeServer.Agents
 
 		async Task<bool> TryUpdateCachedFiltersAsync(CancellationToken cancellationToken)
 		{
-			DateTime utcNow = _clock.UtcNow;
 			IDatabase database = _redisService.GetDatabase();
 
 			IoHash[] filterHashes = await database.HashKeysAsync(Keys.Filters.Current);
@@ -559,7 +558,7 @@ namespace HordeServer.Agents
 			}
 		}
 
-		void UpdateFilters(ITransaction transaction, SessionId sessionId, IReadOnlyDictionary<IoHash, bool> oldFilters, IReadOnlyDictionary<IoHash, bool> newFilters)
+		static void UpdateFilters(ITransaction transaction, SessionId sessionId, IReadOnlyDictionary<IoHash, bool> oldFilters, IReadOnlyDictionary<IoHash, bool> newFilters)
 		{
 			// Create any new filters
 			foreach ((IoHash newFilter, bool newAvailable) in newFilters)
diff --git a/Engine/Source/Programs/UnrealGameSync/UnrealGameSync/AutomationServer.cs b/Engine/Source/Programs/UnrealGameSync/UnrealGameSync/AutomationServer.cs
--- a/Engine/Source/Programs/UnrealGameSync/UnrealGameSync/AutomationServer.cs
+++ b/Engine/Source/Programs/UnrealGameSync/UnrealGameSync/AutomationServer.cs
@@ -5,7 +5,6 @@ using System.IO;
 using System.IO.Pipes;
 using System.Net;
 using System.Net.Sockets;
-using System.Runtime.CompilerServices;
 using System.Text;
 using System.Threading;
 using System.Threading.Tasks;
