Horde: Dashboard side of bundled tool management
#skipci

diff --git a/Engine/Source/Programs/Horde/HordeDashboard/src/backend/Api.ts b/Engine/Source/Programs/Horde/HordeDashboard/src/backend/Api.ts
--- a/Engine/Source/Programs/Horde/HordeDashboard/src/backend/Api.ts
+++ b/Engine/Source/Programs/Horde/HordeDashboard/src/backend/Api.ts
@@ -4944,6 +4944,9 @@ export type GetToolSummaryResponse = {
 	/** Progress of the current deployment */
 	deploymentProgress?: number;
 
+	/** Whether the tool is bundled with the server */
+	bundled?: boolean;
+
 	/** Whether to show this tool for download inside UGS */
 	showInUGS: boolean;
 
@@ -4971,7 +4974,8 @@ export type GetToolResponse = {
 	platforms?: string[];
 	deployments: GetToolDeploymentResponse[];
 	public: boolean;
-	showInUgs: boolean;
+	bundled?: boolean;
+	showInUgs: boolean;	
 	showInDashboard: boolean;
 };
 
diff --git a/Engine/Source/Programs/Horde/HordeDashboard/src/components/ToolView.tsx b/Engine/Source/Programs/Horde/HordeDashboard/src/components/ToolView.tsx
--- a/Engine/Source/Programs/Horde/HordeDashboard/src/components/ToolView.tsx
+++ b/Engine/Source/Programs/Horde/HordeDashboard/src/components/ToolView.tsx
@@ -238,7 +238,7 @@ const ToolPanel: React.FC<{ selectedKey: string }> = observer(({ selectedKey })
             <Text style={{ color: modeColors.text }}>{tool.version}</Text>
          </Stack>
          <Stack >
-            <PrimaryButton split menuProps={downloadProps} style={{ width: 130, color: "#FFFFFF" }} text="Download" onClick={() => window.location.assign(`/api/v1/tools/${tool.id}?action=download`)} />
+            <PrimaryButton split={!tool.bundled} menuProps={!tool.bundled ? downloadProps : undefined} style={{ width: 130, color: "#FFFFFF" }} text="Download" onClick={() => window.location.assign(`/api/v1/tools/${tool.id}?action=download`)} />
          </Stack>
       </Stack>
    }
