Horde: Change tray app to use Avalonia rather than Winforms.

diff --git a/Engine/Source/Programs/Horde/HordeTrayApp/AgentPlugin.cs b/Engine/Source/Programs/Horde/HordeTrayApp/AgentPlugin.cs
--- a/Engine/Source/Programs/Horde/HordeTrayApp/AgentPlugin.cs
+++ b/Engine/Source/Programs/Horde/HordeTrayApp/AgentPlugin.cs
@@ -4,34 +4,33 @@ using System.Diagnostics;
 using System.IO.Pipes;
 using System.Runtime.InteropServices;
 using System.Text.Json;
+using Avalonia.Controls;
 using EpicGames.Core;
 using HordeAgent;
-using HordeTrayApp.Forms;
 using Microsoft.Win32;
 
 namespace HordeTrayApp
 {
-	record struct IdleStat(string Name, long Value, long MinValue);
-
 	class AgentPlugin : TrayAppPluginBase
 	{
+		record struct IdleStat(string Name, long Value, long MinValue);
+
 		readonly ITrayAppHost _host;
 
 		readonly BackgroundTask _clientTask;
 		readonly BackgroundTask _tickPauseStateTask;
 
-		readonly ToolStripMenuItem _enrollMenuItem;
-		readonly ToolStripMenuItem _statusEnabled;
-		readonly ToolStripMenuItem _statusDisabled;
-		readonly ToolStripMenuItem _statusWhenIdle;
+		readonly NativeMenuItem _enrollMenuItem;
+		readonly NativeMenuItem _statusEnabled;
+		readonly NativeMenuItem _statusDisabled;
+		readonly NativeMenuItem _statusWhenIdle;
 
-		readonly ToolStripMenuItem _statusMenuItem;
-		readonly ToolStripMenuItem _logsMenuItem;
+		readonly NativeMenuItem _statusMenuItem;
+		readonly NativeMenuItem _logsMenuItem;
 
 		readonly Settings _settings;
 
 		AgentSettingsMessage? _agentSettings;
-		IdleForm? _idleForm;
 
 		void EnrollWithServer()
 		{
@@ -47,33 +46,34 @@ namespace HordeTrayApp
 			_host = host;
 			_settings = LoadSettings();
 
-			_enrollMenuItem = new ToolStripMenuItem("Enroll with Server...");
+			_enrollMenuItem = new NativeMenuItem("Enroll with Server...");
 			_enrollMenuItem.Click += (s, e) => EnrollWithServer();
 
-			_statusEnabled = new ToolStripMenuItem("Enabled");
+			_statusEnabled = new NativeMenuItem("Enabled");
+			_statusEnabled.ToggleType = NativeMenuItemToggleType.CheckBox;
 			_statusEnabled.Click += (s, e) => SetUserStatus(UserStatus.Enabled);
 
-			_statusDisabled = new ToolStripMenuItem("Disabled");
+			_statusDisabled = new NativeMenuItem("Disabled");
+			_statusDisabled.ToggleType = NativeMenuItemToggleType.CheckBox;
 			_statusDisabled.Click += (s, e) => SetUserStatus(UserStatus.Disabled);
 
-			_statusWhenIdle = new ToolStripMenuItem("When Idle");
+			_statusWhenIdle = new NativeMenuItem("When Idle");
+			_statusWhenIdle.ToggleType = NativeMenuItemToggleType.CheckBox;
 			_statusWhenIdle.Click += (s, e) => SetUserStatus(UserStatus.WhenIdle);
 
-			ToolStripMenuItem showStatsMenuItem = new ToolStripMenuItem("Stats...");
-			showStatsMenuItem.Click += Status_Stats_OnClick;
-
-			_statusMenuItem = new ToolStripMenuItem("Status");
-			_statusMenuItem.DropDownItems.Add(_statusEnabled);
-			_statusMenuItem.DropDownItems.Add(_statusDisabled);
-			_statusMenuItem.DropDownItems.Add(_statusWhenIdle);
-			_statusMenuItem.DropDownItems.Add(new ToolStripSeparator());
-			_statusMenuItem.DropDownItems.Add(showStatsMenuItem);
+			_statusMenuItem = new NativeMenuItem("Status");
+			_statusMenuItem.Menu = new NativeMenu();
+			_statusMenuItem.Menu.Items.Add(_statusEnabled);
+			_statusMenuItem.Menu.Items.Add(_statusDisabled);
+			_statusMenuItem.Menu.Items.Add(_statusWhenIdle);
 
-			_logsMenuItem = new ToolStripMenuItem("Open logs dir");
+			_logsMenuItem = new NativeMenuItem("Open logs dir");
 			_logsMenuItem.Click += OnOpenLogs;
 
 			_clientTask = BackgroundTask.StartNew(StatusTaskAsync);
 			_tickPauseStateTask = BackgroundTask.StartNew(ctx => TickPauseStateAsync(ctx));
+
+			UpdateContextMenu(GetUserStatus());
 		}
 
 		private static Settings LoadSettings()
@@ -107,44 +107,27 @@ namespace HordeTrayApp
 			return result ?? new Settings();
 		}
 
-		public override void PopulateMenu(ContextMenuStrip contextMenu)
+		public override void PopulateContextMenu(NativeMenu contextMenu)
 		{
-			base.PopulateMenu(contextMenu);
+			base.PopulateContextMenu(contextMenu);
 
 			contextMenu.Items.Add(_enrollMenuItem);
 			contextMenu.Items.Add(_statusMenuItem);
-			contextMenu.Items.Add(new ToolStripSeparator());
+			contextMenu.Items.Add(new NativeMenuItemSeparator());
 			contextMenu.Items.Add(_logsMenuItem);
 		}
 
-		public override void UpdateMenu()
+		void UpdateContextMenu(UserStatus status)
 		{
-			base.UpdateMenu();
-
-			UserStatus status = GetUserStatus();
-			_statusEnabled.Checked = (status == UserStatus.Enabled);
-			_statusDisabled.Checked = (status == UserStatus.Disabled);
-			_statusWhenIdle.Checked = (status == UserStatus.WhenIdle);
+			_statusEnabled.IsChecked = (status == UserStatus.Enabled);
+			_statusDisabled.IsChecked = (status == UserStatus.Disabled);
+			_statusWhenIdle.IsChecked = (status == UserStatus.WhenIdle);
 		}
 
 		public override async ValueTask DisposeAsync()
 		{
 			await base.DisposeAsync();
 
-			if (_idleForm != null)
-			{
-				_idleForm.Dispose();
-				_idleForm = null;
-			}
-
-			_enrollMenuItem.Dispose();
-			_statusEnabled.Dispose();
-			_statusDisabled.Dispose();
-			_statusWhenIdle.Dispose();
-
-			_logsMenuItem.Dispose();
-			_statusMenuItem.Dispose();
-
 			await _tickPauseStateTask.DisposeAsync();
 			await _clientTask.DisposeAsync();
 		}
@@ -159,10 +142,6 @@ namespace HordeTrayApp
 				{
 					Process.Start(new ProcessStartInfo { FileName = logsDir.FullName, UseShellExecute = true });
 				}
-				else
-				{
-					MessageBox.Show("Unable to open logs dir " + logsDir.FullName, "Horde Tray App", MessageBoxButtons.OK, MessageBoxIcon.Error);
-				}
 			}
 		}
 
@@ -185,20 +164,7 @@ namespace HordeTrayApp
 		{
 			Registry.SetValue(RegistryKey, RegistryStatusValue, (int)status);
 			_statusChangedEvent.Set();
-		}
-
-		private void Status_Stats_OnClick(object? sender, EventArgs e)
-		{
-			if (_idleForm == null)
-			{
-				_idleForm = new IdleForm();
-				_idleForm.FormClosed += (s, e) =>
-				{
-					_idleForm.Dispose();
-					_idleForm = null;
-				};
-				_idleForm.Show();
-			}
+			UpdateContextMenu(status);
 		}
 
 		TrayAppPluginStatus? _status;
@@ -343,7 +309,7 @@ namespace HordeTrayApp
 				const int IdleTimeSecs = 30;
 				int stateChangeTime = (int)stateChangeTimer.Elapsed.TotalSeconds;
 				int stateChangeMaxTime = _enabled ? WakeTimeSecs : IdleTimeSecs;
-				_idleForm?.TickStats(_enabled, stateChangeTime, stateChangeMaxTime, idleStats);
+				//				_idleForm?.TickStats(_enabled, stateChangeTime, stateChangeMaxTime, idleStats);
 
 				if (userStatus == UserStatus.WhenIdle && stateChangeTime >= stateChangeMaxTime)
 				{
diff --git a/Engine/Source/Programs/Horde/HordeTrayApp/App.axaml.cs b/Engine/Source/Programs/Horde/HordeTrayApp/App.axaml.cs
new file mode 100644
--- /dev/null
+++ b/Engine/Source/Programs/Horde/HordeTrayApp/App.axaml.cs
@@ -0,0 +1,113 @@
+// Copyright Epic Games, Inc. All Rights Reserved.
+
+using Avalonia;
+using Avalonia.Controls;
+using Avalonia.Controls.ApplicationLifetimes;
+using Avalonia.Markup.Xaml;
+using Avalonia.Threading;
+
+namespace HordeTrayApp
+{
+	/// <summary>
+	/// Main application class
+	/// </summary>
+	public sealed partial class App : Application, ITrayAppHost
+	{
+		readonly List<ITrayAppPlugin> _plugins = new List<ITrayAppPlugin>();
+
+		WindowIcon? _normalIcon;
+		WindowIcon? _busyIcon;
+		WindowIcon? _pausedIcon;
+		WindowIcon? _errorIcon;
+
+		/// <summary>
+		/// Constructor
+		/// </summary>
+		public App()
+		{
+		}
+
+		/// <inheritdoc/>
+		public override void Initialize()
+		{
+			AvaloniaXamlLoader.Load(this);
+
+			_normalIcon = (WindowIcon)Resources["StatusNormal"]!;
+			_busyIcon = (WindowIcon)Resources["StatusBusy"]!;
+			_pausedIcon = (WindowIcon)Resources["StatusPaused"]!;
+			_errorIcon = (WindowIcon)Resources["StatusError"]!;
+
+			NativeMenuItem exitMenuItem = new NativeMenuItem("Exit");
+			exitMenuItem.Click += TrayIcon_Exit;
+
+			NativeMenu contextMenu = new NativeMenu();
+
+			TrayIcon trayIcon = TrayIcon.GetIcons(this)![0];
+			trayIcon.Menu = contextMenu;
+
+//			_icon = new TrayIcon();
+//			_icon.Icon = _normalIcon;
+//			_icon.Menu = contextMenu;
+//			_icon.ToolTipText = "Hello world";
+
+//			TrayIcons icons = new TrayIcons();
+//			icons.Add(_icon);
+
+//			TrayIcon.SetIcons(this, icons);
+
+			_plugins.Add(new AgentPlugin(this));
+
+//			NativeMenu contextMenu = trayIcon.Menu!;
+			foreach (ITrayAppPlugin plugin in _plugins)
+			{
+				plugin.PopulateContextMenu(contextMenu);
+			}
+			contextMenu.Items.Add(new NativeMenuItemSeparator());
+			contextMenu.Items.Add(exitMenuItem);
+		}
+
+		private void TrayIcon_Exit(object? sender, EventArgs e)
+		{
+			((IClassicDesktopStyleApplicationLifetime)ApplicationLifetime!).Shutdown();
+		}
+
+		/// <inheritdoc/>
+		public void UpdateStatus()
+		{
+			Dispatcher.UIThread.Post(() => UpdateStatusMainThread());
+		}
+
+		private void UpdateStatusMainThread()
+		{
+			TrayAppPluginState state = TrayAppPluginState.Undefined;
+			List<string> messages = new List<string>();
+
+			foreach (ITrayAppPlugin plugin in _plugins)
+			{
+				TrayAppPluginStatus status = plugin.GetStatus();
+				if (status.State >= state)
+				{
+					if (status.State > state)
+					{
+						messages.Clear();
+						state = status.State;
+					}
+					if (status.Message != null)
+					{
+						messages.Add(status.Message);
+					}
+				}
+			}
+
+			TrayIcon trayIcon = TrayIcon.GetIcons(this)![0];
+			trayIcon!.Icon = state switch
+			{
+				TrayAppPluginState.Busy => _busyIcon,
+				TrayAppPluginState.Paused => _pausedIcon,
+				TrayAppPluginState.Error => _errorIcon,
+				_ => _normalIcon
+			};
+			trayIcon.ToolTipText = String.Join("\n", messages);
+		}
+	}
+}
diff --git a/Engine/Source/Programs/Horde/HordeTrayApp/Forms/IdleForm.cs b/Engine/Source/Programs/Horde/HordeTrayApp/Forms/IdleForm.cs
deleted file mode 100644
diff --git a/Engine/Source/Programs/Horde/HordeTrayApp/HordeTrayApp.csproj b/Engine/Source/Programs/Horde/HordeTrayApp/HordeTrayApp.csproj
--- a/Engine/Source/Programs/Horde/HordeTrayApp/HordeTrayApp.csproj
+++ b/Engine/Source/Programs/Horde/HordeTrayApp/HordeTrayApp.csproj
@@ -5,10 +5,10 @@
     <TargetFramework>net6.0-windows</TargetFramework>
     <Nullable>enable</Nullable>
     <ImplicitUsings>enable</ImplicitUsings>
-    <UseWindowsForms>true</UseWindowsForms>
     <NeutralLanguage>en-US</NeutralLanguage>
     <Configurations>Debug;Release;Analyze</Configurations>
 	<AssemblyName>Horde.Agent.TrayApp</AssemblyName>
+    <AvaloniaUseCompiledBindingsByDefault>true</AvaloniaUseCompiledBindingsByDefault>
   </PropertyGroup>
 
   <PropertyGroup Condition=" '$(Configuration)|$(Platform)' == 'Analyze|AnyCPU' ">
@@ -41,5 +41,18 @@
       <LastGenOutput>Resources.Designer.cs</LastGenOutput>
     </EmbeddedResource>
   </ItemGroup>
+  
+  <ItemGroup>
+    <AvaloniaResource Include="Resources\**" />
+  </ItemGroup>
+
+  <ItemGroup>
+    <PackageReference Include="Avalonia" Version="11.1.0" />
+    <PackageReference Include="Avalonia.Desktop" Version="11.1.0" />
+    <PackageReference Include="Avalonia.Themes.Fluent" Version="11.1.0" />
+    <PackageReference Include="Avalonia.Fonts.Inter" Version="11.1.0" />
+    <!--Condition below is needed to remove Avalonia.Diagnostics package from build output in Release configuration.-->
+    <PackageReference Condition="'$(Configuration)' == 'Debug'" Include="Avalonia.Diagnostics" Version="11.1.0" />
+  </ItemGroup>
 
 </Project>
\ No newline at end of file
diff --git a/Engine/Source/Programs/Horde/HordeTrayApp/ITrayAppPlugin.cs b/Engine/Source/Programs/Horde/HordeTrayApp/ITrayAppPlugin.cs
--- a/Engine/Source/Programs/Horde/HordeTrayApp/ITrayAppPlugin.cs
+++ b/Engine/Source/Programs/Horde/HordeTrayApp/ITrayAppPlugin.cs
@@ -1,5 +1,7 @@
 // Copyright Epic Games, Inc. All Rights Reserved.
 
+using Avalonia.Controls;
+
 namespace HordeTrayApp
 {
 	enum TrayAppPluginState
@@ -29,12 +31,7 @@ namespace HordeTrayApp
 		/// <summary>
 		/// Allow the plugin to customize the tray icon context menu
 		/// </summary>
-		void PopulateMenu(ContextMenuStrip contextMenu);
-
-		/// <summary>
-		/// Update the menu state before it's shown
-		/// </summary>
-		void UpdateMenu();
+		void PopulateContextMenu(NativeMenu contextMenu);
 	}
 
 	/// <summary>
@@ -49,12 +46,7 @@ namespace HordeTrayApp
 		}
 
 		/// <inheritdoc/>
-		public virtual void PopulateMenu(ContextMenuStrip contextMenu)
-		{
-		}
-
-		/// <inheritdoc/>
-		public virtual void UpdateMenu()
+		public virtual void PopulateContextMenu(NativeMenu contextMenu)
 		{
 		}
 
diff --git a/Engine/Source/Programs/Horde/HordeTrayApp/Program.cs b/Engine/Source/Programs/Horde/HordeTrayApp/Program.cs
--- a/Engine/Source/Programs/Horde/HordeTrayApp/Program.cs
+++ b/Engine/Source/Programs/Horde/HordeTrayApp/Program.cs
@@ -2,8 +2,12 @@
 
 using System.Diagnostics;
 using System.Reflection;
+using System.Runtime.InteropServices;
+using Avalonia;
+using Avalonia.Controls;
 using EpicGames.Core;
-using HordeTrayApp.Properties;
+
+#pragma warning disable CA1806 // Not checking return code from MessageBox
 
 namespace HordeTrayApp
 {
@@ -12,6 +16,9 @@ namespace HordeTrayApp
 		const string MutexName = "Horde.Agent.TrayApp-Mutex";
 		const string EventName = "Horde.Agent.TrayApp-Exit";
 
+		[DllImport("user32.dll", SetLastError = true, CharSet = CharSet.Auto)]
+		public static extern int MessageBox(IntPtr hWnd, string text, string? caption, uint type);
+
 		[STAThread]
 		public static int Main(string[] args)
 		{
@@ -54,14 +61,16 @@ namespace HordeTrayApp
 						}
 						catch (Exception ex)
 						{
-							MessageBox.Show($"Unable to copy app to temp location. Error:\n\n{ex}");
+							MessageBox(IntPtr.Zero, $"Unable to copy app to temp location. Error:\n\n{ex}", null, 0);
 							return 1;
 						}
 					}
 				}
 			}
 
-			MainAsync(closeEvent).GetAwaiter().GetResult();
+			BuildAvaloniaApp()
+				.StartWithClassicDesktopLifetime(Array.Empty<string>(), ShutdownMode.OnExplicitShutdown);
+
 			return 0;
 		}
 
@@ -80,14 +89,12 @@ namespace HordeTrayApp
 			}
 		}
 
-		static async Task MainAsync(EventWaitHandle closeEvent)
+		// Avalonia configuration, don't remove; also used by visual designer.
+		public static AppBuilder BuildAvaloniaApp()
 		{
-			ApplicationConfiguration.Initialize();
-
-			await using (CustomApplicationContext appContext = new CustomApplicationContext(closeEvent))
-			{
-				Application.Run(appContext);
-			}
+			return AppBuilder.Configure<App>()
+				.UsePlatformDetect()
+				.LogToTrace();
 		}
 	}
 
@@ -132,141 +139,4 @@ namespace HordeTrayApp
 			_mutex.Dispose();
 		}
 	}
-
-	class CustomApplicationContext : ApplicationContext, ITrayAppHost, IAsyncDisposable
-	{
-		readonly NotifyIcon _trayIcon;
-		readonly Control _mainThreadInvokeTarget;
-		readonly List<ITrayAppPlugin> _plugins = new List<ITrayAppPlugin>();
-		readonly BackgroundTask _waitForExitTask;
-
-		bool _disposed;
-
-		public CustomApplicationContext(EventWaitHandle eventHandle)
-		{
-			_plugins.Add(new AgentPlugin(this));
-
-			_mainThreadInvokeTarget = new Control();
-			_mainThreadInvokeTarget.CreateControl();
-
-			ToolStripMenuItem exitMenuItem = new ToolStripMenuItem("Exit");
-			exitMenuItem.Click += OnExit;
-
-			ContextMenuStrip menu = new ContextMenuStrip();
-			foreach (ITrayAppPlugin plugin in _plugins)
-			{
-				plugin.PopulateMenu(menu);
-			}
-			menu.Items.Add(new ToolStripSeparator());
-			menu.Items.Add(exitMenuItem);
-
-			_trayIcon = new NotifyIcon()
-			{
-				Icon = Resources.StatusNormal,
-				ContextMenuStrip = menu,
-				Visible = true
-			};
-			_trayIcon.Click += TrayIcon_Click;
-
-			_waitForExitTask = BackgroundTask.StartNew(ctx => WaitForExitAsync(eventHandle, ctx));
-		}
-
-		protected override void Dispose(bool disposing)
-		{
-			base.Dispose(disposing);
-
-			if (disposing)
-			{
-				_mainThreadInvokeTarget.Dispose();
-				_trayIcon.Dispose();
-				_disposed = true;
-			}
-		}
-
-		private void TrayIcon_Click(object? sender, EventArgs e)
-		{
-			foreach (AgentPlugin plugin in _plugins)
-			{
-				plugin.UpdateMenu();
-			}
-		}
-
-		public async ValueTask DisposeAsync()
-		{
-			await _waitForExitTask.DisposeAsync();
-
-			foreach (AgentPlugin plugin in _plugins)
-			{
-				await plugin.DisposeAsync();
-			}
-
-			Dispose();
-			GC.SuppressFinalize(this);
-		}
-
-		private void OnExit(object? sender, EventArgs e)
-		{
-			ExitThread();
-		}
-
-		public void UpdateStatus()
-		{
-			_mainThreadInvokeTarget.BeginInvoke(() => SetStatus_MainThread());
-		}
-
-		void SetStatus_MainThread()
-		{
-			if (!_disposed)
-			{
-				TrayAppPluginState state = TrayAppPluginState.Undefined;
-				List<string> messages = new List<string>();
-
-				foreach (ITrayAppPlugin plugin in _plugins)
-				{
-					TrayAppPluginStatus status = plugin.GetStatus();
-					if (status.State >= state)
-					{
-						if (status.State > state)
-						{
-							messages.Clear();
-							state = status.State;
-						}
-						if (status.Message != null)
-						{
-							messages.Add(status.Message);
-						}
-					}
-				}
-
-				_trayIcon.Icon = state switch
-				{
-					TrayAppPluginState.Busy => Resources.StatusBusy,
-					TrayAppPluginState.Paused => Resources.StatusPaused,
-					TrayAppPluginState.Error => Resources.StatusError,
-					_ => Resources.StatusNormal
-				};
-				_trayIcon.Text = String.Join("\n", messages);
-			}
-		}
-
-		async Task WaitForExitAsync(EventWaitHandle eventHandle, CancellationToken cancellationToken)
-		{
-			await eventHandle.WaitOneAsync(cancellationToken);
-			_mainThreadInvokeTarget.BeginInvoke(() => Exit_MainThread());
-		}
-
-		void Exit_MainThread()
-		{
-			if (!_disposed)
-			{
-				_trayIcon.Visible = false;
-				Application.Exit();
-			}
-		}
-
-		void Exit(object sender, EventArgs e)
-		{
-			Exit_MainThread();
-		}
-	}
 }
