Horde: Fix cancellation of leases which are no longer running on an agent.

#rnx

diff --git a/Engine/Source/Programs/Horde/Plugins/Compute/HordeServer.Compute/Agents/AgentService.cs b/Engine/Source/Programs/Horde/Plugins/Compute/HordeServer.Compute/Agents/AgentService.cs
--- a/Engine/Source/Programs/Horde/Plugins/Compute/HordeServer.Compute/Agents/AgentService.cs
+++ b/Engine/Source/Programs/Horde/Plugins/Compute/HordeServer.Compute/Agents/AgentService.cs
@@ -1,5 +1,6 @@
 // Copyright Epic Games, Inc. All Rights Reserved.
 
+using System.Diagnostics.CodeAnalysis;
 using System.Diagnostics.Metrics;
 using System.Security.Claims;
 using EpicGames.Core;
@@ -571,30 +572,36 @@ namespace HordeServer.Agents
 		/// <summary>
 		/// Cancels the specified agent lease
 		/// </summary>
-		/// <param name="agent">Agent to cancel the lease on</param>
-		/// <param name="leaseId">The lease id to cancel</param>
+		/// <param name="lease">The lease to cancel</param>
 		/// <param name="cancellationToken">Cancellation token for the operation</param>
 		/// <returns></returns>
-		public async Task<bool> CancelLeaseAsync(IAgent agent, LeaseId leaseId, CancellationToken cancellationToken)
+		public async Task<bool> CancelLeaseAsync(ILease lease, CancellationToken cancellationToken)
 		{
-			_ = this;
-
-			int index = 0;
-			while (index < agent.Leases.Count && agent.Leases[index].Id != leaseId)
+			// If it's already complete, don't do anything
+			if (lease.FinishTime != null)
 			{
-				index++;
+				return false;
 			}
 
-			if (index == agent.Leases.Count)
+			// Cancel it on the agent that's currently executing it
+			IAgent? agent = await GetAgentAsync(lease.AgentId, cancellationToken);
+			if (agent == null)
 			{
+				_logger.LogWarning("Forcing lease {LeaseId} to cancelled state; agent {AgentId} not found", lease.Id, lease.AgentId);
+				await _leases.TrySetOutcomeAsync(lease.Id, _clock.UtcNow, LeaseOutcome.Cancelled, null, cancellationToken);
 				return false;
 			}
 
-			if (agent.Leases[index].State == LeaseState.Cancelled)
+			// Find the index of the lease to remove
+			int index = agent.Leases.FindIndex(x => x.Id == lease.Id);
+			if (index == -1)
 			{
+				_logger.LogWarning("Forcing lease {LeaseId} to cancelled state; no longer running on agent {AgentId}", lease.Id, lease.AgentId);
+				await _leases.TrySetOutcomeAsync(lease.Id, _clock.UtcNow, LeaseOutcome.Cancelled, null, cancellationToken);
 				return false;
 			}
 
+			// Update the agent
 			await agent.TryCancelLeaseAsync(index, cancellationToken);
 			return true;
 		}
diff --git a/Engine/Source/Programs/Horde/Plugins/Compute/HordeServer.Compute/Agents/Leases/LeasesController.cs b/Engine/Source/Programs/Horde/Plugins/Compute/HordeServer.Compute/Agents/Leases/LeasesController.cs
--- a/Engine/Source/Programs/Horde/Plugins/Compute/HordeServer.Compute/Agents/Leases/LeasesController.cs
+++ b/Engine/Source/Programs/Horde/Plugins/Compute/HordeServer.Compute/Agents/Leases/LeasesController.cs
@@ -232,19 +232,7 @@ namespace HordeServer.Agents.Leases
 				return NotFound(leaseId);
 			}
 
-			IAgent? agent = await _agentService.GetAgentAsync(lease.AgentId, cancellationToken);
-			if (agent == null)
-			{
-				return NotFound(lease.AgentId);
-			}
-
-			IAgentLease? agentLease = agent.Leases.FirstOrDefault(x => x.Id == leaseId);
-			if (agentLease == null)
-			{
-				return NotFound(agent.Id, leaseId);
-			}
-
-			await _agentService.CancelLeaseAsync(agent, leaseId, cancellationToken);
+			await _agentService.CancelLeaseAsync(lease, cancellationToken);
 			return Ok();
 		}
 	}
