Horde: Updating schema documentation

#rnx

diff --git a/Engine/Source/Programs/Horde/Docs/Config/Schema/Globals.md b/Engine/Source/Programs/Horde/Docs/Config/Schema/Globals.md
--- a/Engine/Source/Programs/Horde/Docs/Config/Schema/Globals.md
+++ b/Engine/Source/Programs/Horde/Docs/Config/Schema/Globals.md
@@ -95,12 +95,12 @@ Information about an individual Perforce server
 
 Name | Description
 ---- | -----------
-`serverAndPort` | `string`<br>The server and port. The server may be a DNS entry with multiple records, in which case it will be actively load balanced.
+`serverAndPort` | `string`<br>The server and port. The server may be a DNS entry with multiple records, in which case it will be actively load balanced. If "ssl:" prefix is used, ensure P4 server's fingerprint/certificate is trusted (see docs for "p4 trust")
 `healthCheck` | `boolean`<br>Whether to query the healthcheck address under each server
 `resolveDns` | `boolean`<br>Whether to resolve the DNS entries and load balance between different hosts
 `maxConformCount` | `integer`<br>Maximum number of simultaneous conforms on this server
-`condition` | `string`<br>Optional condition for a machine to be eligable to use this server
-`properties` | `string[]`<br>List of properties for an agent to be eligable to use this server
+`condition` | `string`<br>Optional condition for a machine to be eligible to use this server
+`properties` | `string[]`<br>List of properties for an agent to be eligible to use this server
 
 ## PerforceCredentials
 
diff --git a/Engine/Source/Programs/Horde/Docs/Deployment/AgentSettings.md b/Engine/Source/Programs/Horde/Docs/Deployment/AgentSettings.md
--- a/Engine/Source/Programs/Horde/Docs/Deployment/AgentSettings.md
+++ b/Engine/Source/Programs/Horde/Docs/Deployment/AgentSettings.md
@@ -25,6 +25,8 @@ Name | Description
 `enableTelemetry` | `boolean`<br>Whether to send telemetry back to Horde server
 `telemetryReportInterval` | `integer`<br>How often to report telemetry events to server in milliseconds
 `bundleCacheSize` | `integer`<br>Maximum size of the bundle cache, in megabytes.
+`cpuCount` | `integer`<br>Maximum number of logical CPU cores workloads should use Currently this is only provided as a hint and requires leases to respect this value as it's set via an env variable (UE_HORDE_CPU_COUNT).
+`cpuMultiplier` | `number`<br>CPU core multiplier applied to CPU core count setting For example, 32 CPU cores and a multiplier of 0.5 results in max 16 CPU usage.
 `properties` | `string` `->` `string`<br>Key/value properties in addition to those set internally by the agent
 
 ## ServerProfile
