Horde: Document the look up of P4TRUST file

#skipci

diff --git a/Engine/Source/Programs/Horde/Docs/Deployment/Integrations/Perforce.md b/Engine/Source/Programs/Horde/Docs/Deployment/Integrations/Perforce.md
--- a/Engine/Source/Programs/Horde/Docs/Deployment/Integrations/Perforce.md
+++ b/Engine/Source/Programs/Horde/Docs/Deployment/Integrations/Perforce.md
@@ -32,7 +32,6 @@ There are several configurables for each cluster:
 * `Name`: Used to reference the cluster from a stream in the CI system.
 * `Servers`: Each server supports several settings of its own:
   * `ServerAndPort`: The server part may be a DNS entry with multiple records, in which case it will be actively load balanced.
-    If `ssl:` prefix is used, ensure P4 server's fingerprint/certificate is trusted (see docs for [p4 trust](https://www.perforce.com/manuals/cmdref/Content/CmdRef/p4_trust.html))
   * `ResolveDns`: If true, the given DNS name is resolved to find a concrete list of servers to be used. This allows
   IT/infrastructure teams to add and remove servers to a cluster without reconfiguring Horde.
   * `Properties`: Specifies properties that the agent must have to select this server.
@@ -45,6 +44,17 @@ There are several configurables for each cluster:
 * `CanImpersonate`: Indicates whether Horde should attempt to impersonate other users when submitting changes after a
   successful preflight-and-submit operation. Typically requires an administrator account.
 
+## Servers with SSL
+
+When the Perforce server uses SSL and the `ssl:` prefix is specified,
+it is essential to ensure that the server's fingerprint or certificate is trusted.
+For more information, refer to the [p4 trust](https://www.perforce.com/manuals/cmdref/Content/CmdRef/p4_trust.html) documentation.
+Depending on the configuration, it may be necessary to explicitly set the `P4TRUST` environment variable,
+which should point to a valid `p4trust.txt` file. Again, see the P4 documentation for further details.
+
+For example, on Windows, the Horde server runs as a service by default.
+In this scenario, the default location of the trust file resolves to `%SystemRoot%\System32\config\systemprofile\p4trust.txt`.
+
 ## Health Checks
 
 If enabled, health checks for Perforce servers are performed by performing an HTTP `GET` request to
diff --git a/Engine/Source/Programs/Horde/HordeServer/Defaults/globals.json b/Engine/Source/Programs/Horde/HordeServer/Defaults/globals.json
--- a/Engine/Source/Programs/Horde/HordeServer/Defaults/globals.json
+++ b/Engine/Source/Programs/Horde/HordeServer/Defaults/globals.json
@@ -41,6 +41,7 @@
 //				{
 //					// DNS name and port for the server.
 //					// If connecting using "ssl:" prefix, ensure P4 server's fingerprint/certificate is trusted (see docs for "p4 trust")
+//					// See Horde's documentation on connecting to SSL-enabled Perforce servers.
 //					"serverAndPort": "perforce:1666"
 //				}
 //			],
diff --git a/Engine/Source/Programs/Horde/Plugins/Build/HordeServer.Build/BuildConfig.cs b/Engine/Source/Programs/Horde/Plugins/Build/HordeServer.Build/BuildConfig.cs
--- a/Engine/Source/Programs/Horde/Plugins/Build/HordeServer.Build/BuildConfig.cs
+++ b/Engine/Source/Programs/Horde/Plugins/Build/HordeServer.Build/BuildConfig.cs
@@ -356,7 +356,8 @@ namespace HordeServer
 	{
 		/// <summary>
 		/// The server and port. The server may be a DNS entry with multiple records, in which case it will be actively load balanced.
-		/// If "ssl:" prefix is used, ensure P4 server's fingerprint/certificate is trusted (see docs for "p4 trust")
+		/// If "ssl:" prefix is used, ensure P4 server's fingerprint/certificate is trusted.
+		/// See Horde's documentation on connecting to SSL-enabled Perforce servers.
 		/// </summary>
 		public string ServerAndPort { get; set; } = "perforce:1666";
 
