Horde: Set the UE_HORDE_STREAMID environment variable containing the current stream id when running under Horde.

diff --git a/Engine/Source/Programs/Horde/Drivers/JobDriver/Execution/JobExecutor.cs b/Engine/Source/Programs/Horde/Drivers/JobDriver/Execution/JobExecutor.cs
--- a/Engine/Source/Programs/Horde/Drivers/JobDriver/Execution/JobExecutor.cs
+++ b/Engine/Source/Programs/Horde/Drivers/JobDriver/Execution/JobExecutor.cs
@@ -1803,6 +1803,7 @@ namespace JobDriver.Execution
 			newEnvVars["UE_HORDE_JOBID"] = JobId.ToString();
 			newEnvVars["UE_HORDE_BATCHID"] = BatchId.ToString();
 			newEnvVars["UE_HORDE_STEPID"] = step.StepId.ToString();
+			newEnvVars["UE_HORDE_STREAMID"] = Batch.StreamId;
 
 			// Enable structured logging output
 			newEnvVars["UE_LOG_JSON_TO_STDOUT"] = "1";
diff --git a/Engine/Source/Programs/Horde/Plugins/Build/HordeServer.Build/Jobs/JobRpcService.cs b/Engine/Source/Programs/Horde/Plugins/Build/HordeServer.Build/Jobs/JobRpcService.cs
--- a/Engine/Source/Programs/Horde/Plugins/Build/HordeServer.Build/Jobs/JobRpcService.cs
+++ b/Engine/Source/Programs/Horde/Plugins/Build/HordeServer.Build/Jobs/JobRpcService.cs
@@ -196,6 +196,7 @@ namespace HordeServer.Jobs
 			RpcBeginBatchResponse response = new RpcBeginBatchResponse();
 			response.LogId = batch.LogId.ToString();
 			response.AgentType = graph.Groups[batch.GroupIdx].AgentType;
+			response.StreamId = streamConfig.Id.ToString();
 			response.StreamName = streamConfig.Name;
 			response.Change = job.CommitId.GetPerforceChange();
 			response.CodeChange = job.CodeCommitId?.GetPerforceChange() ?? 0;
