Horde: Updating schema documentation

#rnx

diff --git a/Engine/Source/Programs/Horde/Docs/Config/Schema/Globals.md b/Engine/Source/Programs/Horde/Docs/Config/Schema/Globals.md
--- a/Engine/Source/Programs/Horde/Docs/Config/Schema/Globals.md
+++ b/Engine/Source/Programs/Horde/Docs/Config/Schema/Globals.md
@@ -95,7 +95,7 @@ Information about an individual Perforce server
 
 Name | Description
 ---- | -----------
-`serverAndPort` | `string`<br>The server and port. The server may be a DNS entry with multiple records, in which case it will be actively load balanced. If "ssl:" prefix is used, ensure P4 server's fingerprint/certificate is trusted (see docs for "p4 trust")
+`serverAndPort` | `string`<br>The server and port. The server may be a DNS entry with multiple records, in which case it will be actively load balanced. If "ssl:" prefix is used, ensure P4 server's fingerprint/certificate is trusted. See Horde's documentation on connecting to SSL-enabled Perforce servers.
 `healthCheck` | `boolean`<br>Whether to query the healthcheck address under each server
 `resolveDns` | `boolean`<br>Whether to resolve the DNS entries and load balance between different hosts
 `maxConformCount` | `integer`<br>Maximum number of simultaneous conforms on this server
