Horde: Improve issue button rendering to not obscure the structured logging eyeball, also change where we render the structured log
#skipci

diff --git a/Engine/Source/Programs/Horde/HordeDashboard/src/components/LogView.tsx b/Engine/Source/Programs/Horde/HordeDashboard/src/components/LogView.tsx
--- a/Engine/Source/Programs/Horde/HordeDashboard/src/components/LogView.tsx
+++ b/Engine/Source/Programs/Horde/HordeDashboard/src/components/LogView.tsx
@@ -790,24 +790,23 @@ export const LogList: React.FC<{ logId: string }> = observer(({ logId }) => {
                fontFamily: "Horde Open Sans SemiBold, sans-serif, sans-serif", color: error ? "#FFFFFF" : modeColors.text, fontSize: handler.fontSize, textDecoration: !item!.issue?.resolvedAt ? undefined : "line-through"
             }
 
-            return <Stack>
+            return <Stack style={{width: 80}}>
                <TooltipHost
                   content={timestamp}
                   id={tooltipId}
                   calloutProps={{ gapSpace: 0 }}
-                  styles={{ root: { display: 'inline-block' } }}>
+                  styles={{ root: { display: 'inline-block', width: 80 } }}>
                   <DefaultButton className={error ? styles.errorButton : styles.warningButton}
                      href={href}
-                     style={{ padding: 0, margin: 0, width: tsWidth, paddingLeft: 8, paddingRight: 8, height: "100%", fontWeight: "unset" }}
-
+                     style={{ padding: 0, margin: 0, minWidth: 65, width: 65, paddingLeft: 4, paddingRight: 6, height: "100%", fontWeight: "unset" }}
                      onClick={(ev) => {
                         ev.preventDefault();
                         ev.stopPropagation();
-                        location.search = `?issue=${issueId}`;
+                        location.search = `?issue=${issueId}`
                         setIssueHistory(true);
                         navigate(location);
                      }}>
-                     <Text variant="small" style={{ ...fontStyle }}>Issue</Text><div style={{ ...fontStyle }}>&nbsp;</div><Text variant="small" style={{ ...fontStyle }}>{`${issueId}`}</Text>
+                     <Text variant="small" style={{ ...fontStyle }}>{`${issueId}`}</Text>
                   </DefaultButton>
                </TooltipHost>
             </Stack>
@@ -849,7 +848,7 @@ export const LogList: React.FC<{ logId: string }> = observer(({ logId }) => {
                            {handler.infoLine === item.lineNumber && <Callout
                               styles={{ root: { padding: "32px 24px", maxWidth: 1300 } }}
                               role="dialog"
-                              gapSpace={12}
+                              gapSpace={4}
                               target={`#callout_target_${item?.lineNumber}`}
                               isBeakVisible={true}
                               beakWidth={12}
@@ -857,7 +856,7 @@ export const LogList: React.FC<{ logId: string }> = observer(({ logId }) => {
                                  handler.infoLine = undefined;
                                  handler.externalUpdate();
                               }}
-                              directionalHint={DirectionalHint.rightCenter}
+                              directionalHint={DirectionalHint.bottomCenter}
                               setInitialFocus>
                               <Stack style={{ maxWidth: 1140 }}>
                                  <Stack style={{ paddingBottom: 24 }}>
