Horde: Support loading artifacts out of log which are not in the step artifacts
#skipci

diff --git a/Engine/Source/Programs/Horde/HordeDashboard/src/components/LogRender.tsx b/Engine/Source/Programs/Horde/HordeDashboard/src/components/LogRender.tsx
--- a/Engine/Source/Programs/Horde/HordeDashboard/src/components/LogRender.tsx
+++ b/Engine/Source/Programs/Horde/HordeDashboard/src/components/LogRender.tsx
@@ -270,8 +270,15 @@ const renderTags = (navigate: NavigateFunction, line: LogLine, lineNumber: numbe
       } else if (tagType === TagType.ArtifactId) {
          const artifactType = properties.ArtifactType;
          if (artifactType?.length) {
+
             const search = new URLSearchParams(window.location.search);
             search.set("artifactContext", encodeURIComponent(artifactType as string));
+            if (properties.ArtifactId) {
+               if (properties.ArtifactId["$text"]) {
+                  search.set("artifactId", encodeURIComponent(properties.ArtifactId["$text"] as string));
+               }               
+            }
+            
             const url = `${window.location.pathname}?` + search.toString();
             return <a key={key} href="/" onClick={(ev) => { ev.stopPropagation(); ev.preventDefault(); navigate(url, { replace: true }) }}><Highlight search={search ? search : ""} className={logStyle.logLine}>{text}</Highlight></a>;
          }
diff --git a/Engine/Source/Programs/Horde/HordeDashboard/src/components/LogView.tsx b/Engine/Source/Programs/Horde/HordeDashboard/src/components/LogView.tsx
--- a/Engine/Source/Programs/Horde/HordeDashboard/src/components/LogView.tsx
+++ b/Engine/Source/Programs/Horde/HordeDashboard/src/components/LogView.tsx
@@ -598,6 +598,7 @@ export const LogList: React.FC<{ logId: string }> = observer(({ logId }) => {
 
    const artifactContext = !!query.get("artifactContext") ? query.get("artifactContext")! as ArtifactContextType : undefined;
    const artifactPath = !!query.get("artifactPath") ? query.get("artifactPath")! : undefined;
+   const artifactId = !!query.get("artifactId") ? query.get("artifactId")! : undefined;
 
    globalHandler = handler;
    globalSearchState = searchState;
@@ -1216,7 +1217,7 @@ export const LogList: React.FC<{ logId: string }> = observer(({ logId }) => {
    return <Stack>
       {!!fixme && <IssueModalV2 issueId={query.get("issue")} popHistoryOnClose={issueHistory} />}
       {!!fixme && logHistory && <StepHistoryModal jobDetails={fixme!} stepId={fixme!.stepByLogId(logId)?.id} onClose={() => setLogHistory(false)} />}
-      {!!fixme && !!artifactContext && logArtifacts !== "legacy" && <JobArtifactsModal jobId={fixme!.jobdata!.id} stepId={fixme!.stepByLogId(logId)?.id!} artifacts={(logSource as JobLogSource).artifactsV2} contextType={artifactContext} artifactPath={artifactPath} onClose={() => { navigate(window.location.pathname, { replace: true }) }} />}
+      {!!fixme && !!artifactContext && logArtifacts !== "legacy" && <JobArtifactsModal jobId={fixme!.jobdata!.id} stepId={fixme!.stepByLogId(logId)?.id!} artifacts={(logSource as JobLogSource).artifactsV2} contextType={artifactContext} artifactPath={artifactPath} artifactId={artifactId} onClose={() => { navigate(window.location.pathname, { replace: true }) }} />}
       {!!historyAgentId && <HistoryModal agentId={historyAgentId} onDismiss={() => { navigate(baseUrl, { replace: true }); setHistoryAgentId(undefined) }} />}
       <Breadcrumbs items={logSource?.crumbs ?? []} title={logSource?.crumbTitle} />
       <Stack tokens={{ childrenGap: 0 }} style={{ backgroundColor: hordeTheme.horde.contentBackground, paddingTop: 12 }}>
diff --git a/Engine/Source/Programs/Horde/HordeDashboard/src/components/artifacts/ArtifactsModal.tsx b/Engine/Source/Programs/Horde/HordeDashboard/src/components/artifacts/ArtifactsModal.tsx
--- a/Engine/Source/Programs/Horde/HordeDashboard/src/components/artifacts/ArtifactsModal.tsx
+++ b/Engine/Source/Programs/Horde/HordeDashboard/src/components/artifacts/ArtifactsModal.tsx
@@ -82,6 +82,15 @@ class ArtifactsHandler {
          return;
       }
 
+      if (this.artifactId && !artifacts.find(a => a.id === this.artifactId)) {
+         try {
+            const artifact = await backend.getArtifactData(this.artifactId);
+            artifacts.push(artifact);
+         } catch (err) {
+            console.log(err);
+         }         
+      }
+
       this.artifacts = artifacts;
 
       if (!this.context) {
@@ -100,7 +109,6 @@ class ArtifactsHandler {
          a = artifacts.find(a => a.type === this.context)!;
       }
 
-
       if (!a) {
          console.error("Unable to find artifact for context", this.context, artifacts);
          this.artifactMissing = true;
