Horde: Add an IArtifact.Content property which can be used to access the artifact data directly.

#rnx

diff --git a/Engine/Source/Programs/Horde/HordeCmd/Commands/Artifacts/ArtifactDownload.cs b/Engine/Source/Programs/Horde/HordeCmd/Commands/Artifacts/ArtifactDownload.cs
--- a/Engine/Source/Programs/Horde/HordeCmd/Commands/Artifacts/ArtifactDownload.cs
+++ b/Engine/Source/Programs/Horde/HordeCmd/Commands/Artifacts/ArtifactDownload.cs
@@ -39,9 +39,13 @@ namespace Horde.Commands.Artifacts
 
 		public override async Task<int> ExecuteAsync(ILogger logger)
 		{
-			HordeHttpClient httpClient = _hordeClient.CreateHttpClient();
+			IArtifact? artifact = await _hordeClient.Artifacts.GetAsync(Id);
+			if (artifact == null)
+			{
+				logger.LogError("Artifact {Id} not found", Id);
+				return 1;
+			}
 
-			GetArtifactResponse artifact = await httpClient.GetArtifactAsync(Id);
 			logger.LogInformation("Downloading artifact {Id}: {Description}", Id, artifact.Description);
 
 			if (CleanOutput)
@@ -50,10 +54,9 @@ namespace Horde.Commands.Artifacts
 				FileUtils.ForceDeleteDirectoryContents(OutputDir);
 			}
 
-			IStorageClient store = _hordeClient.CreateStorageClient(artifact.Id);
-
 			Stopwatch timer = Stopwatch.StartNew();
 
+			IStorageClient store = _hordeClient.CreateStorageClient(artifact.Id);
 			IBlobRef<DirectoryNode> handle = await store.ReadRefAsync<DirectoryNode>(new RefName("default"));
 			await handle.ExtractAsync(OutputDir.ToDirectoryInfo(), new ExtractStatsLogger(logger), logger, CancellationToken.None);
 
diff --git a/Engine/Source/Programs/Horde/Plugins/Build/HordeServer.Build/Artifacts/ArtifactCollection.cs b/Engine/Source/Programs/Horde/Plugins/Build/HordeServer.Build/Artifacts/ArtifactCollection.cs
--- a/Engine/Source/Programs/Horde/Plugins/Build/HordeServer.Build/Artifacts/ArtifactCollection.cs
+++ b/Engine/Source/Programs/Horde/Plugins/Build/HordeServer.Build/Artifacts/ArtifactCollection.cs
@@ -6,6 +6,7 @@ using EpicGames.Core;
 using EpicGames.Horde.Artifacts;
 using EpicGames.Horde.Commits;
 using EpicGames.Horde.Storage;
+using EpicGames.Horde.Storage.Nodes;
 using EpicGames.Horde.Streams;
 using HordeServer.Commits;
 using HordeServer.Projects;
@@ -46,6 +47,9 @@ namespace HordeServer.Artifacts
 			RefName IArtifact.RefName => _document.RefName;
 			DateTime IArtifact.CreatedAtUtc => (_document.CreatedAtUtc == default) ? _document.Id.CreationTime : _document.CreatedAtUtc;
 
+			IBlobRef<DirectoryNode> IArtifact.Content
+				=> _collection.Open(this);
+
 			public Artifact(ArtifactCollection collection, ArtifactDocument document)
 			{
 				_collection = collection;
@@ -184,6 +188,12 @@ namespace HordeServer.Artifacts
 		public async Task StopAsync(CancellationToken cancellationToken)
 			=> await _ticker.StopAsync();
 
+		IBlobRef<DirectoryNode> Open(Artifact artifact)
+		{
+			IStorageClient client = _storageService.CreateClient(artifact.Document.NamespaceId);
+			return client.CreateBlobRef<DirectoryNode>(artifact.Document.RefName);
+		}
+
 #pragma warning disable CA1308 // Expect ansi-only keys here
 		static string NormalizeKey(string key)
 			=> key.ToLowerInvariant();
diff --git a/Engine/Source/Programs/Shared/EpicGames.Horde/Artifacts/ArtifactHttpCollection.cs b/Engine/Source/Programs/Shared/EpicGames.Horde/Artifacts/ArtifactHttpCollection.cs
--- a/Engine/Source/Programs/Shared/EpicGames.Horde/Artifacts/ArtifactHttpCollection.cs
+++ b/Engine/Source/Programs/Shared/EpicGames.Horde/Artifacts/ArtifactHttpCollection.cs
@@ -8,6 +8,7 @@ using System.Threading;
 using System.Threading.Tasks;
 using EpicGames.Horde.Commits;
 using EpicGames.Horde.Storage;
+using EpicGames.Horde.Storage.Nodes;
 using EpicGames.Horde.Streams;
 
 namespace EpicGames.Horde.Artifacts
@@ -51,6 +52,9 @@ namespace EpicGames.Horde.Artifacts
 			public RefName RefName { get; }
 			public DateTime CreatedAtUtc { get; }
 
+			public IBlobRef<DirectoryNode> Content
+				=> _collection.Open(this);
+
 			public Task DeleteAsync(CancellationToken cancellationToken)
 				=> _collection.DeleteAsync(Id, cancellationToken);
 		}
@@ -60,6 +64,12 @@ namespace EpicGames.Horde.Artifacts
 		public ArtifactHttpCollection(IHordeClient hordeClient)
 			=> _hordeClient = hordeClient;
 
+		IBlobRef<DirectoryNode> Open(Artifact artifact)
+		{
+			IStorageClient store = _hordeClient.CreateStorageClient(artifact.Id);
+			return store.CreateBlobRef<DirectoryNode>(new RefName("default"));
+		}
+
 		/// <inheritdoc/>
 		public async Task<IArtifact> AddAsync(ArtifactName name, ArtifactType type, string? description, StreamId streamId, CommitId commitId, IEnumerable<string> keys, IEnumerable<string> metadata, CancellationToken cancellationToken = default)
 		{
diff --git a/Engine/Source/Programs/Shared/EpicGames.Horde/Artifacts/IArtifact.cs b/Engine/Source/Programs/Shared/EpicGames.Horde/Artifacts/IArtifact.cs
--- a/Engine/Source/Programs/Shared/EpicGames.Horde/Artifacts/IArtifact.cs
+++ b/Engine/Source/Programs/Shared/EpicGames.Horde/Artifacts/IArtifact.cs
@@ -6,6 +6,7 @@ using System.Threading;
 using System.Threading.Tasks;
 using EpicGames.Horde.Commits;
 using EpicGames.Horde.Storage;
+using EpicGames.Horde.Storage.Nodes;
 using EpicGames.Horde.Streams;
 
 namespace EpicGames.Horde.Artifacts
@@ -70,6 +71,11 @@ namespace EpicGames.Horde.Artifacts
 		/// </summary>
 		DateTime CreatedAtUtc { get; }
 
+		/// <summary>
+		/// Handle to the artifact data
+		/// </summary>
+		IBlobRef<DirectoryNode> Content { get; }
+
 		/// <summary>
 		/// Deletes this artifact
 		/// </summary>
