Horde: Remove code to read imports from uploaded blobs. These should now be set at upload time.

diff --git a/Engine/Source/Programs/Horde/Plugins/Storage/HordeServer.Storage/Storage/StorageService.cs b/Engine/Source/Programs/Horde/Plugins/Storage/HordeServer.Storage/Storage/StorageService.cs
--- a/Engine/Source/Programs/Horde/Plugins/Storage/HordeServer.Storage/Storage/StorageService.cs
+++ b/Engine/Source/Programs/Horde/Plugins/Storage/HordeServer.Storage/Storage/StorageService.cs
@@ -688,39 +688,11 @@ namespace HordeServer.Storage
 					break;
 				}
 
-				// Find imports, and add a check record for each new blob
-				using TelemetrySpan innerSpan = _tracer.StartActiveSpan($"{nameof(StorageService)}.{nameof(TickBlobsAsync)}.Batch");
-				innerSpan.SetAttribute("First", current[0].Id.ToString());
-				innerSpan.SetAttribute("Last", current[^1].Id.ToString());
-				innerSpan.SetAttribute("Count", current.Count);
-
-				_logger.LogDebug("Ticking {NumBlobs} blobs from {FirstId} to {LastId}", current.Count, current[0].Id, current[^1].Id);
-
+				// Add a check record for each new blob
+				_logger.LogDebug("Adding {NumBlobs} blobs for GC consideration ({FirstId} to {LastId})", current.Count, current[0].Id, current[^1].Id);
 				foreach (BlobInfo blobInfo in current)
 				{
-					NamespaceInfo? namespaceInfo;
-					if (state.Namespaces.TryGetValue(blobInfo.NamespaceId, out namespaceInfo))
-					{
-						BundleStorageClient? storageClient;
-						if (!cachedClients.TryGetValue(namespaceInfo.Id, out storageClient))
-						{
-							storageClient = new BundleStorageClient(namespaceInfo.Backend, _bundleCache, null, _logger);
-							cachedClients.Add(namespaceInfo.Id, storageClient);
-						}
-
-						try
-						{
-							await TickBlobAsync(storageClient, blobInfo, cancellationToken);
-						}
-						catch (ObjectNotFoundException ex)
-						{
-							_logger.LogInformation(ex, "Unable to read references for {NamespaceId} blob {BlobId}: {Message}", blobInfo.NamespaceId, blobInfo.Id, ex.Message);
-						}
-						catch (Exception ex)
-						{
-							_logger.LogWarning(ex, "Unable to read references for {NamespaceId} blob {BlobId} (key: {ObjectKey}): {Message}", blobInfo.NamespaceId, blobInfo.Id, GetObjectKey(blobInfo.Locator), ex.Message);
-						}
-					}
+					AddGcCheckRecord(blobInfo.NamespaceId, blobInfo.Id);
 				}
 
 				// Update the last imported blob id
@@ -731,45 +703,6 @@ namespace HordeServer.Storage
 			_logger.LogInformation("Added {NumBlobs} blobs for GC (upper time: {Time})", ingestedCount, ingestTimeUtc);
 		}
 
-		async Task TickBlobAsync(BundleStorageClient storageClient, BlobInfo blobInfo, CancellationToken cancellationToken)
-		{
-			List<ObjectId> importInfoIds = new List<ObjectId>();
-
-			IEnumerable<BlobLocator> importLocators = await storageClient.ReadBundleReferencesAsync(blobInfo.Locator, cancellationToken);
-			foreach (BlobLocator importLocator in importLocators)
-			{
-				string importPath = importLocator.BaseLocator.ToString();
-
-				FilterDefinition<BlobInfo> filter = Builders<BlobInfo>.Filter.Expr(x => x.NamespaceId == blobInfo.NamespaceId && x.Path == importPath);
-				UpdateDefinition<BlobInfo> update = Builders<BlobInfo>.Update.SetOnInsert(x => x.Imports, null);
-				BlobInfo blobInfoDoc = await _blobCollection.FindOneAndUpdateAsync(filter, update, new FindOneAndUpdateOptions<BlobInfo> { IsUpsert = true, ReturnDocument = ReturnDocument.After }, cancellationToken);
-
-				importInfoIds.Add(blobInfoDoc.Id);
-			}
-
-			importInfoIds.Sort();
-
-			if (blobInfo.Imports != null && !Enumerable.SequenceEqual(importInfoIds, blobInfo.Imports))
-			{
-				List<ObjectId> missing = importInfoIds.Except(blobInfo.Imports).ToList();
-				if (missing.Count > 0)
-				{
-					_logger.LogWarning("Missing imports for blob {Locator}: {Missing}", blobInfo.Path, String.Join(", ", missing.Select(x => x.ToString())));
-				}
-
-				List<ObjectId> extra = blobInfo.Imports.Except(importInfoIds).ToList();
-				if (extra.Count > 0)
-				{
-					_logger.LogWarning("Extra imports for blob {Locator}: {Extra}", blobInfo.Path, String.Join(", ", extra.Select(x => x.ToString())));
-				}
-			}
-
-			await _blobCollection.UpdateOneAsync(x => x.Id == blobInfo.Id, Builders<BlobInfo>.Update.Set(x => x.Imports, importInfoIds).Unset(x => x.GcVersion), null, cancellationToken);
-
-			AddGcCheckRecord(blobInfo.NamespaceId, blobInfo.Id);
-			_logger.LogDebug("Added {Count} imports for {NamespaceId} blob {BlobId}", importInfoIds.Count, blobInfo.NamespaceId, blobInfo.Id);
-		}
-
 		#endregion
 
 		#region Nodes
