Horde: Support for batch updating refs and aliases for a namespace through HTTP.

#rnx

diff --git a/Engine/Source/Programs/Horde/Plugins/Storage/HordeServer.Storage/Storage/StorageController.cs b/Engine/Source/Programs/Horde/Plugins/Storage/HordeServer.Storage/Storage/StorageController.cs
--- a/Engine/Source/Programs/Horde/Plugins/Storage/HordeServer.Storage/Storage/StorageController.cs
+++ b/Engine/Source/Programs/Horde/Plugins/Storage/HordeServer.Storage/Storage/StorageController.cs
@@ -45,6 +45,46 @@ namespace HordeServer.Storage
 			return _storageConfig.Value.TryGetNamespace(namespaceId, out NamespaceConfig? namespaceConfig) && namespaceConfig.Authorize(action, User);
 		}
 
+		/// <summary>
+		/// Update aliases in the storage service
+		/// </summary>
+		/// <param name="namespaceId">Namespace to write to</param>
+		/// <param name="request">Request for the alias  to write</param>
+		/// <param name="cancellationToken">Cancellation token for the operation</param>
+		[HttpPost]
+		[Route("/api/v1/storage/{namespaceId}")]
+		public async Task<ActionResult> UpdateNamespaceAsync(NamespaceId namespaceId, [FromBody] UpdateNamespaceRequest request, CancellationToken cancellationToken)
+		{
+			IStorageBackend? backend = _storageService.TryCreateBackend(namespaceId);
+			if (backend == null)
+			{
+				return NotFound(namespaceId);
+			}
+			if (!Authorize(namespaceId, StorageAclAction.WriteRefs))
+			{
+				return Forbid(StorageAclAction.WriteRefs, namespaceId);
+			}
+
+			foreach (AddAliasRequest addAlias in request.AddAliases)
+			{
+				await backend.AddAliasAsync(addAlias.Name, addAlias.Target, addAlias.Rank, addAlias.Data ?? Array.Empty<byte>(), cancellationToken);
+			}
+			foreach (RemoveAliasRequest removeAlias in request.RemoveAliases)
+			{
+				await backend.RemoveAliasAsync(removeAlias.Name, removeAlias.Target, cancellationToken);
+			}
+			foreach (AddRefRequest writeRef in request.AddRefs)
+			{
+				await backend.WriteRefAsync(writeRef.RefName, new HashedBlobRefValue(writeRef.Hash, writeRef.Target), writeRef.Options, cancellationToken);
+			}
+			foreach (RemoveRefRequest removeRef in request.RemoveRefs)
+			{
+				await backend.DeleteRefAsync(removeRef.RefName, cancellationToken);
+			}
+
+			return Ok();
+		}
+
 		/// <summary>
 		/// Uploads data to the storage service. 
 		/// </summary>
@@ -210,6 +250,30 @@ namespace HordeServer.Storage
 			return response;
 		}
 
+		/// <summary>
+		/// Deletes a ref from the storage service.
+		/// </summary>
+		/// <param name="namespaceId">Namespace to write to</param>
+		/// <param name="refName">Name of the ref</param>
+		/// <param name="cancellationToken">Cancellation token for the operation</param>
+		[HttpDelete]
+		[Route("/api/v1/storage/{namespaceId}/refs/{*refName}")]
+		public async Task<ActionResult> WriteRefAsync(NamespaceId namespaceId, RefName refName, CancellationToken cancellationToken)
+		{
+			IStorageBackend? backend = _storageService.TryCreateBackend(namespaceId);
+			if (backend == null)
+			{
+				return NotFound(namespaceId);
+			}
+			if (!Authorize(namespaceId, StorageAclAction.WriteRefs) && !HasPathClaim(User, HordeClaimTypes.WriteNamespace, namespaceId, refName.ToString()))
+			{
+				return Forbid(StorageAclAction.WriteRefs, namespaceId);
+			}
+
+			await backend.DeleteRefAsync(refName, cancellationToken);
+			return Ok();
+		}
+
 		/// <summary>
 		/// Writes a ref to the storage service.
 		/// </summary>
diff --git a/Engine/Source/Programs/Shared/EpicGames.Horde/Storage/Backends/HttpStorageBackend.cs b/Engine/Source/Programs/Shared/EpicGames.Horde/Storage/Backends/HttpStorageBackend.cs
--- a/Engine/Source/Programs/Shared/EpicGames.Horde/Storage/Backends/HttpStorageBackend.cs
+++ b/Engine/Source/Programs/Shared/EpicGames.Horde/Storage/Backends/HttpStorageBackend.cs
@@ -233,15 +233,29 @@ namespace EpicGames.Horde.Storage.Backends
 		#region Aliases
 
 		/// <inheritdoc/>
-		public Task AddAliasAsync(string name, BlobLocator target, int rank = 0, ReadOnlyMemory<byte> data = default, CancellationToken cancellationToken = default)
+		public async Task AddAliasAsync(string name, BlobLocator target, int rank = 0, ReadOnlyMemory<byte> data = default, CancellationToken cancellationToken = default)
 		{
-			throw new NotSupportedException("Http storage client does not currently support aliases.");
+			using (HttpClient httpClient = _createClient())
+			{
+				UpdateNamespaceRequest request = new UpdateNamespaceRequest();
+				request.AddAliases.Add(new AddAliasRequest { Name = name, Target = target, Rank = rank, Data = data.ToArray() });
+
+				using HttpResponseMessage response = await httpClient.PostAsJsonAsync($"{_basePath}", request, cancellationToken: cancellationToken);
+				response.EnsureSuccessStatusCode();
+			}
 		}
 
 		/// <inheritdoc/>
-		public Task RemoveAliasAsync(string name, BlobLocator target, CancellationToken cancellationToken = default)
+		public async Task RemoveAliasAsync(string name, BlobLocator target, CancellationToken cancellationToken = default)
 		{
-			throw new NotSupportedException("Http storage client does not currently support aliases.");
+			using (HttpClient httpClient = _createClient())
+			{
+				UpdateNamespaceRequest request = new UpdateNamespaceRequest();
+				request.RemoveAliases.Add(new RemoveAliasRequest { Name = name, Target = target });
+
+				using HttpResponseMessage response = await httpClient.PostAsJsonAsync($"{_basePath}", request, cancellationToken: cancellationToken);
+				response.EnsureSuccessStatusCode();
+			}
 		}
 
 		/// <inheritdoc/>
diff --git a/Engine/Source/Programs/Shared/EpicGames.Horde/Storage/StorageMessages.cs b/Engine/Source/Programs/Shared/EpicGames.Horde/Storage/StorageMessages.cs
--- a/Engine/Source/Programs/Shared/EpicGames.Horde/Storage/StorageMessages.cs
+++ b/Engine/Source/Programs/Shared/EpicGames.Horde/Storage/StorageMessages.cs
@@ -79,6 +79,74 @@ namespace EpicGames.Horde.Storage
 		public List<FindNodeResponse> Nodes { get; set; } = new List<FindNodeResponse>();
 	}
 
+	/// <summary>
+	/// Request to batch update items in the database
+	/// </summary>
+	public class UpdateNamespaceRequest
+	{
+		/// <summary>
+		/// List of aliases to add
+		/// </summary>
+		public List<AddAliasRequest> AddAliases { get; set; } = new List<AddAliasRequest>();
+
+		/// <summary>
+		/// List of aliases to remove
+		/// </summary>
+		public List<RemoveAliasRequest> RemoveAliases { get; set; } = new List<RemoveAliasRequest>();
+
+		/// <summary>
+		/// List of refs to add
+		/// </summary>
+		public List<AddRefRequest> AddRefs { get; set; } = new List<AddRefRequest>();
+
+		/// <summary>
+		/// List of refs to remove
+		/// </summary>
+		public List<RemoveRefRequest> RemoveRefs { get; set; } = new List<RemoveRefRequest>();
+	}
+
+	/// <summary>
+	/// Request object for adding an alias
+	/// </summary>
+	public class AddAliasRequest
+	{
+		/// <summary>
+		/// Name of the alias
+		/// </summary>
+		public string Name { get; set; } = String.Empty;
+
+		/// <summary>
+		/// Rank for the new alias
+		/// </summary>
+		public int Rank { get; set; }
+
+		/// <summary>
+		/// Data to store with the ref
+		/// </summary>
+		public byte[]? Data { get; set; }
+
+		/// <summary>
+		/// Path to the target blob
+		/// </summary>
+		public BlobLocator Target { get; set; }
+	}
+
+	/// <summary>
+	/// Request object for removing an alias
+	/// </summary>
+	public class RemoveAliasRequest
+	{
+		/// <summary>
+		/// Name of the alias
+		/// </summary>
+		public string Name { get; set; } = String.Empty;
+
+		/// <summary>
+		/// Path to the target blob
+		/// </summary>
+		public BlobLocator Target { get; set; }
+	}
+
 	/// <summary>
 	/// Request object for writing a ref
 	/// </summary>
@@ -117,6 +185,28 @@ namespace EpicGames.Horde.Storage
 		public RefOptions? Options { get; set; }
 	}
 
+	/// <summary>
+	/// Request object for removing a ref
+	/// </summary>
+	public class AddRefRequest : WriteRefRequest
+	{
+		/// <summary>
+		/// Name of the ref
+		/// </summary>
+		public RefName RefName { get; set; }
+	}
+
+	/// <summary>
+	/// Request object for removing a ref
+	/// </summary>
+	public class RemoveRefRequest
+	{
+		/// <summary>
+		/// Name of the ref
+		/// </summary>
+		public RefName RefName { get; set; }
+	}
+
 	/// <summary>
 	/// Response object for reading a ref
 	/// </summary>
