Horde: Fix name of aliases over blobs collection.

diff --git a/Engine/Source/Programs/Horde/Plugins/Storage/HordeServer.Storage/Storage/StorageService.cs b/Engine/Source/Programs/Horde/Plugins/Storage/HordeServer.Storage/Storage/StorageService.cs
--- a/Engine/Source/Programs/Horde/Plugins/Storage/HordeServer.Storage/Storage/StorageService.cs
+++ b/Engine/Source/Programs/Horde/Plugins/Storage/HordeServer.Storage/Storage/StorageService.cs
@@ -411,7 +411,7 @@ namespace HordeServer.Storage
 			List<MongoIndex<BlobInfo>> blobIndexes = new List<MongoIndex<BlobInfo>>();
 			blobIndexes.Add(keys => keys.Ascending(x => x.Imports));
 			blobIndexes.Add(keys => keys.Ascending(x => x.NamespaceId).Ascending(x => x.Path), unique: true);
-			blobIndexes.Add(keys => keys.Ascending(x => x.NamespaceId).Ascending($"{nameof(BlobInfo.Aliases)}.{nameof(AliasInfo.Name)}"));
+			blobIndexes.Add(keys => keys.Ascending(x => x.NamespaceId).Ascending(x => x.Aliases![0].Name));
 			_blobCollection = mongoService.GetCollection<BlobInfo>("Storage.Blobs", blobIndexes);
 
 			List<MongoIndex<RefInfo>> refIndexes = new List<MongoIndex<RefInfo>>();
