Horde: Add agent version to compute resource class

Used for exposing version in compute resource API.

diff --git a/Engine/Source/Programs/Horde/Plugins/Compute/HordeServer.Compute/Compute/ComputeController.cs b/Engine/Source/Programs/Horde/Plugins/Compute/HordeServer.Compute/Compute/ComputeController.cs
--- a/Engine/Source/Programs/Horde/Plugins/Compute/HordeServer.Compute/Compute/ComputeController.cs
+++ b/Engine/Source/Programs/Horde/Plugins/Compute/HordeServer.Compute/Compute/ComputeController.cs
@@ -154,6 +154,7 @@ namespace HordeServer.Compute
 			response.Certificate = StringUtils.FormatHexString(computeResource.Task.Certificate.Span);
 			response.ClusterId = clusterId;
 			response.AgentId = computeResource.AgentId;
+			response.AgentVersion = computeResource.AgentVersion;
 			response.LeaseId = computeResource.LeaseId;
 			response.Properties = computeResource.Properties;
 			response.Protocol = computeResource.Task.Protocol;
diff --git a/Engine/Source/Programs/Horde/Plugins/Compute/HordeServer.Compute/Compute/ComputeService.cs b/Engine/Source/Programs/Horde/Plugins/Compute/HordeServer.Compute/Compute/ComputeService.cs
--- a/Engine/Source/Programs/Horde/Plugins/Compute/HordeServer.Compute/Compute/ComputeService.cs
+++ b/Engine/Source/Programs/Horde/Plugins/Compute/HordeServer.Compute/Compute/ComputeService.cs
@@ -721,7 +721,7 @@ namespace HordeServer.Compute
 					ports[portId] = new ComputeResourcePort(port, port);
 				}
 
-				return new ComputeResource(ConnectionMode.Direct, agentIp, null, ports, computeTask, agent.Properties, agent.Id, leaseId);
+				return new ComputeResource(ConnectionMode.Direct, agentIp, null, ports, computeTask, agent.Properties, agent.Id, agent.Version, leaseId);
 			}
 			else if (arp.ConnectionMode == ConnectionMode.Tunnel && tunnelAddress != null)
 			{
@@ -733,7 +733,7 @@ namespace HordeServer.Compute
 					ports[portId] = new ComputeResourcePort(-1, port);
 				}
 
-				return new ComputeResource(ConnectionMode.Tunnel, agentIp, tunnelAddress, ports, computeTask, agent.Properties, agent.Id, leaseId);
+				return new ComputeResource(ConnectionMode.Tunnel, agentIp, tunnelAddress, ports, computeTask, agent.Properties, agent.Id, agent.Version, leaseId);
 			}
 			else if (arp.ConnectionMode == ConnectionMode.Relay)
 			{
@@ -764,7 +764,7 @@ namespace HordeServer.Compute
 					}
 				}
 
-				return new ComputeResource(ConnectionMode.Relay, agentIp, relayIp.ToString(), ports, computeTask, agent.Properties, agent.Id, leaseId);
+				return new ComputeResource(ConnectionMode.Relay, agentIp, relayIp.ToString(), ports, computeTask, agent.Properties, agent.Id, agent.Version, leaseId);
 			}
 
 			throw new Exception("Unable to resolve a suitable connection mode for compute task");
diff --git a/Engine/Source/Programs/Horde/Plugins/Compute/HordeServer.Compute/Compute/ComputeTaskSource.cs b/Engine/Source/Programs/Horde/Plugins/Compute/HordeServer.Compute/Compute/ComputeTaskSource.cs
--- a/Engine/Source/Programs/Horde/Plugins/Compute/HordeServer.Compute/Compute/ComputeTaskSource.cs
+++ b/Engine/Source/Programs/Horde/Plugins/Compute/HordeServer.Compute/Compute/ComputeTaskSource.cs
@@ -110,6 +110,11 @@ namespace HordeServer.Compute
 		/// Agent id on the remote machine
 		/// </summary>
 		public AgentId AgentId { get; }
+		
+		/// <summary>
+		/// Agent version on the remote machine
+		/// </summary>
+		public string? AgentVersion { get; }
 
 		/// <summary>
 		/// Lease id on the remote machine
@@ -119,7 +124,16 @@ namespace HordeServer.Compute
 		/// <summary>
 		/// Constructor
 		/// </summary>
-		public ComputeResource(ConnectionMode connectionMode, IPAddress ip, string? connectionAddress, IReadOnlyDictionary<string, ComputeResourcePort> ports, ComputeTask task, IReadOnlyList<string> properties, AgentId agentId, LeaseId leaseId)
+		public ComputeResource(
+			ConnectionMode connectionMode,
+			IPAddress ip,
+			string? connectionAddress,
+			IReadOnlyDictionary<string, ComputeResourcePort> ports,
+			ComputeTask task,
+			IReadOnlyList<string> properties,
+			AgentId agentId,
+			string? agentVersion,
+			LeaseId leaseId)
 		{
 			ConnectionMode = connectionMode;
 			Ip = ip;
@@ -128,6 +142,7 @@ namespace HordeServer.Compute
 			Task = task;
 			Properties = properties;
 			AgentId = agentId;
+			AgentVersion = agentVersion;
 			LeaseId = leaseId;
 		}
 	}
