Horde: Allow graceful draining of leases when yielding to local user activity

The termination signal file used for spot interruptions is now written to let a workload know about the lease being drained.
This requires workload to scan and respect the file for this to be effective, which UBA does for example.

diff --git a/Engine/Source/Programs/Horde/HordeAgent/Leases/LeaseManager.cs b/Engine/Source/Programs/Horde/HordeAgent/Leases/LeaseManager.cs
--- a/Engine/Source/Programs/Horde/HordeAgent/Leases/LeaseManager.cs
+++ b/Engine/Source/Programs/Horde/HordeAgent/Leases/LeaseManager.cs
@@ -1,6 +1,7 @@
 // Copyright Epic Games, Inc. All Rights Reserved.
 
 using System.Diagnostics;
+using System.Text;
 using EpicGames.Core;
 using Google.Protobuf.WellKnownTypes;
 using Grpc.Core;
@@ -9,6 +10,7 @@ using HordeCommon.Rpc;
 using HordeCommon.Rpc.Messages;
 using Microsoft.Extensions.DependencyInjection;
 using Microsoft.Extensions.Logging;
+using Microsoft.Extensions.Options;
 
 namespace HordeAgent.Leases
 {
@@ -85,11 +87,12 @@ namespace HordeAgent.Leases
 		readonly ISystemMetrics _systemMetrics;
 		readonly Dictionary<string, LeaseHandlerFactory> _typeUrlToLeaseHandler;
 		readonly LeaseLoggerFactory _leaseLoggerFactory;
+		readonly IOptions<AgentSettings> _settings;
 		readonly ILogger _logger;
 
 		RpcAgentCapabilities? _capabilities;
 
-		public LeaseManager(ISession session, CapabilitiesService capabilitiesService, StatusService statusService, ISystemMetrics systemMetrics, IEnumerable<LeaseHandlerFactory> leaseHandlerFactories, LeaseLoggerFactory leaseLoggerFactory, ILogger logger)
+		public LeaseManager(ISession session, CapabilitiesService capabilitiesService, StatusService statusService, ISystemMetrics systemMetrics, IEnumerable<LeaseHandlerFactory> leaseHandlerFactories, LeaseLoggerFactory leaseLoggerFactory, IOptions<AgentSettings> settings, ILogger logger)
 		{
 			_session = session;
 			_capabilitiesService = capabilitiesService;
@@ -97,11 +100,19 @@ namespace HordeAgent.Leases
 			_systemMetrics = systemMetrics;
 			_typeUrlToLeaseHandler = leaseHandlerFactories.ToDictionary(x => x.LeaseType, x => x);
 			_leaseLoggerFactory = leaseLoggerFactory;
+			_settings = settings;
 			_logger = logger;
 		}
 
 		public LeaseManager(ISession session, IServiceProvider serviceProvider)
-			: this(session, serviceProvider.GetRequiredService<CapabilitiesService>(), serviceProvider.GetRequiredService<StatusService>(), serviceProvider.GetRequiredService<ISystemMetrics>(), serviceProvider.GetRequiredService<IEnumerable<LeaseHandlerFactory>>(), serviceProvider.GetRequiredService<LeaseLoggerFactory>(), serviceProvider.GetRequiredService<ILogger<LeaseManager>>())
+			: this(session,
+				serviceProvider.GetRequiredService<CapabilitiesService>(),
+				serviceProvider.GetRequiredService<StatusService>(),
+				serviceProvider.GetRequiredService<ISystemMetrics>(),
+				serviceProvider.GetRequiredService<IEnumerable<LeaseHandlerFactory>>(),
+				serviceProvider.GetRequiredService<LeaseLoggerFactory>(),
+				serviceProvider.GetRequiredService<IOptions<AgentSettings>>(),
+				serviceProvider.GetRequiredService<ILogger<LeaseManager>>())
 		{
 		}
 
@@ -112,6 +123,25 @@ namespace HordeAgent.Leases
 				return _activeLeases.Select(x => x.RpcLease).ToList();
 			}
 		}
+		
+		/// <summary>
+		/// Write a termination signal file to disk
+		/// Used to communicate an impending termination of workload. It does require workload to be aware of the file and also act on it.
+		/// </summary>
+		/// <param name="filePath">Path to signal file</param>
+		/// <param name="reason">Reason for termination</param>
+		/// <param name="terminateAt">Timestamp when termination takes place</param>
+		/// <param name="timeToLive">Time left to live before termination (relative to above)</param>
+		/// <param name="cancellationToken">Cancellation token</param>
+		public static Task WriteTerminationSignalFileAsync(string filePath, string reason, DateTime terminateAt, TimeSpan timeToLive, CancellationToken cancellationToken)
+		{
+			StringBuilder sb = new(100);
+			sb.Append("v1\n");
+			sb.Append($"{timeToLive.TotalMilliseconds}\n");
+			sb.Append($"{new DateTimeOffset(terminateAt).ToUnixTimeMilliseconds()}\n");
+			sb.Append($"{reason}\n");
+			return File.WriteAllTextAsync(filePath, sb.ToString(), cancellationToken);
+		}
 
 		public async Task<SessionResult> RunAsync(CancellationToken stoppingToken)
 		{
@@ -147,8 +177,20 @@ namespace HordeAgent.Leases
 			return result;
 		}
 
-		async Task DrainLeasesAsync(string reason)
+		/// <summary>
+		/// Drain and terminate any currently running leases
+		/// </summary>
+		/// <param name="reason">Human-readable reason</param>
+		/// <param name="graceful">Whether to warn workload about upcoming lease termination via signal file</param>
+		async Task DrainLeasesAsync(string reason, bool graceful = false)
 		{
+			if (graceful)
+			{
+				TimeSpan ttl = TimeSpan.FromSeconds(1);
+				await WriteTerminationSignalFileAsync(_settings.Value.GetTerminationSignalFile().FullName, reason, DateTime.UtcNow + ttl, ttl, CancellationToken.None);
+				await Task.Delay(TimeSpan.FromSeconds(6)); // Allow workload some time to act on signal file
+			}
+			
 			for (int idx = 0; idx < _activeLeases.Count; idx++)
 			{
 				LeaseHandler activeLease = _activeLeases[idx];
@@ -253,7 +295,7 @@ namespace HordeAgent.Leases
 					}
 				}
 
-				// Get the new agent status
+				// Get the new agent status to be reported back to server
 				bool busy = _statusService.IsBusy;
 				if (stopping)
 				{
@@ -364,7 +406,7 @@ namespace HordeAgent.Leases
 						if (_activeLeases.Count > 0)
 						{
 							_logger.LogInformation("Agent marked itself as busy. Draining any active leases to prevent them from using up local resources...");
-							await DrainLeasesAsync("user is active");
+							await DrainLeasesAsync("user is active", true);
 						}
 					}
 					else if (_activeLeases.Count == 0)
diff --git a/Engine/Source/Programs/Horde/HordeAgent/Services/AwsInstanceLifecycleService.cs b/Engine/Source/Programs/Horde/HordeAgent/Services/AwsInstanceLifecycleService.cs
--- a/Engine/Source/Programs/Horde/HordeAgent/Services/AwsInstanceLifecycleService.cs
+++ b/Engine/Source/Programs/Horde/HordeAgent/Services/AwsInstanceLifecycleService.cs
@@ -1,8 +1,8 @@
 // Copyright Epic Games, Inc. All Rights Reserved.
 
 using System.Net;
-using System.Text;
 using EpicGames.Core;
+using HordeAgent.Leases;
 using HordeAgent.Utility;
 using Microsoft.Extensions.Hosting;
 using Microsoft.Extensions.Logging;
@@ -231,7 +231,7 @@ class AwsInstanceLifecycleService : BackgroundService
 
 		// Create and write the termination signal file, containing the time-to-live for the EC2 instance.
 		// Workloads executed by the agent that support this protocol can pick this up and prepare/clean up prior to termination
-		await WriteTerminationSignalFileAsync(info, cancellationToken);
+		await LeaseManager.WriteTerminationSignalFileAsync(_terminationSignalFile.FullName, info.Reason, info.TerminateAt, info.TimeToLive, cancellationToken);
 	}
 
 	private Task OnTerminationAsync(Ec2TerminationInfo info, CancellationToken cancellationToken)
@@ -245,16 +245,6 @@ class AwsInstanceLifecycleService : BackgroundService
 		return Task.CompletedTask;
 	}
 
-	private Task WriteTerminationSignalFileAsync(Ec2TerminationInfo info, CancellationToken cancellationToken)
-	{
-		StringBuilder sb = new(100);
-		sb.Append("v1\n");
-		sb.Append($"{info.TimeToLive.TotalMilliseconds}\n");
-		sb.Append($"{new DateTimeOffset(info.TerminateAt).ToUnixTimeMilliseconds()}\n");
-		sb.Append($"{info.Reason}\n");
-		return File.WriteAllTextAsync(_terminationSignalFile.FullName, sb.ToString(), cancellationToken);
-	}
-
 	/// <inheritdoc/>
 	protected override async Task ExecuteAsync(CancellationToken stoppingToken)
 	{
