Horde: Pause the addition of new blobs for GC once the queue is longer than 50,000 entries.

diff --git a/Engine/Source/Programs/Horde/Plugins/Storage/HordeServer.Storage/Storage/StorageService.cs b/Engine/Source/Programs/Horde/Plugins/Storage/HordeServer.Storage/Storage/StorageService.cs
--- a/Engine/Source/Programs/Horde/Plugins/Storage/HordeServer.Storage/Storage/StorageService.cs
+++ b/Engine/Source/Programs/Horde/Plugins/Storage/HordeServer.Storage/Storage/StorageService.cs
@@ -687,6 +687,10 @@ namespace HordeServer.Storage
 				{
 					break;
 				}
+				if (await ShouldPauseBlobTickAsync(current, cancellationToken))
+				{
+					break;
+				}
 
 				// Add a check record for each new blob
 				_logger.LogDebug("Adding {NumBlobs} blobs for GC consideration ({FirstId} to {LastId})", current.Count, current[0].Id, current[^1].Id);
@@ -700,7 +704,25 @@ namespace HordeServer.Storage
 				ingestedCount += current.Count;
 			}
 
-			_logger.LogInformation("Added {NumBlobs} blobs for GC (upper time: {Time})", ingestedCount, ingestTimeUtc);
+			_logger.LogInformation("Added {NumBlobs} blobs for GC (upper time: {Time})", ingestedCount, gcState.LastImportBlobInfoId.CreationTime);
+		}
+
+		async Task<bool> ShouldPauseBlobTickAsync(IEnumerable<BlobInfo> blobs, CancellationToken cancellationToken)
+		{
+			const long MaxLength = 50000;
+
+			bool shouldPause = false;
+			foreach (NamespaceId namespaceId in blobs.Select(x => x.NamespaceId).Distinct())
+			{
+				long length = await _redisService.GetDatabase().SortedSetLengthAsync(GetGcCheckSet(namespaceId)).WaitAsync(cancellationToken);
+				if (length > MaxLength)
+				{
+					_logger.LogInformation("Length of GC queue for namespace {NamespaceId} is {Count}. Pausing addition of new items.", namespaceId, length);
+					shouldPause = true;
+				}
+			}
+
+			return shouldPause;
 		}
 
 		#endregion
