Horde: Fix linq expression for generation of alias index.

diff --git a/Engine/Source/Programs/Horde/Plugins/Storage/HordeServer.Storage/Storage/StorageService.cs b/Engine/Source/Programs/Horde/Plugins/Storage/HordeServer.Storage/Storage/StorageService.cs
--- a/Engine/Source/Programs/Horde/Plugins/Storage/HordeServer.Storage/Storage/StorageService.cs
+++ b/Engine/Source/Programs/Horde/Plugins/Storage/HordeServer.Storage/Storage/StorageService.cs
@@ -21,6 +21,7 @@ using Microsoft.Extensions.Options;
 using MongoDB.Bson;
 using MongoDB.Bson.Serialization.Attributes;
 using MongoDB.Driver;
+using MongoDB.Driver.Linq;
 using OpenTelemetry.Trace;
 using StackExchange.Redis;
 
@@ -411,7 +412,7 @@ namespace HordeServer.Storage
 			List<MongoIndex<BlobInfo>> blobIndexes = new List<MongoIndex<BlobInfo>>();
 			blobIndexes.Add(keys => keys.Ascending(x => x.Imports));
 			blobIndexes.Add(keys => keys.Ascending(x => x.NamespaceId).Ascending(x => x.Path), unique: true);
-			blobIndexes.Add(keys => keys.Ascending(x => x.NamespaceId).Ascending(x => x.Aliases![0].Name));
+			blobIndexes.Add(keys => keys.Ascending(x => x.NamespaceId).Ascending(x => x.Aliases.AllElements().Name));
 			_blobCollection = mongoService.GetCollection<BlobInfo>("Storage.Blobs", blobIndexes);
 
 			List<MongoIndex<RefInfo>> refIndexes = new List<MongoIndex<RefInfo>>();
