Horde: Fix for getting the correct field definition for aliases.

#rnx

diff --git a/Engine/Source/Programs/Horde/Plugins/Storage/HordeServer.Storage/Storage/StorageService.cs b/Engine/Source/Programs/Horde/Plugins/Storage/HordeServer.Storage/Storage/StorageService.cs
--- a/Engine/Source/Programs/Horde/Plugins/Storage/HordeServer.Storage/Storage/StorageService.cs
+++ b/Engine/Source/Programs/Horde/Plugins/Storage/HordeServer.Storage/Storage/StorageService.cs
@@ -3,6 +3,7 @@
 using System.Buffers;
 using System.ComponentModel;
 using System.Diagnostics;
+using System.Linq.Expressions;
 using EpicGames.Core;
 using EpicGames.Horde.Storage;
 using EpicGames.Horde.Storage.Bundles;
@@ -19,6 +20,7 @@ using Microsoft.Extensions.Hosting;
 using Microsoft.Extensions.Logging;
 using Microsoft.Extensions.Options;
 using MongoDB.Bson;
+using MongoDB.Bson.Serialization;
 using MongoDB.Bson.Serialization.Attributes;
 using MongoDB.Driver;
 using MongoDB.Driver.Linq;
@@ -412,7 +414,7 @@ namespace HordeServer.Storage
 			List<MongoIndex<BlobInfo>> blobIndexes = new List<MongoIndex<BlobInfo>>();
 			blobIndexes.Add(keys => keys.Ascending(x => x.Imports));
 			blobIndexes.Add(keys => keys.Ascending(x => x.NamespaceId).Ascending(x => x.Path), unique: true);
-			blobIndexes.Add(keys => keys.Ascending(x => x.NamespaceId).Ascending(x => x.Aliases.AllElements().Name));
+			blobIndexes.Add(keys => keys.Ascending(x => x.NamespaceId).Ascending($"{GetFieldName<BlobInfo>(x => x.Aliases)}.{GetFieldName<AliasInfo>(x => x.Name)}"));
 			_blobCollection = mongoService.GetCollection<BlobInfo>("Storage.Blobs", blobIndexes);
 
 			List<MongoIndex<RefInfo>> refIndexes = new List<MongoIndex<RefInfo>>();
@@ -428,6 +430,13 @@ namespace HordeServer.Storage
 			_gcTicker = clock.AddTicker("Storage:GC", TimeSpan.FromMinutes(5.0), TickGcAsync, logger);
 		}
 
+		static string GetFieldName<TClass>(Expression<Func<TClass, object?>> expr)
+		{
+			ExpressionFieldDefinition<TClass> collection = new ExpressionFieldDefinition<TClass>(expr);
+			RenderedFieldDefinition fd = collection.Render(BsonSerializer.LookupSerializer<TClass>(), BsonSerializer.SerializerRegistry);
+			return fd.FieldName;
+		}
+
 		/// <inheritdoc/>
 		public async ValueTask DisposeAsync()
 		{
