Horde: Use the uppercase path for any requested symbol from the symbol store.

#rnx

diff --git a/Engine/Source/Programs/Horde/Plugins/Symbols/HordeServer.Symbols/Symbols/SymbolsController.cs b/Engine/Source/Programs/Horde/Plugins/Symbols/HordeServer.Symbols/Symbols/SymbolsController.cs
--- a/Engine/Source/Programs/Horde/Plugins/Symbols/HordeServer.Symbols/Symbols/SymbolsController.cs
+++ b/Engine/Source/Programs/Horde/Plugins/Symbols/HordeServer.Symbols/Symbols/SymbolsController.cs
@@ -51,7 +51,7 @@ namespace HordeServer.Tools
 			}
 
 			IStorageClient storageClient = _storageService.CreateClient(symbolStoreConfig.NamespaceId);
-			BlobAlias? alias = await storageClient.FindAliasAsync($"sym:{path}", cancellationToken);
+			BlobAlias? alias = await storageClient.FindAliasAsync($"sym:{path.ToUpperInvariant()}", cancellationToken);
 			if (alias == null)
 			{
 				return NotFound();
