Horde: Enable the symbols plugin by default.

#rnx

diff --git a/Engine/Source/Programs/Horde/HordeServer/HordeServer.csproj b/Engine/Source/Programs/Horde/HordeServer/HordeServer.csproj
--- a/Engine/Source/Programs/Horde/HordeServer/HordeServer.csproj
+++ b/Engine/Source/Programs/Horde/HordeServer/HordeServer.csproj
@@ -115,6 +115,7 @@
     <ProjectReference Include="..\Plugins\Ddc\HordeServer.Ddc\HordeServer.Ddc.csproj" />
     <ProjectReference Include="..\Plugins\Secrets\HordeServer.Secrets\HordeServer.Secrets.csproj" />
     <ProjectReference Include="..\Plugins\Storage\HordeServer.Storage\HordeServer.Storage.csproj" />
+    <ProjectReference Include="..\Plugins\Symbols\HordeServer.Symbols\HordeServer.Symbols.csproj" />
     <ProjectReference Include="..\Plugins\Tools\HordeServer.Tools\HordeServer.Tools.csproj" />
   </ItemGroup>
 
diff --git a/Engine/Source/Programs/Horde/HordeServer/ServerApp.cs b/Engine/Source/Programs/Horde/HordeServer/ServerApp.cs
--- a/Engine/Source/Programs/Horde/HordeServer/ServerApp.cs
+++ b/Engine/Source/Programs/Horde/HordeServer/ServerApp.cs
@@ -210,6 +210,7 @@ namespace HordeServer
 			pluginCollection.Add<DdcPlugin>();
 			pluginCollection.Add<SecretsPlugin>();
 			pluginCollection.Add<StoragePlugin>();
+			pluginCollection.Add<SymbolsPlugin>();
 			pluginCollection.Add<ToolsPlugin>();
 			return pluginCollection;
 		}
