Horde: Fix installed server not correctly identifying code changes correctly. Server was incorrectly setting a flag indicating that all files for a change had been enumerated, preventing it scanning the entire set.

diff --git a/Engine/Source/Programs/Horde/Plugins/Build/HordeServer.Build/Commits/ICommit.cs b/Engine/Source/Programs/Horde/Plugins/Build/HordeServer.Build/Commits/ICommit.cs
--- a/Engine/Source/Programs/Horde/Plugins/Build/HordeServer.Build/Commits/ICommit.cs
+++ b/Engine/Source/Programs/Horde/Plugins/Build/HordeServer.Build/Commits/ICommit.cs
@@ -70,9 +70,26 @@ namespace HordeServer.Commits
 		/// <summary>
 		/// Gets the files for this change, relative to the root of the stream
 		/// </summary>
+		/// <param name="minFiles">Minimum number of files to return. The response will include at least this number of files, unless the commit has fewer files.</param>
 		/// <param name="maxFiles">Maximum number of files to return. Querying large number of files may cause performance issues with merge commits.</param>
 		/// <param name="cancellationToken">Cancellation token for the operation</param>
 		/// <returns>List of files modified by this commit</returns>
-		ValueTask<IReadOnlyList<string>> GetFilesAsync(int maxFiles, CancellationToken cancellationToken);
+		ValueTask<IReadOnlyList<string>> GetFilesAsync(int? minFiles, int? maxFiles, CancellationToken cancellationToken);
+	}
+
+	/// <summary>
+	/// Extension methods for operating on commits
+	/// </summary>
+	public static class CommitExtensions
+	{
+		/// <summary>
+		/// Gets the files for this change, relative to the root of the stream
+		/// </summary>
+		/// <param name="commit">Commit to operate on</param>
+		/// <param name="files">Number of files to return</param>
+		/// <param name="cancellationToken">Cancellation token for the operation</param>
+		/// <returns>List of files modified by this commit</returns>
+		public static ValueTask<IReadOnlyList<string>> GetFilesAsync(this ICommit commit, int files, CancellationToken cancellationToken)
+			=> commit.GetFilesAsync(files, files, cancellationToken);
 	}
 }
diff --git a/Engine/Source/Programs/Horde/Plugins/Build/HordeServer.Build/Perforce/PerforceService.cs b/Engine/Source/Programs/Horde/Plugins/Build/HordeServer.Build/Perforce/PerforceService.cs
--- a/Engine/Source/Programs/Horde/Plugins/Build/HordeServer.Build/Perforce/PerforceService.cs
+++ b/Engine/Source/Programs/Horde/Plugins/Build/HordeServer.Build/Perforce/PerforceService.cs
@@ -175,10 +175,10 @@ namespace HordeServer.Perforce
 				DateUtc = dateUtc;
 			}
 
-			public void SetFiles(IReadOnlyList<string> files, int? maxFiles)
+			public void SetFiles(IReadOnlyList<string> files, bool hasAllFiles)
 			{
 				_files = files;
-				_hasAllFiles = (maxFiles == null || files.Count < maxFiles.Value);
+				_hasAllFiles = hasAllFiles;
 			}
 
 			public async ValueTask<IReadOnlyList<CommitTag>> GetTagsAsync(CancellationToken cancellationToken)
@@ -200,7 +200,7 @@ namespace HordeServer.Perforce
 				for (; ; )
 				{
 					// Query the files up to the current maximum
-					IReadOnlyList<string> files = await GetFilesAsync(maxFiles, cancellationToken);
+					IReadOnlyList<string> files = await GetFilesAsync(maxFiles, null, cancellationToken);
 					if (filter.ApplyTo(files).Any())
 					{
 						return true;
@@ -217,18 +217,25 @@ namespace HordeServer.Perforce
 				}
 			}
 
-			public async ValueTask<IReadOnlyList<string>> GetFilesAsync(int maxFiles, CancellationToken cancellationToken)
+			public async ValueTask<IReadOnlyList<string>> GetFilesAsync(int? minFiles, int? maxFiles, CancellationToken cancellationToken)
 			{
-				if (maxFiles > _files.Count && !_hasAllFiles)
+				if (minFiles.HasValue && minFiles.Value > _files.Count && !_hasAllFiles)
 				{
 					using (IPooledPerforceConnection perforce = await _owner.ConnectAsync(_streamConfig.ClusterName, null, cancellationToken))
 					{
-						DescribeRecord describeRecord = await perforce.DescribeAsync(DescribeOptions.None, maxFiles, Number, cancellationToken);
+						DescribeRecord describeRecord = await perforce.DescribeAsync(DescribeOptions.None, minFiles.Value, Number, cancellationToken);
 						_files = await perforce.GetStreamFilesAsync(_streamConfig, describeRecord, cancellationToken);
-						_hasAllFiles = _files.Count < maxFiles;
+						_hasAllFiles = _files.Count < minFiles.Value;
 					}
 				}
-				return _files;
+
+				IReadOnlyList<string> files = _files;
+				if (maxFiles.HasValue && files.Count > maxFiles)
+				{
+					files = files.Take(maxFiles.Value).ToArray();
+				}
+
+				return files;
 			}
 		}
 
@@ -696,7 +703,7 @@ namespace HordeServer.Perforce
 			return new Commit(this, streamConfig, number, originalChange, authorUser.Id, ownerUser.Id, description, basePath ?? String.Empty, dateUtc);
 		}
 
-		protected async ValueTask<ICommit> CreateCommitAsync(IPooledPerforceConnection perforce, StreamConfig streamConfig, DescribeRecord describeRecord, int maxFiles, InfoRecord serverInfo, CancellationToken cancellationToken)
+		protected async ValueTask<ICommit> CreateCommitAsync(IPooledPerforceConnection perforce, StreamConfig streamConfig, DescribeRecord describeRecord, bool hasAllFiles, InfoRecord serverInfo, CancellationToken cancellationToken)
 		{
 			DateTime timeUtc = new DateTime(describeRecord.Time.Ticks - serverInfo.TimeZoneOffsetSecs * TimeSpan.TicksPerSecond, DateTimeKind.Utc);
 
@@ -707,7 +714,7 @@ namespace HordeServer.Perforce
 			{
 				throw new PerforceException($"Changelist {commit.Number} does not contain any files in {streamConfig.Id}");
 			}
-			commit.SetFiles(files, maxFiles);
+			commit.SetFiles(files, hasAllFiles);
 
 			return commit;
 		}
@@ -778,7 +785,8 @@ namespace HordeServer.Perforce
 
 				DateTime timeUtc = new DateTime(describeRecord.Time.Ticks - info.TimeZoneOffsetSecs * TimeSpan.TicksPerSecond, DateTimeKind.Utc);
 
-				ICommit commit = await CreateCommitAsync(perforce, streamConfig, describeRecord, MaxFiles, info, cancellationToken);
+				bool hasAllFiles = describeRecord.Files.Count < MaxFiles;
+				ICommit commit = await CreateCommitAsync(perforce, streamConfig, describeRecord, hasAllFiles, info, cancellationToken);
 				return commit;
 			}
 		}
@@ -1256,7 +1264,8 @@ namespace HordeServer.Perforce
 							List<DescribeRecord> describeRecords = await perforce.DescribeAsync(DescribeOptions.None, MaxFiles, changesRecords.Select(x => x.Number).ToArray(), cancellationToken);
 							foreach (DescribeRecord describeRecord in describeRecords)
 							{
-								ICommit commit = await PerforceService.CreateCommitAsync(perforce, StreamConfig, describeRecord, MaxFiles, info, cancellationToken);
+								bool hasAllFiles = describeRecord.Files.Count < MaxFiles;
+								ICommit commit = await PerforceService.CreateCommitAsync(perforce, StreamConfig, describeRecord, hasAllFiles, info, cancellationToken);
 
 								IReadOnlyList<CommitTag> commitTags = await commit.GetTagsAsync(cancellationToken);
 								if (commitTags.Intersect(tags).Any())
diff --git a/Engine/Source/Programs/Horde/Plugins/Build/HordeServer.Build/Perforce/PerforceServiceCache.cs b/Engine/Source/Programs/Horde/Plugins/Build/HordeServer.Build/Perforce/PerforceServiceCache.cs
--- a/Engine/Source/Programs/Horde/Plugins/Build/HordeServer.Build/Perforce/PerforceServiceCache.cs
+++ b/Engine/Source/Programs/Horde/Plugins/Build/HordeServer.Build/Perforce/PerforceServiceCache.cs
@@ -210,14 +210,14 @@ namespace HordeServer.Perforce
 				return await other.MatchesFilterAsync(filter, cancellationToken);
 			}
 
-			public async ValueTask<IReadOnlyList<string>> GetFilesAsync(int maxFiles, CancellationToken cancellationToken)
+			public async ValueTask<IReadOnlyList<string>> GetFilesAsync(int? minFiles, int? maxFiles, CancellationToken cancellationToken)
 			{
 				ICommit? other = await _owner!.GetChangeDetailsAsync(_streamConfig, Number, cancellationToken);
 				if (other == null)
 				{
 					return Array.Empty<string>();
 				}
-				return await other.GetFilesAsync(maxFiles, cancellationToken);
+				return await other.GetFilesAsync(minFiles, maxFiles, cancellationToken);
 			}
 		}
 
@@ -533,19 +533,14 @@ namespace HordeServer.Perforce
 
 						if (files.Count > 0)
 						{
-							ICommit commit = await CreateCommitAsync(perforce, streamInfo.StreamConfig, describeRecord, MaxFiles, info, cancellationToken);
+							bool hasAllFiles = describeRecord.Files.Count < MaxFiles;
+
+							ICommit commit = await CreateCommitAsync(perforce, streamInfo.StreamConfig, describeRecord, hasAllFiles, info, cancellationToken);
 							_logger.LogInformation("Replicating {StreamId} commit {Change}", streamInfo.StreamConfig.Id, describeRecord.Number);
 
-							List<CommitTag> commitTags = new List<CommitTag>();
-							foreach (CommitTagInfo commitTagInfo in streamInfo.CommitTags)
-							{
-								if (commitTagInfo.Filter.ApplyTo(files.Select(x => "/" + x)).Any())
-								{
-									commitTags.Add(commitTagInfo.Name);
-								}
-							}
+							IReadOnlyList<CommitTag> commitTags = await commit.GetTagsAsync(cancellationToken);
 
-							CachedCommitDoc commitDoc = new CachedCommitDoc(commit, commitTags);
+							CachedCommitDoc commitDoc = new CachedCommitDoc(commit, commitTags.ToList());
 							await AddCachedCommitAsync(commitDoc, cancellationToken);
 
 							await _redisService.PublishAsync(s_commitUpdateChannel, streamInfo.StreamConfig.Id, CommandFlags.FireAndForget);
