Horde: Test fix.

#rnx

diff --git a/Engine/Source/Programs/Horde/Plugins/Build/HordeServer.Build.Tests/Stubs/Services/PerforceServiceStub.cs b/Engine/Source/Programs/Horde/Plugins/Build/HordeServer.Build.Tests/Stubs/Services/PerforceServiceStub.cs
--- a/Engine/Source/Programs/Horde/Plugins/Build/HordeServer.Build.Tests/Stubs/Services/PerforceServiceStub.cs
+++ b/Engine/Source/Programs/Horde/Plugins/Build/HordeServer.Build.Tests/Stubs/Services/PerforceServiceStub.cs
@@ -82,7 +82,7 @@ namespace HordeServer.Tests.Stubs.Services
 				return new ValueTask<bool>(filter.ApplyTo(Files).Any());
 			}
 
-			public ValueTask<IReadOnlyList<string>> GetFilesAsync(int maxFiles, CancellationToken cancellationToken)
+			public ValueTask<IReadOnlyList<string>> GetFilesAsync(int? minFiles, int? maxFiles, CancellationToken cancellationToken)
 			{
 				return new ValueTask<IReadOnlyList<string>>(Files);
 			}
