Horde: Symbol store fixes.

#rnx

diff --git a/Engine/Source/Programs/Horde/HordeServer.Shared/Utilities/HordeControllerBase.cs b/Engine/Source/Programs/Horde/HordeServer.Shared/Utilities/HordeControllerBase.cs
--- a/Engine/Source/Programs/Horde/HordeServer.Shared/Utilities/HordeControllerBase.cs
+++ b/Engine/Source/Programs/Horde/HordeServer.Shared/Utilities/HordeControllerBase.cs
@@ -1,5 +1,6 @@
 // Copyright Epic Games, Inc. All Rights Reserved.
 
+using System.Linq.Expressions;
 using EpicGames.Core;
 using EpicGames.Horde.Accounts;
 using EpicGames.Horde.Acls;
@@ -14,6 +15,7 @@ using EpicGames.Horde.Projects;
 using EpicGames.Horde.Secrets;
 using EpicGames.Horde.Storage;
 using EpicGames.Horde.Streams;
+using EpicGames.Horde.Symbols;
 using EpicGames.Horde.Telemetry;
 using EpicGames.Horde.Tools;
 using Microsoft.AspNetCore.Http;
@@ -131,6 +133,15 @@ namespace HordeServer.Utilities
 			return Forbid(action, "stream {StreamId}", streamId);
 		}
 
+		/// <summary>
+		/// Returns a 403 (forbidden) response with the given action and object
+		/// </summary>
+		[NonAction]
+		protected ActionResult Forbid(AclAction action, SymbolStoreId storeId)
+		{
+			return Forbid(action, "symbol store {StoreId}", storeId);
+		}
+
 		/// <summary>
 		/// Returns a 403 (forbidden) response with the given action and object
 		/// </summary>
diff --git a/Engine/Source/Programs/Horde/Plugins/Symbols/HordeServer.Symbols/Symbols/SymbolStoreConfig.cs b/Engine/Source/Programs/Horde/Plugins/Symbols/HordeServer.Symbols/Symbols/SymbolStoreConfig.cs
--- a/Engine/Source/Programs/Horde/Plugins/Symbols/HordeServer.Symbols/Symbols/SymbolStoreConfig.cs
+++ b/Engine/Source/Programs/Horde/Plugins/Symbols/HordeServer.Symbols/Symbols/SymbolStoreConfig.cs
@@ -3,6 +3,7 @@
 using System.Security.Claims;
 using EpicGames.Horde.Acls;
 using EpicGames.Horde.Storage;
+using EpicGames.Horde.Symbols;
 using HordeServer.Acls;
 
 namespace HordeServer.Symbols
diff --git a/Engine/Source/Programs/Horde/Plugins/Symbols/HordeServer.Symbols/Symbols/SymbolStoreId.cs b/Engine/Source/Programs/Horde/Plugins/Symbols/HordeServer.Symbols/Symbols/SymbolStoreId.cs
deleted file mode 100644
diff --git a/Engine/Source/Programs/Horde/Plugins/Symbols/HordeServer.Symbols/Symbols/SymbolsController.cs b/Engine/Source/Programs/Horde/Plugins/Symbols/HordeServer.Symbols/Symbols/SymbolsController.cs
--- a/Engine/Source/Programs/Horde/Plugins/Symbols/HordeServer.Symbols/Symbols/SymbolsController.cs
+++ b/Engine/Source/Programs/Horde/Plugins/Symbols/HordeServer.Symbols/Symbols/SymbolsController.cs
@@ -3,6 +3,7 @@
 using System.Net;
 using EpicGames.Horde.Storage;
 using EpicGames.Horde.Storage.Nodes;
+using EpicGames.Horde.Symbols;
 using HordeServer.Storage;
 using HordeServer.Symbols;
 using HordeServer.Utilities;
@@ -12,7 +13,7 @@ using Microsoft.Extensions.Options;
 namespace HordeServer.Tools
 {
 	/// <summary>
-	/// Controller for the /api/v1/tools endpoint
+	/// Controller for the /api/v1/symbols endpoint
 	/// </summary>
 	[ApiController]
 	[TryAuthorize]
@@ -37,7 +38,7 @@ namespace HordeServer.Tools
 		/// <param name="path">Path for the file to retrieve</param>
 		/// <param name="cancellationToken">Cancellation token for the operation</param>
 		[HttpGet]
-		[Route("/api/v1/storage/{storeId}/{*path}")]
+		[Route("/api/v1/symbols/{storeId}/{*path}")]
 		public async Task<ActionResult> ReadFileAsync(SymbolStoreId storeId, string path, CancellationToken cancellationToken)
 		{
 			SymbolStoreConfig? symbolStoreConfig;
@@ -47,7 +48,7 @@ namespace HordeServer.Tools
 			}
 			if (!symbolStoreConfig.Authorize(SymbolStoreAclAction.ReadSymbols, User))
 			{
-				return Forbid(SymbolStoreAclAction.ReadSymbols);//, storeId);
+				return Forbid(SymbolStoreAclAction.ReadSymbols, storeId);
 			}
 
 			IStorageClient storageClient = _storageService.CreateClient(symbolStoreConfig.NamespaceId);
diff --git a/Engine/Source/Programs/Horde/Plugins/Symbols/HordeServer.Symbols/SymbolsConfig.cs b/Engine/Source/Programs/Horde/Plugins/Symbols/HordeServer.Symbols/SymbolsConfig.cs
--- a/Engine/Source/Programs/Horde/Plugins/Symbols/HordeServer.Symbols/SymbolsConfig.cs
+++ b/Engine/Source/Programs/Horde/Plugins/Symbols/HordeServer.Symbols/SymbolsConfig.cs
@@ -1,6 +1,7 @@
 // Copyright Epic Games, Inc. All Rights Reserved.
 
 using System.Diagnostics.CodeAnalysis;
+using EpicGames.Horde.Symbols;
 using HordeServer.Acls;
 using HordeServer.Plugins;
 
diff --git a/Engine/Source/Programs/Shared/EpicGames.Horde/Symbols/SymbolStoreId.cs b/Engine/Source/Programs/Shared/EpicGames.Horde/Symbols/SymbolStoreId.cs
new file mode 100644
--- /dev/null
+++ b/Engine/Source/Programs/Shared/EpicGames.Horde/Symbols/SymbolStoreId.cs
@@ -0,0 +1,45 @@
+// Copyright Epic Games, Inc. All Rights Reserved.
+
+using System.ComponentModel;
+using EpicGames.Core;
+using EpicGames.Horde;
+using EpicGames.Serialization;
+
+namespace EpicGames.Horde.Symbols
+{
+	/// <summary>
+	/// Identifier for a symbol store
+	/// </summary>
+	/// <param name="Id">Id to construct from</param>
+	[JsonSchemaString]
+	[TypeConverter(typeof(StringIdTypeConverter<SymbolStoreId, SymbolStoreIdConverter>))]
+	[StringIdConverter(typeof(SymbolStoreIdConverter))]
+	[CbConverter(typeof(StringIdCbConverter<SymbolStoreId, SymbolStoreIdConverter>))]
+	public record struct SymbolStoreId(StringId Id)
+	{
+		/// <summary>
+		/// Constructor
+		/// </summary>
+		public SymbolStoreId(string id) : this(new StringId(id))
+		{
+		}
+
+		/// <inheritdoc cref="StringId.IsEmpty"/>
+		public bool IsEmpty => Id.IsEmpty;
+
+		/// <inheritdoc/>
+		public override string ToString() => Id.ToString();
+	}
+
+	/// <summary>
+	/// Converter to and from <see cref="StringId"/> instances.
+	/// </summary>
+	class SymbolStoreIdConverter : StringIdConverter<SymbolStoreId>
+	{
+		/// <inheritdoc/>
+		public override SymbolStoreId FromStringId(StringId id) => new SymbolStoreId(id);
+
+		/// <inheritdoc/>
+		public override StringId ToStringId(SymbolStoreId value) => value.Id;
+	}
+}
