Horde: Fix lease outcome not being set correctly in MongoDB.

#rnx

diff --git a/Engine/Source/Programs/Horde/Plugins/Compute/HordeServer.Compute/Agents/AgentCollection.cs b/Engine/Source/Programs/Horde/Plugins/Compute/HordeServer.Compute/Agents/AgentCollection.cs
--- a/Engine/Source/Programs/Horde/Plugins/Compute/HordeServer.Compute/Agents/AgentCollection.cs
+++ b/Engine/Source/Programs/Horde/Plugins/Compute/HordeServer.Compute/Agents/AgentCollection.cs
@@ -956,7 +956,7 @@ namespace HordeServer.Agents
 					RpcLease? rpcLease = options.Leases?.FirstOrDefault(x => x.Id == oldLease.Id);
 					if (rpcLease != null && rpcLease.State == RpcLeaseState.Completed)
 					{
-						await _leaseCollection.TrySetOutcomeAsync(oldLease.Id, _clock.UtcNow, LeaseOutcome.Success, rpcLease.Output.ToArray(), CancellationToken.None);
+						await _leaseCollection.TrySetOutcomeAsync(oldLease.Id, _clock.UtcNow, (LeaseOutcome)rpcLease.Outcome, rpcLease.Output.ToArray(), CancellationToken.None);
 					}
 					else
 					{
