Horde: Remove unused using directive.

#rnx

diff --git a/Engine/Source/Programs/Horde/HordeServer.Shared/Utilities/HordeControllerBase.cs b/Engine/Source/Programs/Horde/HordeServer.Shared/Utilities/HordeControllerBase.cs
--- a/Engine/Source/Programs/Horde/HordeServer.Shared/Utilities/HordeControllerBase.cs
+++ b/Engine/Source/Programs/Horde/HordeServer.Shared/Utilities/HordeControllerBase.cs
@@ -1,6 +1,5 @@
 // Copyright Epic Games, Inc. All Rights Reserved.
 
-using System.Linq.Expressions;
 using EpicGames.Core;
 using EpicGames.Horde.Accounts;
 using EpicGames.Horde.Acls;
