Horde: Fix alias filter definition not being compiled correctly by Mongo driver.

#rnx

diff --git a/Engine/Source/Programs/Horde/Plugins/Storage/HordeServer.Storage/Storage/StorageService.cs b/Engine/Source/Programs/Horde/Plugins/Storage/HordeServer.Storage/Storage/StorageService.cs
--- a/Engine/Source/Programs/Horde/Plugins/Storage/HordeServer.Storage/Storage/StorageService.cs
+++ b/Engine/Source/Programs/Horde/Plugins/Storage/HordeServer.Storage/Storage/StorageService.cs
@@ -396,6 +396,9 @@ namespace HordeServer.Storage
 
 		internal IMongoCollection<BlobInfo> BlobCollection => _blobCollection;
 
+		static FieldDefinition<BlobInfo, string> s_blobAliasField
+			= new StringFieldDefinition<BlobInfo, string>($"{GetFieldName<BlobInfo>(x => x.Aliases)}.{GetFieldName<AliasInfo>(x => x.Name)}");
+
 		/// <summary>
 		/// Constructor
 		/// </summary>
@@ -414,7 +417,7 @@ namespace HordeServer.Storage
 			List<MongoIndex<BlobInfo>> blobIndexes = new List<MongoIndex<BlobInfo>>();
 			blobIndexes.Add(keys => keys.Ascending(x => x.Imports));
 			blobIndexes.Add(keys => keys.Ascending(x => x.NamespaceId).Ascending(x => x.Path), unique: true);
-			blobIndexes.Add(keys => keys.Ascending(x => x.NamespaceId).Ascending($"{GetFieldName<BlobInfo>(x => x.Aliases)}.{GetFieldName<AliasInfo>(x => x.Name)}"));
+			blobIndexes.Add(keys => keys.Ascending(x => x.NamespaceId).Ascending(s_blobAliasField));
 			_blobCollection = mongoService.GetCollection<BlobInfo>("Storage.Blobs", blobIndexes);
 
 			List<MongoIndex<RefInfo>> refIndexes = new List<MongoIndex<RefInfo>>();
@@ -796,8 +799,11 @@ namespace HordeServer.Storage
 		/// <returns>Sequence of thandles</returns>
 		async Task<List<(BlobLocator, AliasInfo)>> FindAliasesAsync(NamespaceId namespaceId, string name, CancellationToken cancellationToken = default)
 		{
+			FilterDefinition<BlobInfo> filter =
+				Builders<BlobInfo>.Filter.Eq(x => x.NamespaceId, namespaceId) & Builders<BlobInfo>.Filter.Eq(s_blobAliasField, name);
+
 			List<(BlobLocator, AliasInfo)> results = new List<(BlobLocator, AliasInfo)>();
-			await foreach (BlobInfo blobInfo in _blobCollection.Find(x => x.NamespaceId == namespaceId && x.Aliases!.Any(y => y.Name == name)).ToAsyncEnumerable(cancellationToken))
+			await foreach (BlobInfo blobInfo in _blobCollection.Find(filter).ToAsyncEnumerable(cancellationToken))
 			{
 				if (blobInfo.Aliases != null)
 				{
