Horde: Fix static analysis warning.

#rnx

diff --git a/Engine/Source/Programs/Horde/Plugins/Storage/HordeServer.Storage/Storage/StorageService.cs b/Engine/Source/Programs/Horde/Plugins/Storage/HordeServer.Storage/Storage/StorageService.cs
--- a/Engine/Source/Programs/Horde/Plugins/Storage/HordeServer.Storage/Storage/StorageService.cs
+++ b/Engine/Source/Programs/Horde/Plugins/Storage/HordeServer.Storage/Storage/StorageService.cs
@@ -396,7 +396,7 @@ namespace HordeServer.Storage
 
 		internal IMongoCollection<BlobInfo> BlobCollection => _blobCollection;
 
-		static FieldDefinition<BlobInfo, string> s_blobAliasField
+		static readonly FieldDefinition<BlobInfo, string> s_blobAliasField
 			= new StringFieldDefinition<BlobInfo, string>($"{GetFieldName<BlobInfo>(x => x.Aliases)}.{GetFieldName<AliasInfo>(x => x.Name)}");
 
 		/// <summary>
