Horde: Lock dashboard license checker version
#rnx
#skipci

diff --git a/Engine/Source/Programs/Horde/HordeDashboard/package.json b/Engine/Source/Programs/Horde/HordeDashboard/package.json
--- a/Engine/Source/Programs/Horde/HordeDashboard/package.json
+++ b/Engine/Source/Programs/Horde/HordeDashboard/package.json
@@ -44,7 +44,7 @@
     "eslint": "^8.57.0",
     "eslint-plugin-react-hooks": "^4.6.0",
     "eslint-plugin-react-refresh": "^0.4.6",
-    "license-checker-rseidelsohn": "^4.3.0",
+    "license-checker-rseidelsohn": "4.3.0",
     "rimraf": "^3.0.0",
     "typescript": "^5.4.5",
     "vite": "^5.2.0",
