Horde: Re-enable dedupe for the tools bundled with Horde. The permissions affecting access to this data are now waived for bundled tools.

diff --git a/Engine/Source/Programs/Horde/Plugins/Tools/HordeServer.Tools/Commands/Generate/ToolDataCommand.cs b/Engine/Source/Programs/Horde/Plugins/Tools/HordeServer.Tools/Commands/Generate/ToolDataCommand.cs
--- a/Engine/Source/Programs/Horde/Plugins/Tools/HordeServer.Tools/Commands/Generate/ToolDataCommand.cs
+++ b/Engine/Source/Programs/Horde/Plugins/Tools/HordeServer.Tools/Commands/Generate/ToolDataCommand.cs
@@ -4,6 +4,7 @@ using System.ComponentModel;
 using System.Text.Json.Nodes;
 using EpicGames.Core;
 using EpicGames.Horde.Storage;
+using EpicGames.Horde.Storage.Backends;
 using EpicGames.Horde.Storage.Bundles;
 using EpicGames.Horde.Storage.Nodes;
 using HordeServer.Tools;
@@ -75,14 +76,20 @@ namespace HordeServer.Commands.Generate
 				IStorageClient client = BundleStorageClient.CreateFromDirectory(bundleDir, bundleCache, memoryMappedFileCache, logger);
 
 				IHashedBlobRef<DirectoryNode> dirNodeRef;
-				await using (IBlobWriter writer = client.CreateBlobWriter(refName))
+				await using (DedupeBlobWriter writer = client.CreateDedupeBlobWriter(refName))
 				{
+					logger.LogInformation("Populating cache with existing refs...");
+					await PopulateCacheAsync(client, writer, bundleDir, CancellationToken.None);
+
 					logger.LogInformation("");
 					logger.LogInformation("Writing tool data for {ToolId}", refName);
 					DirectoryNode dirNode = new DirectoryNode();
 					await dirNode.AddFilesAsync(InputDir.ToDirectoryInfo(), writer);
 					dirNodeRef = await writer.WriteBlobAsync(dirNode);
 
+					logger.LogInformation("");
+					writer.GetStats().Print(logger);
+
 					logger.LogInformation("");
 				}
 				await client.WriteRefAsync(refName, dirNodeRef);
@@ -130,7 +137,7 @@ namespace HordeServer.Commands.Generate
 
 			return 0;
 		}
-		/*
+
 		static async Task PopulateCacheAsync(IStorageClient client, DedupeBlobWriter writer, DirectoryReference searchDir, CancellationToken cancellationToken)
 		{
 			foreach (RefName refName in FileStorageBackend.EnumerateRefs(searchDir))
@@ -147,6 +154,5 @@ namespace HordeServer.Commands.Generate
 				}
 			}
 		}
-		*/
 	}
 }
diff --git a/Engine/Source/Programs/Horde/Plugins/Tools/HordeServer.Tools/Tools/ToolsController.cs b/Engine/Source/Programs/Horde/Plugins/Tools/HordeServer.Tools/Tools/ToolsController.cs
--- a/Engine/Source/Programs/Horde/Plugins/Tools/HordeServer.Tools/Tools/ToolsController.cs
+++ b/Engine/Source/Programs/Horde/Plugins/Tools/HordeServer.Tools/Tools/ToolsController.cs
@@ -374,7 +374,7 @@ namespace HordeServer.Tools
 				return Forbid(ToolAclAction.DownloadTool, id);
 			}
 
-			if (!locator.WithinFolder(tool.Id.Id.Text))
+			if (!locator.WithinFolder(tool.Id.Id.Text) && !tool.Bundled)
 			{
 				return BadRequest("Invalid blob id for tool");
 			}
