Horde: Allow specifying a list of properties required for agents to execute compute leases.

diff --git a/Engine/Source/Programs/Horde/Plugins/Compute/HordeServer.Compute/Agents/IAgent.cs b/Engine/Source/Programs/Horde/Plugins/Compute/HordeServer.Compute/Agents/IAgent.cs
--- a/Engine/Source/Programs/Horde/Plugins/Compute/HordeServer.Compute/Agents/IAgent.cs
+++ b/Engine/Source/Programs/Horde/Plugins/Compute/HordeServer.Compute/Agents/IAgent.cs
@@ -523,7 +523,7 @@ namespace HordeServer.Agents
 				combinedConditions.Add(requirements.Condition);
 			}
 			combinedConditions.AddRange(conditions);
-			return MeetsRequirements(agent, poolId, combinedConditions, requirements.Resources, requirements.Exclusive, assignedResources);
+			return MeetsRequirements(agent, poolId, combinedConditions, requirements.Properties, requirements.Resources, requirements.Exclusive, assignedResources);
 		}
 
 		/// <summary>
@@ -532,11 +532,12 @@ namespace HordeServer.Agents
 		/// <param name="agent">The agent to create a lease for</param>
 		/// <param name="poolId">Pool to take the machine from</param>
 		/// <param name="conditions">Conditions to satisfy</param>
+		/// <param name="properties">Required properties to match</param>
 		/// <param name="resources">Resources required to execute</param>
 		/// <param name="exclusive">Whether the lease needs to be executed exclusively on the machine</param>
 		/// <param name="assignedResources">Resources allocated to the task</param>
 		/// <returns>True if the new lease can be granted</returns>
-		public static bool MeetsRequirements(this IAgent agent, PoolId? poolId, List<Condition> conditions, Dictionary<string, ResourceRequirements>? resources, bool exclusive, Dictionary<string, int> assignedResources)
+		public static bool MeetsRequirements(this IAgent agent, PoolId? poolId, List<Condition> conditions, IEnumerable<string>? properties, Dictionary<string, ResourceRequirements>? resources, bool exclusive, Dictionary<string, int> assignedResources)
 		{
 			if (!agent.Enabled || agent.Status != AgentStatus.Ok)
 			{
@@ -558,6 +559,10 @@ namespace HordeServer.Agents
 			{
 				return false;
 			}
+			if (properties != null && properties.Any(property => !agent.HasProperty(property)))
+			{
+				return false;
+			}
 			if (resources != null)
 			{
 				foreach ((string name, ResourceRequirements resourceRequirements) in resources)
diff --git a/Engine/Source/Programs/Shared/EpicGames.Horde/Compute/Requirements.cs b/Engine/Source/Programs/Shared/EpicGames.Horde/Compute/Requirements.cs
--- a/Engine/Source/Programs/Shared/EpicGames.Horde/Compute/Requirements.cs
+++ b/Engine/Source/Programs/Shared/EpicGames.Horde/Compute/Requirements.cs
@@ -15,25 +15,26 @@ namespace EpicGames.Horde.Compute
 		/// <summary>
 		/// Pool of machines to draw from
 		/// </summary>
-		[CbField("p")]
 		public string? Pool { get; set; }
 
 		/// <summary>
 		/// Condition string to be evaluated against the machine spec, eg. cpu-cores >= 10 &amp;&amp; ram.mb >= 200 &amp;&amp; pool == 'worker'
 		/// </summary>
-		[CbField("c")]
 		public Condition? Condition { get; set; }
 
+		/// <summary>
+		/// Properties required from the remote machine
+		/// </summary>
+		public HashSet<string> Properties { get; set; } = new HashSet<string>();
+
 		/// <summary>
 		/// Resources used by the process
 		/// </summary>
-		[CbField("r")]
 		public Dictionary<string, ResourceRequirements> Resources { get; } = new Dictionary<string, ResourceRequirements>();
 
 		/// <summary>
 		/// Whether we require exclusive access to the device
 		/// </summary>
-		[CbField("e")]
 		public bool Exclusive { get; set; }
 
 		/// <summary>
@@ -52,16 +53,6 @@ namespace EpicGames.Horde.Compute
 			Condition = condition;
 		}
 
-		/// <summary>
-		/// Serialize this object to bytes
-		/// </summary>
-		public byte[] Serialize()
-		{
-			CbWriter writer = new CbWriter();
-			CbSerializer.Serialize(writer, this);
-			return writer.ToByteArray();
-		}
-
 		/// <inheritdoc/>
 		public override string ToString()
 		{
@@ -94,13 +85,11 @@ namespace EpicGames.Horde.Compute
 		/// <summary>
 		/// Minimum allocation of the requested resource
 		/// </summary>
-		[CbField("min")]
 		public int Min { get; set; } = 1;
 
 		/// <summary>
 		/// Maximum allocation of the requested resource. Allocates as much as possible unless capped.
 		/// </summary>
-		[CbField("max")]
 		public int? Max { get; set; }
 	}
 }
