Horde: Advertise a new OSFamilyCompatibility property from agents, indicating OSes that the system can emulate (for Linux agents with WINE to indicate Windows compatibility).

diff --git a/Engine/Source/Programs/Horde/HordeAgent/Services/CapabilitiesService.cs b/Engine/Source/Programs/Horde/HordeAgent/Services/CapabilitiesService.cs
--- a/Engine/Source/Programs/Horde/HordeAgent/Services/CapabilitiesService.cs
+++ b/Engine/Source/Programs/Horde/HordeAgent/Services/CapabilitiesService.cs
@@ -88,12 +88,13 @@ namespace HordeAgent.Services
 			if (RuntimeInformation.IsOSPlatform(OSPlatform.Windows))
 			{
 				using IScope winScope = GlobalTracer.Instance.BuildSpan("GetWindowsCapabilities").StartActive();
-				capabilities.Properties.Add("Platform=Win64");
-				capabilities.Properties.Add("PlatformGroup=Windows");
-				capabilities.Properties.Add("PlatformGroup=Microsoft");
-				capabilities.Properties.Add("PlatformGroup=Desktop");
+				capabilities.Properties.Add($"{KnownPropertyNames.Platform}=Win64");
+				capabilities.Properties.Add($"{KnownPropertyNames.PlatformGroup}=Windows");
+				capabilities.Properties.Add($"{KnownPropertyNames.PlatformGroup}=Microsoft");
+				capabilities.Properties.Add($"{KnownPropertyNames.PlatformGroup}=Desktop");
 
-				capabilities.Properties.Add("OSFamily=Windows");
+				capabilities.Properties.Add($"{KnownPropertyNames.OsFamily}=Windows");
+				capabilities.Properties.Add($"{KnownPropertyNames.OsFamilyCompatibility}=Windows");
 
 				// WMI doesn't work currently on ARM64, due to reliance on a NET Framework assembly.
 				using (CimSession session = CimSession.Create(null))
@@ -229,13 +230,13 @@ namespace HordeAgent.Services
 			else if (RuntimeInformation.IsOSPlatform(OSPlatform.Linux))
 			{
 				using IScope _ = GlobalTracer.Instance.BuildSpan("GetLinuxCapabilities").StartActive();
-				capabilities.Properties.Add("Platform=Linux");
-				capabilities.Properties.Add("PlatformGroup=Linux");
-				capabilities.Properties.Add("PlatformGroup=Unix");
-				capabilities.Properties.Add("PlatformGroup=Desktop");
+				capabilities.Properties.Add($"{KnownPropertyNames.Platform}=Linux");
+				capabilities.Properties.Add($"{KnownPropertyNames.PlatformGroup}=Linux");
+				capabilities.Properties.Add($"{KnownPropertyNames.PlatformGroup}=Unix");
+				capabilities.Properties.Add($"{KnownPropertyNames.PlatformGroup}=Desktop");
 
-				capabilities.Properties.Add("OSFamily=Linux");
-				capabilities.Properties.Add("OSVersion=Linux");
+				capabilities.Properties.Add($"{KnownPropertyNames.OsFamily}=Linux");
+				capabilities.Properties.Add($"{KnownPropertyNames.OsFamilyCompatibility}=Linux");
 
 				// Add EC2 properties if needed
 				if (_settings.EnableAwsEc2Support)
@@ -245,6 +246,7 @@ namespace HordeAgent.Services
 				
 				if (_settings.WineExecutablePath != null)
 				{
+					capabilities.Properties.Add($"{KnownPropertyNames.OsFamilyCompatibility}=Windows");
 					capabilities.Properties.Add($"{KnownPropertyNames.WineEnabled}=true");
 				}
 				
@@ -306,12 +308,12 @@ namespace HordeAgent.Services
 			else if (RuntimeInformation.IsOSPlatform(OSPlatform.OSX))
 			{
 				using IScope _ = GlobalTracer.Instance.BuildSpan("GetMacCapabilities").StartActive();
-				capabilities.Properties.Add("Platform=Mac");
-				capabilities.Properties.Add("PlatformGroup=Apple");
-				capabilities.Properties.Add("PlatformGroup=Desktop");
+				capabilities.Properties.Add($"{KnownPropertyNames.Platform}=Mac");
+				capabilities.Properties.Add($"{KnownPropertyNames.PlatformGroup}=Apple");
+				capabilities.Properties.Add($"{KnownPropertyNames.PlatformGroup}=Desktop");
 
-				capabilities.Properties.Add("OSFamily=MacOS");
-				capabilities.Properties.Add("OSVersion=MacOS");
+				capabilities.Properties.Add($"{KnownPropertyNames.OsFamily}=MacOS");
+				capabilities.Properties.Add($"{KnownPropertyNames.OsFamilyCompatibility}=MacOS");
 
 				string output;
 				using (Process process = new Process())
diff --git a/Engine/Source/Programs/Shared/EpicGames.Horde/Agents/KnownPropertyNames.cs b/Engine/Source/Programs/Shared/EpicGames.Horde/Agents/KnownPropertyNames.cs
--- a/Engine/Source/Programs/Shared/EpicGames.Horde/Agents/KnownPropertyNames.cs
+++ b/Engine/Source/Programs/Shared/EpicGames.Horde/Agents/KnownPropertyNames.cs
@@ -11,17 +11,27 @@ public static class KnownPropertyNames
 	/// The agent id
 	/// </summary>
 	public const string Id = "Id";
-	
+
+	/// <summary>
+	/// The UBT platform enum
+	/// </summary>
+	public const string Platform = "Platform";
+
 	/// <summary>
-	/// Handle of the primary device (usually the agent VM itself)
+	/// The UBT platform group enum
 	/// </summary>
-	public const string PrimaryDeviceHandle = "Primary";
+	public const string PlatformGroup = "PlatformGroup";
 
 	/// <summary>
 	/// The operating system (Linux, MacOS, Windows)
 	/// </summary>
 	public const string OsFamily = "OSFamily";
 
+	/// <summary>
+	/// Compatible operating system (mainly for Linux WINE agents to advertise Windows support)
+	/// </summary>
+	public const string OsFamilyCompatibility = "OSFamilyCompatibility";
+
 	/// <summary>
 	/// Whether the agent is a .NET self-contained app
 	/// </summary>
