Horde: Prevent the agent capabilities hash changing unless a new capabilities object is provided.

#rnx

diff --git a/Engine/Source/Programs/Horde/Plugins/Compute/HordeServer.Compute/Agents/AgentScheduler.cs b/Engine/Source/Programs/Horde/Plugins/Compute/HordeServer.Compute/Agents/AgentScheduler.cs
--- a/Engine/Source/Programs/Horde/Plugins/Compute/HordeServer.Compute/Agents/AgentScheduler.cs
+++ b/Engine/Source/Programs/Horde/Plugins/Compute/HordeServer.Compute/Agents/AgentScheduler.cs
@@ -324,17 +324,16 @@ namespace HordeServer.Agents
 				LocalCapabilities capabilities = localSession.Capabilities;
 				if (newCapabilities != null)
 				{
-					byte[] capabilitiesData = newCapabilities.ToByteArray();
-					IoHash capabilitiesHash = IoHash.Compute(capabilitiesData);
+					byte[] newCapabilitiesData = newCapabilities.ToByteArray();
+					IoHash newCapabilitiesHash = IoHash.Compute(newCapabilitiesData);
 
-					if (capabilitiesHash != session.CapabilitiesHash)
+					if (newCapabilitiesHash != capabilities.Hash)
 					{
-						newSession.CapabilitiesHash = capabilitiesHash;
-						_ = transaction.StringSetAsync(Keys.Sessions[session.SessionId].Capabilities.Inner, capabilitiesData, flags: CommandFlags.FireAndForget);
-
-						capabilities = new LocalCapabilities(capabilitiesHash, newCapabilities);
+						capabilities = new LocalCapabilities(newCapabilitiesHash, newCapabilities);
+						_ = transaction.StringSetAsync(Keys.Sessions[session.SessionId].Capabilities.Inner, newCapabilitiesData, flags: CommandFlags.FireAndForget);
 					}
 				}
+				newSession.CapabilitiesHash = capabilities.Hash;
 
 				// Update the session state
 				_ = transaction.StringSetAsync(Keys.Sessions[session.SessionId].State, newSession, flags: CommandFlags.FireAndForget);
