Horde: Add better error for failing to start a process from the Horde Agent.

diff --git a/Engine/Source/Programs/Horde/HordeAgent/Leases/LeaseHandler.cs b/Engine/Source/Programs/Horde/HordeAgent/Leases/LeaseHandler.cs
--- a/Engine/Source/Programs/Horde/HordeAgent/Leases/LeaseHandler.cs
+++ b/Engine/Source/Programs/Horde/HordeAgent/Leases/LeaseHandler.cs
@@ -184,31 +184,38 @@ namespace HordeAgent.Leases
 		{
 			string commandLine = CommandLineArguments.Join(arguments);
 			logger.LogInformation("Running child process with arguments: {CommandLine}", commandLine);
-
-			using (ManagedProcessGroup processGroup = new ManagedProcessGroup())
-			using (ManagedProcess process = new ManagedProcess(processGroup, executable, commandLine, null, environment, ProcessPriorityClass.Normal))
+			try
 			{
-				for (; ; )
+				using (ManagedProcessGroup processGroup = new ManagedProcessGroup())
+				using (ManagedProcess process = new ManagedProcess(processGroup, executable, commandLine, null, environment, ProcessPriorityClass.Normal))
 				{
-					string? line = await process.ReadLineAsync(cancellationToken);
-					if (line == null)
+					for (; ; )
 					{
-						break;
+						string? line = await process.ReadLineAsync(cancellationToken);
+						if (line == null)
+						{
+							break;
+						}
+
+						JsonLogEvent jsonLogEvent;
+						if (JsonLogEvent.TryParse(line, out jsonLogEvent))
+						{
+							logger.LogJsonLogEvent(jsonLogEvent);
+						}
+						else
+						{
+							logger.LogInformation("{Line}", line);
+						}
 					}
 
-					JsonLogEvent jsonLogEvent;
-					if (JsonLogEvent.TryParse(line, out jsonLogEvent))
-					{
-						logger.LogJsonLogEvent(jsonLogEvent);
-					}
-					else
-					{
-						logger.LogInformation("{Line}", line);
-					}
+					await process.WaitForExitAsync(CancellationToken.None);
+					return process.ExitCode;
 				}
-
-				await process.WaitForExitAsync(CancellationToken.None);
-				return process.ExitCode;
+			}
+			catch (Exception ex)
+			{
+				logger.LogError(ex, "Failed to run process: {Message}", ex.Message);
+				throw;
 			}
 		}
 
