Horde: Include the executable name in log output for starting a new process from the agent.

#rnx

diff --git a/Engine/Source/Programs/Horde/HordeAgent/Leases/LeaseHandler.cs b/Engine/Source/Programs/Horde/HordeAgent/Leases/LeaseHandler.cs
--- a/Engine/Source/Programs/Horde/HordeAgent/Leases/LeaseHandler.cs
+++ b/Engine/Source/Programs/Horde/HordeAgent/Leases/LeaseHandler.cs
@@ -183,7 +183,7 @@ namespace HordeAgent.Leases
 		protected static async Task<int> RunProcessAsync(string executable, IEnumerable<string> arguments, IReadOnlyDictionary<string, string>? environment, ILogger logger, CancellationToken cancellationToken)
 		{
 			string commandLine = CommandLineArguments.Join(arguments);
-			logger.LogInformation("Running child process with arguments: {CommandLine}", commandLine);
+			logger.LogInformation("Running child process: {Executable} {CommandLine}", CommandLineArguments.Quote(executable), commandLine);
 			try
 			{
 				using (ManagedProcessGroup processGroup = new ManagedProcessGroup())
