Horde: Add a log message for creating a new session, to help track down machines stuck in a connection loop.

#rnx

diff --git a/Engine/Source/Programs/Horde/Plugins/Compute/HordeServer.Compute/Server/RpcService.cs b/Engine/Source/Programs/Horde/Plugins/Compute/HordeServer.Compute/Server/RpcService.cs
--- a/Engine/Source/Programs/Horde/Plugins/Compute/HordeServer.Compute/Server/RpcService.cs
+++ b/Engine/Source/Programs/Horde/Plugins/Compute/HordeServer.Compute/Server/RpcService.cs
@@ -188,6 +188,7 @@ namespace HordeServer.Server
 			}
 
 			AgentId agentId = new AgentId(request.Id);
+			_logger.LogInformation("Attempting to create session for agent {AgentId}", agentId);
 			using IDisposable? scope = _logger.BeginScope("CreateSession({AgentId})", agentId.ToString());
 
 			ComputeConfig computeConfig = _computeConfig.Value;
