Horde: Additional logging for terminating sessions.

diff --git a/Engine/Source/Programs/Horde/Plugins/Compute/HordeServer.Compute/Agents/AgentCollection.cs b/Engine/Source/Programs/Horde/Plugins/Compute/HordeServer.Compute/Agents/AgentCollection.cs
--- a/Engine/Source/Programs/Horde/Plugins/Compute/HordeServer.Compute/Agents/AgentCollection.cs
+++ b/Engine/Source/Programs/Horde/Plugins/Compute/HordeServer.Compute/Agents/AgentCollection.cs
@@ -517,14 +517,16 @@ namespace HordeServer.Agents
 					{
 						return null;
 					}
-					if (document.SessionId == null)
+
+					SessionId? sessionId = document.SessionId;
+					if (sessionId == null)
 					{
 						session = null;
 						break;
 					}
 
 					// Fetch the session
-					session = await _scheduler.TryGetSessionAsync(document.SessionId.Value, cancellationToken);
+					session = await _scheduler.TryGetSessionAsync(sessionId.Value, cancellationToken);
 					if (session != null)
 					{
 						break;
@@ -535,6 +537,7 @@ namespace HordeServer.Agents
 					document = await TryUpdateAsync(document, Builders<AgentDocument>.Update.Unset(x => x.SessionId).Set(x => x.LastOnlineTime, _clock.UtcNow), cancellationToken);
 					if (document != null)
 					{
+						_logger.LogInformation("Cleared session state from agent {AgentId}; no matching document for {SessionId} in Redis.", document.Id, sessionId.Value);
 						break;
 					}
 				}
@@ -1118,6 +1121,7 @@ namespace HordeServer.Agents
 		/// <inheritdoc/>
 		async Task<Agent?> TryTerminateSessionAsync(Agent agent, CancellationToken cancellationToken)
 		{
+			_logger.LogInformation("Terminating session {SessionId}", agent.Session?.SessionId ?? SessionId.Empty);
 			return await TryUpdateSessionAsync(agent, new UpdateSessionOptions(AgentStatus.Stopped), cancellationToken);
 		}
 
