Horde: Add tooltips for step times
#skipci

diff --git a/Engine/Source/Programs/Horde/HordeDashboard/src/base/utilities/timeUtils.ts b/Engine/Source/Programs/Horde/HordeDashboard/src/base/utilities/timeUtils.ts
--- a/Engine/Source/Programs/Horde/HordeDashboard/src/base/utilities/timeUtils.ts
+++ b/Engine/Source/Programs/Horde/HordeDashboard/src/base/utilities/timeUtils.ts
@@ -10,9 +10,11 @@ export const displayTimeZone = () => {
     return dashboard.displayUTC ? "UTC" : localTimeZone;
 }
 
-type HordeTime = {
+export type HordeTime = {
     display: string;
     server: string;
+    displayNice?: string;
+    serverNice?: string;
 }
 
 export const msecToElapsed = (millisec: number, includeMinutes: boolean = true, includeSeconds: boolean = true): string => {
@@ -119,12 +121,14 @@ export const getStepETA = (step: StepData, job: JobData): HordeTime => {
 
     end = roundTime(end, moment.duration(5, "minutes"));
 
-    const serverETA = moment.utc(end).tz(serverTimeZone).format(format);
-    const displayETA = moment.utc(end).tz(displayTimeZone()).format(format);
+    const serverETA = moment.utc(end).tz(serverTimeZone);
+    const displayETA = moment.utc(end).tz(displayTimeZone());
 
     return {
-        display: displayETA,
-        server: serverETA
+        display: displayETA.format(format),
+        server: serverETA.format(format),
+        displayNice: getNiceTime(displayETA.toDate()),
+        serverNice: getNiceTime(serverETA.toDate())
     };
 
 };
@@ -133,18 +137,27 @@ export const getStepFinishTime = (step: StepData): HordeTime => {
 
     let displayFinish = "";
     let serverFinish = "";
+    let displayNice: string | undefined;
+    let serverNice: string | undefined;
 
     const format = dashboard.display24HourClock ? "HH:mm:ss z" : "LT z";
 
     if (step.finishTime) {
         const end = moment(step.finishTime);
-        displayFinish = moment.utc(end).tz(displayTimeZone()).format(format);
-        serverFinish = moment.utc(end).tz(serverTimeZone).format(format);
+        const dfinish = moment.utc(end).tz(displayTimeZone());
+        const sfinish = moment.utc(end).tz(serverTimeZone);
+
+        displayFinish = dfinish.format(format);
+        serverFinish = sfinish.format(format);
+        displayNice = getNiceTime(dfinish.toDate());
+        serverNice = getNiceTime(sfinish.toDate());
     }
 
     return {
         display: displayFinish,
-        server: serverFinish
+        server: serverFinish,
+        displayNice: displayNice,
+        serverNice: serverFinish
     };
 };
 
@@ -152,18 +165,28 @@ export const getStepStartTime = (step: StepData): HordeTime => {
 
     let displayStart = "";
     let serverStart = "";
+    let displayNice: string | undefined;
+    let serverNice: string | undefined;
 
     const format = dashboard.display24HourClock ? "HH:mm:ss z" : "LT z";
 
     if (step.startTime) {
         const end = moment(step.startTime);
-        displayStart = moment.utc(end).tz(displayTimeZone()).format(format);
-        serverStart = moment.utc(end).tz(serverTimeZone).format(format);
+        const dstart = moment.utc(end).tz(displayTimeZone());
+        const sstart = moment.utc(end).tz(serverTimeZone);
+
+        displayStart = dstart.format(format);
+        serverStart = sstart.format(format);
+        displayNice = getNiceTime(dstart.toDate());
+        serverNice = getNiceTime(sstart.toDate());
+
     }
 
     return {
         display: displayStart,
-        server: serverStart
+        server: serverStart,
+        displayNice: displayNice,
+        serverNice: serverNice
     };
 };
 
diff --git a/Engine/Source/Programs/Horde/HordeDashboard/src/components/jobDetailsV2/JobDetailViewSteps.tsx b/Engine/Source/Programs/Horde/HordeDashboard/src/components/jobDetailsV2/JobDetailViewSteps.tsx
--- a/Engine/Source/Programs/Horde/HordeDashboard/src/components/jobDetailsV2/JobDetailViewSteps.tsx
+++ b/Engine/Source/Programs/Horde/HordeDashboard/src/components/jobDetailsV2/JobDetailViewSteps.tsx
@@ -1,11 +1,11 @@
-import { CollapseAllVisibility, ConstrainMode, DetailsHeader, DetailsList, DetailsListLayoutMode, DetailsRow, IColumn, Icon, IDetailsListProps, ProgressIndicator, SelectionMode, Stack, Text } from "@fluentui/react";
+import { CollapseAllVisibility, ConstrainMode, DetailsHeader, DetailsList, DetailsListLayoutMode, DetailsRow, DirectionalHint, IColumn, Icon, IDetailsListProps, ProgressIndicator, SelectionMode, Stack, Text, TooltipHost } from "@fluentui/react";
 import { observer } from "mobx-react-lite";
-import React, { useEffect, useState } from "react";
+import React, { useEffect, useId, useState } from "react";
 import { Link, useParams } from "react-router-dom";
 import { BatchData, GetBatchResponse, JobStepBatchError, JobStepBatchState, JobStepOutcome, JobStepState, StepData } from "../../backend/Api";
 import dashboard, { StatusColor } from "../../backend/Dashboard";
 import { ISideRailLink } from "../../base/components/SideRail";
-import { getBatchInitElapsed, getNiceTime, getStepElapsed, getStepETA, getStepFinishTime, getStepPercent, getStepStartTime, getStepTimingDelta } from "../../base/utilities/timeUtils";
+import { getBatchInitElapsed, getNiceTime, getStepElapsed, getStepETA, getStepFinishTime, getStepPercent, getStepStartTime, getStepTimingDelta, HordeTime } from "../../base/utilities/timeUtils";
 import { HistoryModal } from "../HistoryModal";
 import { getBatchText, getStepStatusMessage } from "../JobDetailCommon";
 import { StepStatusIcon } from "../StatusIcon";
@@ -53,10 +53,13 @@ class StepsDataView extends JobDataView {
 
 }
 
+let stepToolTipId = 0;
+
 JobDetailsV2.registerDataView("StepsDataView", (details: JobDetailsV2) => new StepsDataView(details));
 
 const RenderDynamic: React.FC<{ jobDetails: JobDetailsV2, dataView: StepsDataView, step?: StepData, batch?: GetBatchResponse, column: string }> = observer(({ jobDetails, dataView, column, step, batch }) => {
 
+
    dataView.subscribe();
 
    if (step && !step.finishTime && !step.abortRequested && (step.state !== JobStepState.Aborted && step.state !== JobStepState.Completed && step.state !== JobStepState.Skipped)) {
@@ -109,7 +112,7 @@ const RenderDynamic: React.FC<{ jobDetails: JobDetailsV2, dataView: StepsDataVie
 
    if (step && column === "ETA") {
 
-      let eta = {
+      let eta:HordeTime = {
          display: "",
          server: ""
       };
@@ -118,7 +121,7 @@ const RenderDynamic: React.FC<{ jobDetails: JobDetailsV2, dataView: StepsDataVie
          return null;
       }
 
-      let finished = { display: "", server: "" };
+      let finished:HordeTime = { display: "", server: "" };
 
       eta = getStepETA(step, jobDetails.jobData!);
 
@@ -127,6 +130,8 @@ const RenderDynamic: React.FC<{ jobDetails: JobDetailsV2, dataView: StepsDataVie
       if (finished.display) {
          eta.display = finished.display;
          eta.server = finished.server;
+         eta.displayNice = finished.displayNice;
+         eta.serverNice = finished.serverNice;
       }
 
       let time = eta.display;
@@ -136,17 +141,17 @@ const RenderDynamic: React.FC<{ jobDetails: JobDetailsV2, dataView: StepsDataVie
       const color = !step.finishTime ? etaColor : undefined;
 
       // Open Sans tilde rendering issue at 13px, and not rendering at all at other px: [REDACTED_PROFILE_URL]/fonts/issues/399, do not change from 12px
-      return <Stack horizontalAlign={"end"}>
+      return <TooltipHost content={eta.displayNice} directionalHint={DirectionalHint.leftCenter} id={`step_unique_tooltip_${stepToolTipId++}`}><Stack horizontalAlign={"end"}>
          <Stack horizontal tokens={{ childrenGap: 2 }}>
             {!!time && !step.finishTime && <Text style={{ color: color, fontSize: "11px", paddingTop: 2 }}>~</Text>}
             <Text style={{ color: color, fontSize: "13px" }}>
                {time}
             </Text>
          </Stack>
-      </Stack>;
+      </Stack>
+      </TooltipHost>;
    };
 
-
    return null;
 
 });
@@ -317,7 +322,7 @@ export const StepsPanelInner: React.FC<{ jobDetails: JobDetailsV2, depStepId?: s
          if (!started.display || !started.server) {
             return null;
          } else {
-            return <Stack horizontalAlign={"end"}><Text style={{ fontSize: "13px" }}>{started.display}</Text></Stack>;
+            return <TooltipHost content={started.displayNice ?? started.display} directionalHint={DirectionalHint.leftCenter} id={`step_unique_tooltip_${stepToolTipId++}`}><Stack horizontalAlign={"end"}><Text style={{ fontSize: "13px" }}>{started.display}</Text></Stack></TooltipHost>;
          }
       };
 
