Horde: Add a mechanism to repair inconsistent state between capabilities and session objects in Redis.

#rnx

diff --git a/Engine/Source/Programs/Horde/Plugins/Compute/HordeServer.Compute/Agents/AgentCollection.cs b/Engine/Source/Programs/Horde/Plugins/Compute/HordeServer.Compute/Agents/AgentCollection.cs
--- a/Engine/Source/Programs/Horde/Plugins/Compute/HordeServer.Compute/Agents/AgentCollection.cs
+++ b/Engine/Source/Programs/Horde/Plugins/Compute/HordeServer.Compute/Agents/AgentCollection.cs
@@ -818,6 +818,7 @@ namespace HordeServer.Agents
 			AgentId agentId = agent.Document.Id;
 			if (agent.Session == null)
 			{
+				_logger.LogWarning("Agent {AgentId} does not have an active session; cannot update", agent.Id);
 				return null;
 			}
 
@@ -865,6 +866,7 @@ namespace HordeServer.Agents
 					newDocument = await TryUpdateAsync(newDocument, updateBuilder.Combine(updates), cancellationToken);
 					if (newDocument == null)
 					{
+						_logger.LogTrace("Update of agent document {AgentId} did not succeed; will retry", agent.Document.Id);
 						return null;
 					}
 				}
@@ -914,6 +916,7 @@ namespace HordeServer.Agents
 			newSession = await _scheduler.TryUpdateSessionAsync(agent.Session, newSession, options.Capabilities, cancellationToken);
 			if (newSession == null)
 			{
+				_logger.LogTrace("Update of agent {AgentId} session {SessionId} did not succeed; will retry", agent.Document.Id, agent.Session.SessionId);
 				return null;
 			}
 
@@ -1121,7 +1124,7 @@ namespace HordeServer.Agents
 		/// <inheritdoc/>
 		async Task<Agent?> TryTerminateSessionAsync(Agent agent, CancellationToken cancellationToken)
 		{
-			_logger.LogInformation("Terminating session {SessionId}", agent.Session?.SessionId ?? SessionId.Empty);
+			_logger.LogInformation("Terminating session {SessionId} for {AgentId}", agent.Session?.SessionId ?? SessionId.Empty, agent.Id);
 			return await TryUpdateSessionAsync(agent, new UpdateSessionOptions(AgentStatus.Stopped), cancellationToken);
 		}
 
diff --git a/Engine/Source/Programs/Horde/Plugins/Compute/HordeServer.Compute/Agents/AgentScheduler.cs b/Engine/Source/Programs/Horde/Plugins/Compute/HordeServer.Compute/Agents/AgentScheduler.cs
--- a/Engine/Source/Programs/Horde/Plugins/Compute/HordeServer.Compute/Agents/AgentScheduler.cs
+++ b/Engine/Source/Programs/Horde/Plugins/Compute/HordeServer.Compute/Agents/AgentScheduler.cs
@@ -148,12 +148,26 @@ namespace HordeServer.Agents
 					byte[]? data = (byte[]?)await database.StringGetAsync(Keys.Sessions[session.SessionId].Capabilities.Inner).WaitAsync(cancellationToken);
 					if (data == null)
 					{
+						_logger.LogTrace("Missing capabilities key for session {SessionId}; can retry.", session.SessionId);
 						return null;
 					}
 
 					IoHash hash = IoHash.Compute(data);
 					if (session.CapabilitiesHash != hash)
 					{
+						ITransaction transaction = database.CreateTransaction();
+						transaction.AddCondition(Keys.Sessions[session.SessionId].State.StringEqual(session));
+						transaction.AddCondition(Condition.StringEqual(Keys.Sessions[session.SessionId].Capabilities.Inner, data));
+						_ = transaction.StringSetAsync(Keys.Sessions[session.SessionId].State, new RpcSession(session) { CapabilitiesHash = hash, UpdateTicks = session.UpdateTicks + 1 }, flags: CommandFlags.FireAndForget);
+
+						if (await transaction.ExecuteAsync())
+						{
+							_logger.LogWarning("Repaired mismatch between capabilities and session hash for agent {AgentId} session {SessionId}. Was {OldHash}, now {NewHash}.", session.AgentId, session.SessionId, session.CapabilitiesHash, hash);  
+						}
+						else
+						{
+							_logger.LogTrace("Hash mismatch for capabilities of session {SessionId}; can retry.", session.SessionId);
+						}
 						return null;
 					}
 
@@ -257,6 +271,7 @@ namespace HordeServer.Agents
 
 			if (!await transaction.ExecuteAsync().WaitAsync(cancellationToken))
 			{
+				_logger.LogTrace("Unable to create session {SessionId}; can retry", sessionId);
 				return null;
 			}
 
@@ -375,6 +390,7 @@ namespace HordeServer.Agents
 			// Execute the transaction
 			if (!await transaction.ExecuteAsync().WaitAsync(cancellationToken))
 			{
+				_logger.LogTrace("Transaction failed for update of session {SessionId}", sessionId);
 				return null;
 			}
 
