Horde: Adding Horde installer custom actions, data directory selector, and fix server to be able to bootstrap using custom data directory

diff --git a/Engine/Source/Programs/Horde/HordeServer/ServerApp.cs b/Engine/Source/Programs/Horde/HordeServer/ServerApp.cs
--- a/Engine/Source/Programs/Horde/HordeServer/ServerApp.cs
+++ b/Engine/Source/Programs/Horde/HordeServer/ServerApp.cs
@@ -135,7 +135,18 @@ namespace HordeServer
 			// Set the default data directory
 			if (baseServerSettings.DataDir != null)
 			{
-				s_dataDir = DirectoryReference.Combine(GetAppDir(), baseServerSettings.DataDir);
+				if (Path.IsPathRooted(baseServerSettings.DataDir))
+				{
+					DirectoryReference? dataDir = DirectoryReference.FromString(baseServerSettings.DataDir);
+					if (dataDir != null)
+					{
+						s_dataDir = dataDir;
+					}					
+				}
+				else
+				{
+					s_dataDir = DirectoryReference.Combine(GetAppDir(), baseServerSettings.DataDir);
+				}								
 			}
 			else if (baseServerSettings.Installed && RuntimeInformation.IsOSPlatform(OSPlatform.Windows))
 			{
@@ -303,9 +314,29 @@ namespace HordeServer
 		static IConfiguration CreateConfig(bool readInstalledConfig, FileReference? serverConfigFile)
 		{
 			IConfigurationBuilder builder = new ConfigurationBuilder();
-			if (readInstalledConfig && OperatingSystem.IsWindows())
+
+			string? regDataDir = null;
+			if (OperatingSystem.IsWindows())
 			{
-				builder = builder.Add(new RegistryConfigurationSource(Registry.LocalMachine, "SOFTWARE\\Epic Games\\Horde\\Server", ServerSettings.SectionName));
+				if (readInstalledConfig)
+				{
+					builder = builder.Add(new RegistryConfigurationSource(Registry.LocalMachine, "SOFTWARE\\Epic Games\\Horde\\Server", ServerSettings.SectionName));
+				}
+				else
+				{
+					// Check the registry for an installed server, if this is running we need to set the data directory if available
+					string? installedServerExecutable = Registry.GetValue("HKEY_LOCAL_MACHINE\\SOFTWARE\\Epic Games\\Horde\\Server", "InstalledServerExecutable", null) as string;
+					if (!String.IsNullOrEmpty(installedServerExecutable))
+					{
+						FileReference installedServer = new FileReference(installedServerExecutable);
+						FileReference executingServer = new FileReference(Assembly.GetExecutingAssembly().Location);
+
+						if (installedServer == executingServer)
+						{
+							regDataDir = Registry.GetValue("HKEY_LOCAL_MACHINE\\SOFTWARE\\Epic Games\\Horde\\Server", "DataDir", null) as string;
+						}
+					}					
+				}
 			}
 
 			builder.SetBasePath(AppDir.FullName)
@@ -347,6 +378,12 @@ namespace HordeServer
 				}
 			}
 
+			// set data directory if we have one from registry
+			if (!readInstalledConfig && !String.IsNullOrEmpty(regDataDir))
+			{
+				remappedValues.Add(new KeyValuePair<string, string?>("Horde:DataDir", regDataDir));
+			}
+
 			builder.Add(new MemoryConfigurationSource { InitialData = remappedValues });
 			return builder.Build();
 		}
diff --git a/Engine/Source/Programs/Horde/Installer/CustomActions/CustomsActions.cs b/Engine/Source/Programs/Horde/Installer/CustomActions/CustomsActions.cs
new file mode 100644
--- /dev/null
+++ b/Engine/Source/Programs/Horde/Installer/CustomActions/CustomsActions.cs
@@ -0,0 +1,60 @@
+// Copyright Epic Games, Inc. All Rights Reserved.
+
+using System;
+using System.Runtime.InteropServices;
+using System.Threading;
+using Microsoft.Deployment.WindowsInstaller;
+using System.Windows.Forms;
+using static System.Environment;
+using System.IO;
+
+namespace HordeInstaller
+{
+	public class InstallerActions
+	{
+		[DllImport("user32.dll", CharSet = CharSet.Auto, SetLastError = true)]
+		static extern IntPtr GetForegroundWindow();
+
+		[CustomAction]
+		public static ActionResult SelectDataDir(Session session)
+		{
+			try
+			{
+				Thread task = new Thread(() => GetDataDir(session));
+				task.SetApartmentState(ApartmentState.STA);
+				task.Start();
+				task.Join();
+			}
+			catch (Exception ex)
+			{
+				session.Log("SelectDataDir Exception: {0}\n StackTrace: {1}", ex.Message, ex.StackTrace);
+				return ActionResult.Failure;
+			}
+			return ActionResult.Success;
+		}
+
+		private static void GetDataDir(Session session)
+		{
+			using (FolderBrowserDialog dialog = new FolderBrowserDialog())
+			{				
+				string defaultPath = @"C:\ProgramData\Epic\Horde\Server";
+
+				if (!Directory.Exists(defaultPath))
+				{
+					Directory.CreateDirectory(defaultPath);
+				}
+
+				dialog.RootFolder = SpecialFolder.Desktop;				
+				dialog.SelectedPath = @"C:\ProgramData\Epic\Horde\Server";
+				dialog.Description = @"Please select a folder for Horde Server data";
+				
+				DialogResult result = dialog.ShowDialog(Control.FromHandle(GetForegroundWindow()));
+
+				if (result == DialogResult.OK)
+				{
+					session["HORDE_DATA_DIR"] = dialog.SelectedPath;
+				}
+			}
+		}
+	}
+}
diff --git a/Engine/Source/Programs/Horde/Installer/CustomActions/InstallerCustomActions.csproj b/Engine/Source/Programs/Horde/Installer/CustomActions/InstallerCustomActions.csproj
new file mode 100644
--- /dev/null
+++ b/Engine/Source/Programs/Horde/Installer/CustomActions/InstallerCustomActions.csproj
@@ -0,0 +1,25 @@
+<Project Sdk="Microsoft.NET.Sdk">
+
+  <PropertyGroup>
+    <TargetFramework>net20</TargetFramework>
+    <RootNamespace>HordeInstaller</RootNamespace>
+    <RunAnalyzersDuringBuild>False</RunAnalyzersDuringBuild>
+    <Configurations>Debug;Release;Analyze</Configurations>
+  </PropertyGroup>
+
+  <PropertyGroup Condition=" '$(Configuration)|$(Platform)' == 'Analyze|AnyCPU' ">
+    <RunAnalyzersDuringBuild>False</RunAnalyzersDuringBuild>
+    <TreatWarningsAsErrors>false</TreatWarningsAsErrors>
+    <EnforceCodeStyleInBuild>True</EnforceCodeStyleInBuild>
+    <GenerateDocumentationFile>False</GenerateDocumentationFile>
+    <DocumentationFile></DocumentationFile>
+  </PropertyGroup>
+    
+  <ItemGroup>
+    <Reference Include="Microsoft.Deployment.WindowsInstaller">
+      <HintPath>..\..\..\..\ThirdParty\WiX\3.8\Microsoft.Deployment.WindowsInstaller.dll</HintPath>
+    </Reference>
+    <Reference Include="System.Windows.Forms" />
+  </ItemGroup>
+
+</Project>
diff --git a/Engine/Source/Programs/Horde/Installer/CustomActions/README.txt b/Engine/Source/Programs/Horde/Installer/CustomActions/README.txt
new file mode 100644
--- /dev/null
+++ b/Engine/Source/Programs/Horde/Installer/CustomActions/README.txt
@@ -0,0 +1,5 @@
+
+The custom action must be in a self extractig DLL file, which uses absolute paths, so like this:
+
+Engine\Source\ThirdParty\WiX\3.8\sdk\MakeSfxCA.exe D:\InstallerCustomActions.CA.dll D:\P4\UE5-Workspace\Engine\Source\ThirdParty\WiX\3.8\sdk\x64\SfxCA.dll D:\P4\UE5-Workspace\Engine\Source\Programs\Horde\Installer\CustomActions\bin\Release\net20\InstallerCustomActions.dll D:\P4\UE5-Workspace\Engine\Source\Programs\Horde\Installer\CustomActions\bin\Release\net20\Microsoft.Deployment.WindowsInstaller.dll
+
