Horde: Increase verbosity for agent log output in the dev environment.

#rnx

diff --git a/Engine/Source/Programs/Horde/HordeServer/appsettings.Development.json b/Engine/Source/Programs/Horde/HordeServer/appsettings.Development.json
--- a/Engine/Source/Programs/Horde/HordeServer/appsettings.Development.json
+++ b/Engine/Source/Programs/Horde/HordeServer/appsettings.Development.json
@@ -8,6 +8,8 @@
 		{
 			"Default": "Debug",
 			"Override": {
+				"HordeServer.Agents.AgentCollection": "Verbose",
+				"HordeServer.Agents.AgentScheduler": "Verbose",
 				"HordeServer.Perforce.PerforceService": "Verbose"
 			}
 		}
