Horde: Change logging levels for agent collections to debug.

#rnx

diff --git a/Engine/Source/Programs/Horde/Plugins/Compute/HordeServer.Compute/Agents/AgentCollection.cs b/Engine/Source/Programs/Horde/Plugins/Compute/HordeServer.Compute/Agents/AgentCollection.cs
--- a/Engine/Source/Programs/Horde/Plugins/Compute/HordeServer.Compute/Agents/AgentCollection.cs
+++ b/Engine/Source/Programs/Horde/Plugins/Compute/HordeServer.Compute/Agents/AgentCollection.cs
@@ -866,7 +866,7 @@ namespace HordeServer.Agents
 					newDocument = await TryUpdateAsync(newDocument, updateBuilder.Combine(updates), cancellationToken);
 					if (newDocument == null)
 					{
-						_logger.LogTrace("Update of agent document {AgentId} did not succeed; will retry", agent.Document.Id);
+						_logger.LogDebug("Update of agent document {AgentId} did not succeed; will retry", agent.Document.Id);
 						return null;
 					}
 				}
@@ -916,7 +916,7 @@ namespace HordeServer.Agents
 			newSession = await _scheduler.TryUpdateSessionAsync(agent.Session, newSession, options.Capabilities, cancellationToken);
 			if (newSession == null)
 			{
-				_logger.LogTrace("Update of agent {AgentId} session {SessionId} did not succeed; will retry", agent.Document.Id, agent.Session.SessionId);
+				_logger.LogDebug("Update of agent {AgentId} session {SessionId} did not succeed; will retry", agent.Document.Id, agent.Session.SessionId);
 				return null;
 			}
 
diff --git a/Engine/Source/Programs/Horde/Plugins/Compute/HordeServer.Compute/Agents/AgentScheduler.cs b/Engine/Source/Programs/Horde/Plugins/Compute/HordeServer.Compute/Agents/AgentScheduler.cs
--- a/Engine/Source/Programs/Horde/Plugins/Compute/HordeServer.Compute/Agents/AgentScheduler.cs
+++ b/Engine/Source/Programs/Horde/Plugins/Compute/HordeServer.Compute/Agents/AgentScheduler.cs
@@ -148,7 +148,7 @@ namespace HordeServer.Agents
 					byte[]? data = (byte[]?)await database.StringGetAsync(Keys.Sessions[session.SessionId].Capabilities.Inner).WaitAsync(cancellationToken);
 					if (data == null)
 					{
-						_logger.LogTrace("Missing capabilities key for session {SessionId}; can retry.", session.SessionId);
+						_logger.LogDebug("Missing capabilities key for session {SessionId}; can retry.", session.SessionId);
 						return null;
 					}
 
@@ -166,7 +166,7 @@ namespace HordeServer.Agents
 						}
 						else
 						{
-							_logger.LogTrace("Hash mismatch for capabilities of session {SessionId}; can retry.", session.SessionId);
+							_logger.LogDebug("Hash mismatch for capabilities of session {SessionId}; can retry.", session.SessionId);
 						}
 						return null;
 					}
@@ -271,7 +271,7 @@ namespace HordeServer.Agents
 
 			if (!await transaction.ExecuteAsync().WaitAsync(cancellationToken))
 			{
-				_logger.LogTrace("Unable to create session {SessionId}; can retry", sessionId);
+				_logger.LogDebug("Unable to create session {SessionId}; can retry", sessionId);
 				return null;
 			}
 
@@ -390,7 +390,7 @@ namespace HordeServer.Agents
 			// Execute the transaction
 			if (!await transaction.ExecuteAsync().WaitAsync(cancellationToken))
 			{
-				_logger.LogTrace("Transaction failed for update of session {SessionId}", sessionId);
+				_logger.LogDebug("Transaction failed for update of session {SessionId}", sessionId);
 				return null;
 			}
 
