Horde: Fix leases not being cleared when terminating a session.

#rnx

diff --git a/Engine/Source/Programs/Horde/Plugins/Compute/HordeServer.Compute/Agents/AgentScheduler.cs b/Engine/Source/Programs/Horde/Plugins/Compute/HordeServer.Compute/Agents/AgentScheduler.cs
--- a/Engine/Source/Programs/Horde/Plugins/Compute/HordeServer.Compute/Agents/AgentScheduler.cs
+++ b/Engine/Source/Programs/Horde/Plugins/Compute/HordeServer.Compute/Agents/AgentScheduler.cs
@@ -320,6 +320,8 @@ namespace HordeServer.Agents
 			{
 				newSession.UpdateTicks = Math.Max(newSession.UpdateTicks + 1, Math.Min(newSession.UpdateTicks + RpcSession.ExpireAfterTime.Ticks, _clock.UtcNow.Ticks));
 
+				newSession.Leases.Clear(); // Need to ensure we remove items from the active set below
+
 				_ = transaction.HashDeleteAsync(Keys.Sessions.Current, newSession.SessionId, flags: CommandFlags.FireAndForget);
 				_ = transaction.KeyDeleteAsync(Keys.Sessions[sessionId].Capabilities, flags: CommandFlags.FireAndForget);
 				_ = transaction.KeyDeleteAsync(Keys.Sessions[sessionId].State, flags: CommandFlags.FireAndForget);
