Horde: Cache agent rate values with more typical Redis keys.

#rnx

diff --git a/Engine/Source/Programs/Horde/Plugins/Compute/HordeServer.Compute/Agents/AgentService.cs b/Engine/Source/Programs/Horde/Plugins/Compute/HordeServer.Compute/Agents/AgentService.cs
--- a/Engine/Source/Programs/Horde/Plugins/Compute/HordeServer.Compute/Agents/AgentService.cs
+++ b/Engine/Source/Programs/Horde/Plugins/Compute/HordeServer.Compute/Agents/AgentService.cs
@@ -769,7 +769,7 @@ namespace HordeServer.Agents
 		/// <returns>Hourly rate of running the given agent</returns>
 		public async ValueTask<double?> GetRateAsync(AgentId agentId, CancellationToken cancellationToken = default)
 		{
-			RedisKey key = $"agent-rate/{agentId}";
+			RedisKey key = $"compute:agent-rates:{agentId}";
 
 			// Try to get the current value
 			RedisValue value = await _redisService.GetDatabase().StringGetAsync(key).WaitAsync(cancellationToken);
