Horde: Rename the Redis key for agent rates.

#rnx

diff --git a/Engine/Source/Programs/Horde/Plugins/Compute/HordeServer.Compute/Agents/AgentService.cs b/Engine/Source/Programs/Horde/Plugins/Compute/HordeServer.Compute/Agents/AgentService.cs
--- a/Engine/Source/Programs/Horde/Plugins/Compute/HordeServer.Compute/Agents/AgentService.cs
+++ b/Engine/Source/Programs/Horde/Plugins/Compute/HordeServer.Compute/Agents/AgentService.cs
@@ -76,7 +76,7 @@ namespace HordeServer.Agents
 		readonly ILogger _logger;
 		readonly ITicker _ticker;
 
-		readonly RedisStringKey<AgentRateTable> _agentRateTableData = new("agent-rates");
+		readonly RedisStringKey<AgentRateTable> _agentRateTableData = new("compute:agent-rates");
 
 		/// <summary>Lazily updated costs for different agent types</summary>
 		readonly AsyncCachedValue<AgentRateTable?> _cachedRates;
