Horde: Remove IAsyncDisposable interface from IHordeClient. This is something that belongs to implementations, not on the interface.

#rnx

diff --git a/Engine/Source/Programs/Horde/Drivers/JobDriver.Tests/Execution/WorkspaceExecutorTest.cs b/Engine/Source/Programs/Horde/Drivers/JobDriver.Tests/Execution/WorkspaceExecutorTest.cs
--- a/Engine/Source/Programs/Horde/Drivers/JobDriver.Tests/Execution/WorkspaceExecutorTest.cs
+++ b/Engine/Source/Programs/Horde/Drivers/JobDriver.Tests/Execution/WorkspaceExecutorTest.cs
@@ -28,7 +28,7 @@ using EpicGames.Horde.Tools;
 namespace JobDriver.Tests.Execution
 {
 	[TestClass]
-	public sealed class WorkspaceExecutorTest : IAsyncDisposable
+	public sealed class WorkspaceExecutorTest : IDisposable
 	{
 		class FakeHordeClient : IHordeClient
 		{
@@ -115,9 +115,6 @@ namespace JobDriver.Tests.Execution
 			public IStorageClient CreateStorageClient(string relativePath, string? accessToken = null)
 				=> throw new NotImplementedException();
 
-			public ValueTask DisposeAsync()
-				=> default;
-
 			public Task<string?> GetAccessTokenAsync(bool interactive, CancellationToken cancellationToken = default)
 				=> Task.FromResult<string?>("access-token");
 
@@ -163,10 +160,9 @@ namespace JobDriver.Tests.Execution
 			_executor = new(executorOptions, _workspace, null, NullLogger.Instance);
 		}
 
-		public async ValueTask DisposeAsync()
+		public void Dispose()
 		{
 			_executor.Dispose();
-			await _hordeClient.DisposeAsync();
 			_workspace.Dispose();
 			_autoSdkWorkspace.Dispose();
 		}
diff --git a/Engine/Source/Programs/Horde/Drivers/JobDriver/Commands/Execution/ExecuteConformCommand.cs b/Engine/Source/Programs/Horde/Drivers/JobDriver/Commands/Execution/ExecuteConformCommand.cs
--- a/Engine/Source/Programs/Horde/Drivers/JobDriver/Commands/Execution/ExecuteConformCommand.cs
+++ b/Engine/Source/Programs/Horde/Drivers/JobDriver/Commands/Execution/ExecuteConformCommand.cs
@@ -30,10 +30,10 @@ namespace JobDriver.Commands.Execution
 		[CommandLine("-WorkingDir=", Required = true)]
 		public DirectoryReference WorkingDir { get; set; } = null!;
 
-		readonly IHordeClientFactory _hordeClientFactory;
+		readonly HordeClientFactory _hordeClientFactory;
 		readonly IOptions<DriverSettings> _driverSettings;
 
-		public ExecuteConformCommand(IHordeClientFactory hordeClientFactory, IOptions<DriverSettings> driverSettings)
+		public ExecuteConformCommand(HordeClientFactory hordeClientFactory, IOptions<DriverSettings> driverSettings)
 		{
 			_hordeClientFactory = hordeClientFactory;
 			_driverSettings = driverSettings;
@@ -45,7 +45,7 @@ namespace JobDriver.Commands.Execution
 
 			ConformTask conformTask = ConformTask.Parser.ParseFrom(Convert.FromBase64String(Task));
 
-			IHordeClient hordeClient = _hordeClientFactory.Create();
+			await using HordeClient hordeClient = _hordeClientFactory.Create();
 
 			ConformExecutor conformExecutor = new ConformExecutor(hordeClient, WorkingDir, AgentId, LeaseId, conformTask, _driverSettings.Value, logger);
 			await conformExecutor.ExecuteAsync(CancellationToken.None);
diff --git a/Engine/Source/Programs/Horde/Drivers/JobDriver/Commands/Execution/ExecuteJobCommand.cs b/Engine/Source/Programs/Horde/Drivers/JobDriver/Commands/Execution/ExecuteJobCommand.cs
--- a/Engine/Source/Programs/Horde/Drivers/JobDriver/Commands/Execution/ExecuteJobCommand.cs
+++ b/Engine/Source/Programs/Horde/Drivers/JobDriver/Commands/Execution/ExecuteJobCommand.cs
@@ -30,11 +30,11 @@ namespace JobDriver.Commands.Execution
 		[CommandLine("-WorkingDir=", Required = true)]
 		public DirectoryReference WorkingDir { get; set; } = null!;
 
-		readonly IHordeClientFactory _hordeClientFactory;
+		readonly HordeClientFactory _hordeClientFactory;
 		readonly IEnumerable<IJobExecutorFactory> _jobExecutorFactories;
 		readonly IOptions<DriverSettings> _driverSettings;
 
-		public ExecuteJobCommand(IHordeClientFactory hordeClientFactory, IEnumerable<IJobExecutorFactory> jobExecutorFactories, IOptions<DriverSettings> driverSettings)
+		public ExecuteJobCommand(HordeClientFactory hordeClientFactory, IEnumerable<IJobExecutorFactory> jobExecutorFactories, IOptions<DriverSettings> driverSettings)
 		{
 			_hordeClientFactory = hordeClientFactory;
 			_jobExecutorFactories = jobExecutorFactories;
@@ -47,7 +47,7 @@ namespace JobDriver.Commands.Execution
 
 			ExecuteJobTask executeTask = ExecuteJobTask.Parser.ParseFrom(Convert.FromBase64String(Task));
 
-			IHordeClient hordeClient = _hordeClientFactory.Create(executeTask.Token);
+			HordeClient hordeClient = _hordeClientFactory.Create(accessToken: executeTask.Token);
 
 			await JobExecutorHelpers.ExecuteAsync(hordeClient, WorkingDir, LeaseId, executeTask, _jobExecutorFactories, _driverSettings.Value, logger, CancellationToken.None);
 			return 0;
diff --git a/Engine/Source/Programs/Horde/HordeAgent/Services/SessionFactory.cs b/Engine/Source/Programs/Horde/HordeAgent/Services/SessionFactory.cs
--- a/Engine/Source/Programs/Horde/HordeAgent/Services/SessionFactory.cs
+++ b/Engine/Source/Programs/Horde/HordeAgent/Services/SessionFactory.cs
@@ -41,7 +41,7 @@ namespace HordeAgent.Services
 		/// <summary>
 		/// Horde client instance
 		/// </summary>
-		IHordeClient HordeClient { get; }
+		HordeClient HordeClient { get; }
 	}
 
 	/// <summary>
@@ -72,12 +72,12 @@ namespace HordeAgent.Services
 		public DirectoryReference WorkingDir { get; }
 
 		/// <inheritdoc/>
-		public IHordeClient HordeClient { get; }
+		public HordeClient HordeClient { get; }
 
 		/// <summary>
 		/// Constructor
 		/// </summary>
-		public Session(AgentId agentId, SessionId sessionId, DirectoryReference workingDir, IHordeClient hordeClient)
+		public Session(AgentId agentId, SessionId sessionId, DirectoryReference workingDir, HordeClient hordeClient)
 		{
 			AgentId = agentId;
 			SessionId = sessionId;
@@ -95,7 +95,7 @@ namespace HordeAgent.Services
 		/// <summary>
 		/// Creates a new agent session
 		/// </summary>
-		public static async Task<Session> CreateAsync(CapabilitiesService capabilitiesService, GrpcService grpcService, StatusService statusService, IOptions<AgentSettings> settings, IHordeClientFactory hordeClientFactory, ILogger logger, CancellationToken cancellationToken)
+		public static async Task<Session> CreateAsync(CapabilitiesService capabilitiesService, GrpcService grpcService, StatusService statusService, IOptions<AgentSettings> settings, HordeClientFactory hordeClientFactory, ILogger logger, CancellationToken cancellationToken)
 		{
 			AgentSettings currentSettings = settings.Value;
 
@@ -172,7 +172,7 @@ namespace HordeAgent.Services
 			statusService.Set(AgentStatusMessage.ConnectingToServer);
 
 			RpcCreateSessionResponse createSessionResponse;
-			await using (IHordeClient sessionClient = hordeClientFactory.Create(registrationInfo.Token))
+			await using (HordeClient sessionClient = hordeClientFactory.Create(accessToken: registrationInfo.Token))
 			{
 				HordeRpc.HordeRpcClient rpcClient = await sessionClient.CreateGrpcClientAsync<HordeRpc.HordeRpcClient>(cancellationToken);
 
@@ -202,7 +202,7 @@ namespace HordeAgent.Services
 
 			// Open a connection to the server
 #pragma warning disable CA2000 // False positive; ownership is transferred to new Session object.
-			IHordeClient client = hordeClientFactory.Create(createSessionResponse.Token);
+			HordeClient client = hordeClientFactory.Create(accessToken: createSessionResponse.Token);
 			return new Session(new AgentId(createSessionResponse.AgentId), SessionId.Parse(createSessionResponse.SessionId), workingDir, client);
 #pragma warning restore CA2000
 		}
@@ -332,14 +332,14 @@ namespace HordeAgent.Services
 		readonly CapabilitiesService _capabilitiesService;
 		readonly GrpcService _grpcService;
 		readonly StatusService _statusService;
-		readonly IHordeClientFactory _hordeClientFactory;
+		readonly HordeClientFactory _hordeClientFactory;
 		readonly IOptions<AgentSettings> _agentSettings;
 		readonly ILogger _logger;
 
 		/// <summary>
 		/// Constructor
 		/// </summary>
-		public SessionFactory(CapabilitiesService capabilitiesService, GrpcService grpcService, StatusService statusService, IHordeClientFactory hordeClientFactory, IOptions<AgentSettings> agentSettings, ILogger<SessionFactory> logger)
+		public SessionFactory(CapabilitiesService capabilitiesService, GrpcService grpcService, StatusService statusService, HordeClientFactory hordeClientFactory, IOptions<AgentSettings> agentSettings, ILogger<SessionFactory> logger)
 		{
 			_capabilitiesService = capabilitiesService;
 			_grpcService = grpcService;
diff --git a/Engine/Source/Programs/Shared/EpicGames.Horde/HordeClient.cs b/Engine/Source/Programs/Shared/EpicGames.Horde/HordeClient.cs
--- a/Engine/Source/Programs/Shared/EpicGames.Horde/HordeClient.cs
+++ b/Engine/Source/Programs/Shared/EpicGames.Horde/HordeClient.cs
@@ -36,7 +36,7 @@ namespace EpicGames.Horde
 	/// <summary>
 	/// Default implementation of <see cref="IHordeClient"/>
 	/// </summary>
-	abstract class HordeClient : IHordeClient
+	public abstract class HordeClient : IHordeClient
 	{
 		readonly Uri _serverUrl;
 		readonly BundleCache _bundleCache;
@@ -251,6 +251,7 @@ namespace EpicGames.Horde
 			return GrpcChannel.ForAddress(serverUri, channelOptions);
 		}
 
+		/// <inheritdoc/>
 		public async Task<TClient> CreateGrpcClientAsync<TClient>(CancellationToken cancellationToken = default)
 			where TClient : ClientBase<TClient>
 		{
@@ -262,6 +263,7 @@ namespace EpicGames.Horde
 		public HordeHttpClient CreateHttpClient()
 			=> new HordeHttpClient(CreateAuthenticatedHttpClient());
 
+		/// <inheritdoc/>
 		public IComputeClient CreateComputeClient()
 		{
 			string? sessionId = null;
@@ -442,14 +444,17 @@ namespace EpicGames.Horde
 	}
 
 	/// <summary>
-	/// Implementation of <see cref="IHordeClientFactory"/>
+	/// Allows creating <see cref="HordeClient"/> instances
 	/// </summary>
-	class HordeClientFactory : IHordeClientFactory
+	public class HordeClientFactory
 	{
 		readonly BundleCache _bundleCache;
 		readonly IOptionsSnapshot<HordeOptions> _hordeOptions;
 		readonly ILoggerFactory _loggerFactory;
 
+		/// <summary>
+		/// Constructor 
+		/// </summary>
 		public HordeClientFactory(BundleCache bundleCache, IOptionsSnapshot<HordeOptions> hordeOptions, ILoggerFactory loggerFactory)
 		{
 			_bundleCache = bundleCache;
@@ -457,7 +462,7 @@ namespace EpicGames.Horde
 			_loggerFactory = loggerFactory;
 		}
 
-		Uri GetServerUrl()
+		Uri GetDefaultServerUrl()
 		{
 			Uri? serverUrl = _hordeOptions.Value.GetServerUrlOrDefault();
 			if (serverUrl == null)
@@ -467,18 +472,20 @@ namespace EpicGames.Horde
 			return serverUrl;
 		}
 
-		/// <inheritdoc/>
-		public IHordeClient Create()
-		{
-			Uri serverUrl = GetServerUrl();
-			return new HordeClientWithDynamicCredentials(serverUrl, _bundleCache, _hordeOptions, _loggerFactory);
-		}
-
-		/// <inheritdoc/>
-		public IHordeClient Create(string? accessToken)
+		/// <summary>
+		/// Create a client using the user's default access token
+		/// </summary>
+		public HordeClient Create(Uri? serverUrl = null, string? accessToken = null)
 		{
-			Uri serverUrl = GetServerUrl();
-			return new HordeClientWithStaticCredentials(serverUrl, accessToken, _bundleCache, _hordeOptions, _loggerFactory);
+			serverUrl ??= GetDefaultServerUrl();
+			if (accessToken == null)
+			{
+				return new HordeClientWithDynamicCredentials(serverUrl, _bundleCache, _hordeOptions, _loggerFactory);
+			}
+			else
+			{
+				return new HordeClientWithStaticCredentials(serverUrl, accessToken, _bundleCache, _hordeOptions, _loggerFactory);
+			}
 		}
 	}
 }
diff --git a/Engine/Source/Programs/Shared/EpicGames.Horde/HordeExtensions.cs b/Engine/Source/Programs/Shared/EpicGames.Horde/HordeExtensions.cs
--- a/Engine/Source/Programs/Shared/EpicGames.Horde/HordeExtensions.cs
+++ b/Engine/Source/Programs/Shared/EpicGames.Horde/HordeExtensions.cs
@@ -36,8 +36,8 @@ namespace EpicGames.Horde
 			serviceCollection.AddSingleton<StorageBackendCache>(CreateBackendCache);
 			serviceCollection.AddSingleton<HttpStorageBackendFactory>();
 			serviceCollection.AddSingleton<HttpStorageClientFactory>();
-			serviceCollection.AddSingleton<IHordeClient>(sp => sp.GetRequiredService<IHordeClientFactory>().Create());
-			serviceCollection.AddSingleton<IHordeClientFactory, HordeClientFactory>();
+			serviceCollection.AddSingleton<IHordeClient>(sp => sp.GetRequiredService<HordeClientFactory>().Create());
+			serviceCollection.AddSingleton<HordeClientFactory>();
 		}
 
 		/// <summary>
diff --git a/Engine/Source/Programs/Shared/EpicGames.Horde/IHordeClient.cs b/Engine/Source/Programs/Shared/EpicGames.Horde/IHordeClient.cs
--- a/Engine/Source/Programs/Shared/EpicGames.Horde/IHordeClient.cs
+++ b/Engine/Source/Programs/Shared/EpicGames.Horde/IHordeClient.cs
@@ -19,7 +19,7 @@ namespace EpicGames.Horde
 	/// <summary>
 	/// Base interface for Horde functionality.
 	/// </summary>
-	public interface IHordeClient : IAsyncDisposable
+	public interface IHordeClient
 	{
 		/// <summary>
 		/// Base URL of the horde server
@@ -64,11 +64,6 @@ namespace EpicGames.Horde
 		/// </summary>
 		Task<string?> GetAccessTokenAsync(bool interactive, CancellationToken cancellationToken = default);
 
-		/// <summary>
-		/// Gets a grpc channel for communication with the server. This should NOT be disposed by the caller.
-		/// </summary>
-		Task<GrpcChannel> CreateGrpcChannelAsync(CancellationToken cancellationToken = default);
-
 		/// <summary>
 		/// Gets a gRPC client interface
 		/// </summary>
@@ -95,22 +90,6 @@ namespace EpicGames.Horde
 		IServerLogger CreateServerLogger(LogId logId, LogLevel minimumLevel = LogLevel.Information);
 	}
 
-	/// <summary>
-	/// Interface to allow creating custom horde client instances. To obtain a default horde client, get an IHordeClient instance via dependency injection.
-	/// </summary>
-	public interface IHordeClientFactory
-	{
-		/// <summary>
-		/// Create a client using the user's default access token
-		/// </summary>
-		IHordeClient Create();
-
-		/// <summary>
-		/// Create a client with an explicit access token
-		/// </summary>
-		IHordeClient Create(string accessToken);
-	}
-
 	/// <summary>
 	/// Extension methods for <see cref="IHordeClient"/>
 	/// </summary>
