Horde: Remove IAsyncDisposable interface from IComputeClient. This is a concrete implementation detail that should not be exposed to the interface.

#rnx

diff --git a/Engine/Source/Programs/Horde/Drivers/JobDriver.Tests/Execution/WorkspaceExecutorTest.cs b/Engine/Source/Programs/Horde/Drivers/JobDriver.Tests/Execution/WorkspaceExecutorTest.cs
--- a/Engine/Source/Programs/Horde/Drivers/JobDriver.Tests/Execution/WorkspaceExecutorTest.cs
+++ b/Engine/Source/Programs/Horde/Drivers/JobDriver.Tests/Execution/WorkspaceExecutorTest.cs
@@ -87,6 +87,7 @@ namespace JobDriver.Tests.Execution
 				=> new Uri("http://fake-horde-server");
 
 			public IArtifactCollection Artifacts => throw new NotImplementedException();
+			public IComputeClient Compute => throw new NotImplementedException();
 			public IProjectCollection Projects => throw new NotImplementedException();
 			public ISecretCollection Secrets => throw new NotImplementedException();
 			public IToolCollection Tools => throw new NotImplementedException();
diff --git a/Engine/Source/Programs/Horde/HordeCmd/Commands/Compute/ComputeRun.cs b/Engine/Source/Programs/Horde/HordeCmd/Commands/Compute/ComputeRun.cs
--- a/Engine/Source/Programs/Horde/HordeCmd/Commands/Compute/ComputeRun.cs
+++ b/Engine/Source/Programs/Horde/HordeCmd/Commands/Compute/ComputeRun.cs
@@ -78,8 +78,25 @@ namespace Horde.Commands.Compute
 		/// <inheritdoc/>
 		public override async Task<int> ExecuteAsync(ILogger logger)
 		{
-			await using IComputeClient client = CreateComputeClient(logger);
+			if (Local)
+			{
+				await using LocalComputeClient client = new LocalComputeClient(2000, SandboxDir, InProc, logger);
+				return await ExecuteAsync(client, logger);
+			}
+			else if (Loopback)
+			{
+				AgentComputeClient client = new AgentComputeClient(Assembly.GetExecutingAssembly().Location, 2000, logger);
+				return await ExecuteAsync(client, logger);
+			}
+			else
+			{
+				IComputeClient client = _hordeClient.Compute;
+				return await ExecuteAsync(client, logger);
+			}
+		}
 
+		async Task<int> ExecuteAsync(IComputeClient client, ILogger logger)
+		{
 			Requirements? requirements = null;
 			if (Requirements != null)
 			{
@@ -104,22 +121,6 @@ namespace Horde.Commands.Compute
 			return result ? 0 : 1;
 		}
 
-		IComputeClient CreateComputeClient(ILogger logger)
-		{
-			if (Local)
-			{
-				return new LocalComputeClient(2000, SandboxDir, InProc, logger);
-			}
-			else if (Loopback)
-			{
-				return new AgentComputeClient(Assembly.GetExecutingAssembly().Location, 2000, logger);
-			}
-			else
-			{
-				return _hordeClient.CreateComputeClient();
-			}
-		}
-
 		/// <inheritdoc/>
 		async Task<bool> HandleRequestAsync(IComputeLease lease, ILogger logger, CancellationToken cancellationToken)
 		{
diff --git a/Engine/Source/Programs/Shared/EpicGames.Horde.Tests/ServerLogPacketBuilderTests.cs b/Engine/Source/Programs/Shared/EpicGames.Horde.Tests/ServerLogPacketBuilderTests.cs
--- a/Engine/Source/Programs/Shared/EpicGames.Horde.Tests/ServerLogPacketBuilderTests.cs
+++ b/Engine/Source/Programs/Shared/EpicGames.Horde.Tests/ServerLogPacketBuilderTests.cs
@@ -219,6 +219,7 @@ namespace EpicGames.Horde.Tests
 			public Uri ServerUrl => throw new NotImplementedException();
 
 			public IArtifactCollection Artifacts => throw new NotImplementedException();
+			public IComputeClient Compute => throw new NotImplementedException();
 			public IProjectCollection Projects => throw new NotImplementedException();
 			public ISecretCollection Secrets => throw new NotImplementedException();
 			public IToolCollection Tools => throw new NotImplementedException();
@@ -229,9 +230,6 @@ namespace EpicGames.Horde.Tests
 			public HordeHttpClient CreateHttpClient()
 				=> throw new NotImplementedException();
 
-			public IComputeClient CreateComputeClient()
-				=> throw new NotImplementedException();
-
 			public IStorageClient CreateStorageClient(string relativePath, string? accessToken = null)
 			{
 				BundleStorageClient? storageClient;
diff --git a/Engine/Source/Programs/Shared/EpicGames.Horde/Compute/Clients/AgentComputeClient.cs b/Engine/Source/Programs/Shared/EpicGames.Horde/Compute/Clients/AgentComputeClient.cs
--- a/Engine/Source/Programs/Shared/EpicGames.Horde/Compute/Clients/AgentComputeClient.cs
+++ b/Engine/Source/Programs/Shared/EpicGames.Horde/Compute/Clients/AgentComputeClient.cs
@@ -75,10 +75,7 @@ namespace EpicGames.Horde.Compute.Clients
 			_port = port;
 			_logger = logger;
 		}
-
-		/// <inheritdoc/>
-		public ValueTask DisposeAsync() => new ValueTask();
-		
+	
 		/// <inheritdoc/>
 		public Task<ClusterId> GetClusterAsync(Requirements? requirements, string? requestId, ConnectionMetadataRequest? connection, ILogger logger, CancellationToken cancellationToken = default)
 		{
diff --git a/Engine/Source/Programs/Shared/EpicGames.Horde/Compute/Clients/ServerComputeClient.cs b/Engine/Source/Programs/Shared/EpicGames.Horde/Compute/Clients/ServerComputeClient.cs
--- a/Engine/Source/Programs/Shared/EpicGames.Horde/Compute/Clients/ServerComputeClient.cs
+++ b/Engine/Source/Programs/Shared/EpicGames.Horde/Compute/Clients/ServerComputeClient.cs
@@ -107,7 +107,7 @@ namespace EpicGames.Horde.Compute.Clients
 	/// <summary>
 	/// Helper class to enlist remote resources to perform compute-intensive tasks.
 	/// </summary>
-	public sealed class ServerComputeClient : IComputeClient
+	public sealed class ServerComputeClient : IComputeClient, IDisposable
 	{
 		/// <summary>
 		/// Length of the nonce sent as part of handshaking between initiator and remote
@@ -206,13 +206,6 @@ namespace EpicGames.Horde.Compute.Clients
 			_externalIpResolver = new ExternalIpResolver(_httpClient);
 		}
 
-		/// <inheritdoc/>
-		public ValueTask DisposeAsync()
-		{
-			Dispose();
-			return new ValueTask();
-		}
-
 		/// <inheritdoc/>
 		public void Dispose()
 		{
diff --git a/Engine/Source/Programs/Shared/EpicGames.Horde/Compute/IComputeClient.cs b/Engine/Source/Programs/Shared/EpicGames.Horde/Compute/IComputeClient.cs
--- a/Engine/Source/Programs/Shared/EpicGames.Horde/Compute/IComputeClient.cs
+++ b/Engine/Source/Programs/Shared/EpicGames.Horde/Compute/IComputeClient.cs
@@ -11,7 +11,7 @@ namespace EpicGames.Horde.Compute
 	/// <summary>
 	/// Interface for uploading compute work to remote machines
 	/// </summary>
-	public interface IComputeClient : IAsyncDisposable
+	public interface IComputeClient
 	{
 		/// <summary>
 		/// Find the most suitable cluster to execute a given compute assignment request
diff --git a/Engine/Source/Programs/Shared/EpicGames.Horde/HordeClient.cs b/Engine/Source/Programs/Shared/EpicGames.Horde/HordeClient.cs
--- a/Engine/Source/Programs/Shared/EpicGames.Horde/HordeClient.cs
+++ b/Engine/Source/Programs/Shared/EpicGames.Horde/HordeClient.cs
@@ -44,6 +44,7 @@ namespace EpicGames.Horde
 		readonly ILoggerFactory _loggerFactory;
 		readonly ILogger _logger;
 
+		readonly ServerComputeClient _serverComputeClient;
 		BackgroundTask<GrpcChannel>? _grpcChannel;
 
 		/// <inheritdoc/>
@@ -52,6 +53,9 @@ namespace EpicGames.Horde
 		/// <inheritdoc/>
 		public IArtifactCollection Artifacts { get; }
 
+		/// <inheritdoc/>
+		public IComputeClient Compute => _serverComputeClient;
+
 		/// <inheritdoc/>
 		public IProjectCollection Projects { get; }
 
@@ -72,6 +76,8 @@ namespace EpicGames.Horde
 			_loggerFactory = loggerFactory;
 			_logger = _loggerFactory.CreateLogger<HordeClient>();
 
+			_serverComputeClient = CreateComputeClient();
+
 			Artifacts = new ArtifactHttpCollection(this);
 			Projects = new ProjectHttpCollection(this);
 			Secrets = new SecretHttpCollection(this);
@@ -81,6 +87,8 @@ namespace EpicGames.Horde
 		/// <inheritdoc/>
 		public virtual async ValueTask DisposeAsync()
 		{
+			_serverComputeClient.Dispose();
+
 			if (_grpcChannel != null)
 			{
 				await _grpcChannel.DisposeAsync();
@@ -264,7 +272,7 @@ namespace EpicGames.Horde
 			=> new HordeHttpClient(CreateAuthenticatedHttpClient());
 
 		/// <inheritdoc/>
-		public IComputeClient CreateComputeClient()
+		ServerComputeClient CreateComputeClient()
 		{
 			string? sessionId = null;
 			string? jobId = Environment.GetEnvironmentVariable("UE_HORDE_JOBID");
diff --git a/Engine/Source/Programs/Shared/EpicGames.Horde/IHordeClient.cs b/Engine/Source/Programs/Shared/EpicGames.Horde/IHordeClient.cs
--- a/Engine/Source/Programs/Shared/EpicGames.Horde/IHordeClient.cs
+++ b/Engine/Source/Programs/Shared/EpicGames.Horde/IHordeClient.cs
@@ -31,6 +31,11 @@ namespace EpicGames.Horde
 		/// </summary>
 		IArtifactCollection Artifacts { get; }
 
+		/// <summary>
+		/// Accessor for the compute client
+		/// </summary>
+		IComputeClient Compute { get; }
+
 		/// <summary>
 		/// Accessor for the project collection
 		/// </summary>
@@ -74,11 +79,6 @@ namespace EpicGames.Horde
 		/// </summary>
 		HordeHttpClient CreateHttpClient();
 
-		/// <summary>
-		/// Creates a compute client
-		/// </summary>
-		IComputeClient CreateComputeClient();
-
 		/// <summary>
 		/// Creates a storage client for the given base path
 		/// </summary>
diff --git a/Engine/Source/Programs/UnrealBuildTool/Executors/UnrealBuildAccelerator/UBAAgentCoordinatorHorde.cs b/Engine/Source/Programs/UnrealBuildTool/Executors/UnrealBuildAccelerator/UBAAgentCoordinatorHorde.cs
--- a/Engine/Source/Programs/UnrealBuildTool/Executors/UnrealBuildAccelerator/UBAAgentCoordinatorHorde.cs
+++ b/Engine/Source/Programs/UnrealBuildTool/Executors/UnrealBuildAccelerator/UBAAgentCoordinatorHorde.cs
@@ -125,7 +125,6 @@ namespace UnrealBuildTool
 				_client = null;
 				// Set CPU resource need to zero (will also expire on the server if not updated)
 				await UpdateCpuCoreNeedAsync(0);
-				await client.DisposeAsync();
 			}
 
 			await _serviceProvider.DisposeAsync();
@@ -139,7 +138,7 @@ namespace UnrealBuildTool
 				throw new InvalidOperationException("Session has already been initialized");
 			}
 
-			_client = _serviceProvider.GetRequiredService<IHordeClient>().CreateComputeClient();
+			_client = _serviceProvider.GetRequiredService<IHordeClient>().Compute;
 			_hordeHttpClient = _serviceProvider.GetRequiredService<IHordeClient>().CreateHttpClient();
 			GetServerInfoResponse serverInfo = await _hordeHttpClient.GetServerInfoAsync(cancellationToken);
 			_logger.LogInformation("Horde server: {ServerVersion}, agent: {AgentVersion}", serverInfo.ServerVersion, serverInfo.AgentVersion);
