Horde: Fix static analysis warning.

#rnx

diff --git a/Engine/Source/Programs/Horde/Drivers/JobDriver.Tests/WorkerServiceTest.cs b/Engine/Source/Programs/Horde/Drivers/JobDriver.Tests/WorkerServiceTest.cs
--- a/Engine/Source/Programs/Horde/Drivers/JobDriver.Tests/WorkerServiceTest.cs
+++ b/Engine/Source/Programs/Horde/Drivers/JobDriver.Tests/WorkerServiceTest.cs
@@ -9,12 +9,11 @@ using EpicGames.Horde.Logs;
 using EpicGames.Horde.Storage;
 using EpicGames.Horde.Storage.Backends;
 using EpicGames.Horde.Storage.Bundles;
-using EpicGames.Horde.Storage.Clients;
 using Grpc.Core;
-using JobDriver.Execution;
 using Horde.Common.Rpc;
 using HordeCommon.Rpc.Messages;
 using HordeCommon.Rpc.Tasks;
+using JobDriver.Execution;
 using Microsoft.Extensions.DependencyInjection;
 using Microsoft.Extensions.Logging;
 using Microsoft.Extensions.Logging.Abstractions;
