Horde: Ignore server-defined properties when sent by agent

diff --git a/Engine/Source/Programs/Horde/Plugins/Build/HordeServer.Build.Tests/Agents/AgentServiceTest.cs b/Engine/Source/Programs/Horde/Plugins/Build/HordeServer.Build.Tests/Agents/AgentServiceTest.cs
--- a/Engine/Source/Programs/Horde/Plugins/Build/HordeServer.Build.Tests/Agents/AgentServiceTest.cs
+++ b/Engine/Source/Programs/Horde/Plugins/Build/HordeServer.Build.Tests/Agents/AgentServiceTest.cs
@@ -153,6 +153,19 @@ public class AgentServiceTest : BuildTestSetup
 		double? rate2 = await AgentService.GetRateAsync(agent2.Id);
 		Assert.AreEqual(300, rate2!.Value, 0.1);
 	}
+	
+	[TestMethod]
+	public async Task Properties_ServerDefinedAsync()
+	{
+		IAgent newAgent = await AgentService.CreateAgentAsync("agent1", false, "");
+		IAgent agent = await AgentService.CreateSessionAsync(newAgent, new RpcAgentCapabilities(["foo=bar", $"{KnownPropertyNames.Trusted}=true", $"{KnownPropertyNames.Trusted}"]), "test");
+		
+		Assert.AreEqual(0, newAgent.Properties.Count);
+		Assert.AreEqual(1, agent.Properties.Count);
+		Assert.AreEqual("foo=bar", agent.Properties[0]);
+		Assert.IsTrue(AgentExtensions.IsPropertyServerDefined(KnownPropertyNames.Trusted));
+		Assert.IsFalse(AgentExtensions.IsPropertyServerDefined(KnownPropertyNames.OsFamily));
+	}
 
 	[TestMethod]
 	public async Task EphemeralTestAsync()
diff --git a/Engine/Source/Programs/Horde/Plugins/Compute/HordeServer.Compute/Agents/AgentService.cs b/Engine/Source/Programs/Horde/Plugins/Compute/HordeServer.Compute/Agents/AgentService.cs
--- a/Engine/Source/Programs/Horde/Plugins/Compute/HordeServer.Compute/Agents/AgentService.cs
+++ b/Engine/Source/Programs/Horde/Plugins/Compute/HordeServer.Compute/Agents/AgentService.cs
@@ -352,9 +352,10 @@ namespace HordeServer.Agents
 				{
 					// Get the new pools for the agent
 					List<PoolId> dynamicPools = await GetDynamicPoolsAsync(agent, cancellationToken);
-
+					RpcAgentCapabilities normalizedCaps = NormalizeCapabilities(agent, capabilities);
+					
 					// Reset the agent to use the new session
-					newAgent = await agent.TryCreateSessionAsync(new CreateSessionOptions(capabilities, dynamicPools, version), cancellationToken);
+					newAgent = await agent.TryCreateSessionAsync(new CreateSessionOptions(normalizedCaps, dynamicPools, version), cancellationToken);
 					if (newAgent != null)
 					{
 						LogPropertyChanges(agentLogger, agent.Properties, newAgent.Properties);
@@ -374,6 +375,21 @@ namespace HordeServer.Agents
 			}
 			return agent;
 		}
+		
+		/// <summary>
+		/// Filters untrusted resource and property keys from capabilities
+		/// Also injects potential server-controlled properties
+		/// </summary>
+		/// <param name="agent">Agent</param>
+		/// <param name="capabilities">Capabilities as reported by agent</param>
+		/// <returns>A new and updated RpcAgentCapabilities object</returns>
+		private static RpcAgentCapabilities NormalizeCapabilities(IAgent agent, RpcAgentCapabilities capabilities)
+		{
+			RpcAgentCapabilities caps = new (capabilities);
+			caps.Properties.Clear();
+			caps.Properties.AddRange(AgentExtensions.RemoveServerDefinedProperties(capabilities.Properties));
+			return caps;
+		}
 
 		/// <summary>
 		/// Determines whether a task source can currently issue tasks
diff --git a/Engine/Source/Programs/Horde/Plugins/Compute/HordeServer.Compute/Agents/IAgent.cs b/Engine/Source/Programs/Horde/Plugins/Compute/HordeServer.Compute/Agents/IAgent.cs
--- a/Engine/Source/Programs/Horde/Plugins/Compute/HordeServer.Compute/Agents/IAgent.cs
+++ b/Engine/Source/Programs/Horde/Plugins/Compute/HordeServer.Compute/Agents/IAgent.cs
@@ -342,6 +342,11 @@ namespace HordeServer.Agents
 		/// Tool ID for Mac-specific and self-contained agent software
 		/// </summary>
 		public static ToolId AgentMacX64ToolId { get; } = new("horde-agent-osx-x64");
+		
+		/// <summary>
+		/// Set of property names that can only be set by the server
+		/// </summary>
+		public static readonly HashSet<string> ServerDefinedPropertyNames = [KnownPropertyNames.Trusted];
 
 		/// <summary>
 		/// Gets the tool ID for the software the given agent should be running
@@ -451,7 +456,28 @@ namespace HordeServer.Agents
 		{
 			return agent.Properties.BinarySearch(property, StringComparer.OrdinalIgnoreCase) >= 0;
 		}
-
+		
+		/// <summary>
+		/// Check if a property name (incl value) is controlled by the server
+		/// Such properties are protected and cannot be set by the agent.
+		/// </summary>
+		/// <param name="propName">Name of the property</param>
+		/// <returns>True if controlled and enforced by the server</returns>
+		public static bool IsPropertyServerDefined(string propName)
+		{
+			return ServerDefinedPropertyNames.Contains(propName);
+		}
+		
+		/// <summary>
+		/// Filters a list of properties from server-defined properties
+		/// </summary>
+		/// <param name="properties">List to filter</param>
+		/// <returns>A new list with any server-defined property removed</returns>
+		public static IReadOnlyList<string> RemoveServerDefinedProperties(IReadOnlyList<string> properties)
+		{
+			return properties.Where(property => !IsPropertyServerDefined(property.Split('=', 2)[0])).ToList();
+		}
+		
 		/// <summary>
 		/// Finds property values from a sorted list of Name=Value pairs
 		/// </summary>
diff --git a/Engine/Source/Programs/Shared/EpicGames.Horde/Agents/KnownPropertyNames.cs b/Engine/Source/Programs/Shared/EpicGames.Horde/Agents/KnownPropertyNames.cs
--- a/Engine/Source/Programs/Shared/EpicGames.Horde/Agents/KnownPropertyNames.cs
+++ b/Engine/Source/Programs/Shared/EpicGames.Horde/Agents/KnownPropertyNames.cs
@@ -81,4 +81,9 @@ public static class KnownPropertyNames
 	/// Whether the Wine compatibility layer is enabled (for running Windows applications on Linux)
 	/// </summary>
 	public const string WineEnabled = "WineEnabled";
+	
+	/// <summary>
+	/// Whether the agent is trusted
+	/// </summary>
+	public const string Trusted = "Trusted";
 }
\ No newline at end of file
