Horde: Add server-defined properties for an agent

- Allows setting properties that will overwrite and merge with properties reported by agent itself
- Refactor parameters for agent creation into an options object
- When a user/agent can create new agents outside enrollment process, it's marked as trusted for the time being (grandfathering existing registration in JobRpcService)

diff --git a/Engine/Source/Programs/Horde/Plugins/Build/HordeServer.Build.Tests/Agents/AgentServiceTest.cs b/Engine/Source/Programs/Horde/Plugins/Build/HordeServer.Build.Tests/Agents/AgentServiceTest.cs
--- a/Engine/Source/Programs/Horde/Plugins/Build/HordeServer.Build.Tests/Agents/AgentServiceTest.cs
+++ b/Engine/Source/Programs/Horde/Plugins/Build/HordeServer.Build.Tests/Agents/AgentServiceTest.cs
@@ -8,7 +8,6 @@ using EpicGames.Horde.Jobs;
 using HordeCommon.Rpc.Messages;
 using HordeServer.Agents;
 using HordeServer.Auditing;
-using HordeServer.Jobs;
 using HordeServer.Utilities;
 using Microsoft.AspNetCore.Mvc;
 using Microsoft.Extensions.DependencyInjection;
@@ -16,7 +15,7 @@ using Microsoft.Extensions.DependencyInjection;
 namespace HordeServer.Tests.Agents;
 
 /// <summary>
-///     Testing the agent service
+/// Testing the agent service
 /// </summary>
 [TestClass]
 public class AgentServiceTest : BuildTestSetup
@@ -63,7 +62,7 @@ public class AgentServiceTest : BuildTestSetup
 	public async Task LastOnlineChangeDuringSessionCreateAsync()
 	{
 		// No session created yet, status change timestamp is empty
-		IAgent agent = await AgentService.CreateAgentAsync("agent1", false, "");
+		IAgent agent = await AgentService.CreateAgentAsync(new CreateAgentOptions(new AgentId("agent1"), false, ""));
 		Assert.AreEqual(AgentStatus.Stopped, agent.Status);
 		Assert.IsTrue(agent.LastOnlineTime.HasValue);
 
@@ -78,7 +77,7 @@ public class AgentServiceTest : BuildTestSetup
 	private static int s_agentId = 1;
 	private async Task<IAgent> CreateAgentSessionAsync()
 	{
-		IAgent agent = await AgentService.CreateAgentAsync("agentServiceTest-" + s_agentId++, false, "");
+		IAgent agent = await AgentService.CreateAgentAsync(new CreateAgentOptions(new AgentId("agentServiceTest-" + s_agentId++), false, ""));
 		agent = await AgentService.CreateSessionAsync(agent, new RpcAgentCapabilities(), "v1");
 		return agent;
 	}
@@ -135,8 +134,8 @@ public class AgentServiceTest : BuildTestSetup
 	[TestMethod]
 	public async Task GetAgentRateTestAsync()
 	{
-		IAgent agent1 = await AgentService.CreateAgentAsync("agent1", false, "");
-		IAgent agent2 = await AgentService.CreateAgentAsync("agent2", false, "");
+		IAgent agent1 = await AgentService.CreateAgentAsync(new CreateAgentOptions(new AgentId("agent1"), false, ""));
+		IAgent agent2 = await AgentService.CreateAgentAsync(new CreateAgentOptions(new AgentId("agent2"), false, ""));
 		await AgentService.CreateSessionAsync(agent1, new RpcAgentCapabilities(new List<string>() { "aws-instance-type=c5.24xlarge", "osfamily=windows" }), "test");
 		await AgentService.CreateSessionAsync(agent2, new RpcAgentCapabilities(new List<string>() { "aws-instance-type=c4.4xLARge", "osfamily=WinDowS" }), "test");
 
@@ -155,9 +154,9 @@ public class AgentServiceTest : BuildTestSetup
 	}
 	
 	[TestMethod]
-	public async Task Properties_ServerDefinedAsync()
+	public async Task Properties_ServerDefined_AgentCannotOverride_Async()
 	{
-		IAgent newAgent = await AgentService.CreateAgentAsync("agent1", false, "");
+		IAgent newAgent = await AgentService.CreateAgentAsync(new CreateAgentOptions(new AgentId("agent1"), false, ""));
 		IAgent agent = await AgentService.CreateSessionAsync(newAgent, new RpcAgentCapabilities(["foo=bar", $"{KnownPropertyNames.Trusted}=true", $"{KnownPropertyNames.Trusted}"]), "test");
 		
 		Assert.AreEqual(0, newAgent.Properties.Count);
@@ -166,6 +165,20 @@ public class AgentServiceTest : BuildTestSetup
 		Assert.IsTrue(AgentExtensions.IsPropertyServerDefined(KnownPropertyNames.Trusted));
 		Assert.IsFalse(AgentExtensions.IsPropertyServerDefined(KnownPropertyNames.OsFamily));
 	}
+	
+	[TestMethod]
+	public async Task Properties_ServerDefined_PropagatedToAgent_Async()
+	{
+		CreateAgentOptions options = new(new AgentId("agent1"), false, "", ["server=foo"]);
+		IAgent newAgent = await AgentService.CreateAgentAsync(options);
+		IAgent agent = await AgentService.CreateSessionAsync(newAgent, new RpcAgentCapabilities(["agent=bar"]), "test");
+		
+		Assert.AreEqual(1, newAgent.ServerDefinedProperties.Count);
+		Assert.AreEqual(0, newAgent.Properties.Count);
+		Assert.AreEqual(2, agent.Properties.Count);
+		CollectionAssert.Contains(agent.Properties.ToList(), "server=foo");
+		CollectionAssert.Contains(agent.Properties.ToList(), "agent=bar");
+	}
 
 	[TestMethod]
 	public async Task EphemeralTestAsync()
diff --git a/Engine/Source/Programs/Horde/Plugins/Build/HordeServer.Build.Tests/BuildTestSetup.cs b/Engine/Source/Programs/Horde/Plugins/Build/HordeServer.Build.Tests/BuildTestSetup.cs
--- a/Engine/Source/Programs/Horde/Plugins/Build/HordeServer.Build.Tests/BuildTestSetup.cs
+++ b/Engine/Source/Programs/Horde/Plugins/Build/HordeServer.Build.Tests/BuildTestSetup.cs
@@ -302,8 +302,8 @@ namespace HordeServer.Tests
 			{
 				tempProps.Add(KnownPropertyNames.AwsInstanceId + "=" + awsInstanceId);
 			}
-
-			IAgent? agent = await AgentService.CreateAgentAsync("TestAgent" + s_agentIdCounter++, ephemeral, "");
+			
+			IAgent? agent = await AgentService.CreateAgentAsync(new CreateAgentOptions(new AgentId("TestAgent" + s_agentIdCounter++), ephemeral, ""));
 			Assert.IsNotNull(agent);
 
 			agent = await agent.TryUpdateAsync(new UpdateAgentOptions { Enabled = enabled, ExplicitPools = poolId != null ? [poolId.Value] : [] });
diff --git a/Engine/Source/Programs/Horde/Plugins/Build/HordeServer.Build.Tests/Fixture.cs b/Engine/Source/Programs/Horde/Plugins/Build/HordeServer.Build.Tests/Fixture.cs
--- a/Engine/Source/Programs/Horde/Plugins/Build/HordeServer.Build.Tests/Fixture.cs
+++ b/Engine/Source/Programs/Horde/Plugins/Build/HordeServer.Build.Tests/Fixture.cs
@@ -1,6 +1,7 @@
 // Copyright Epic Games, Inc. All Rights Reserved.
 
 using EpicGames.Core;
+using EpicGames.Horde.Agents;
 using EpicGames.Horde.Agents.Pools;
 using EpicGames.Horde.Commits;
 using EpicGames.Horde.Jobs;
@@ -116,7 +117,7 @@ namespace HordeServer.Tests
 			Job2 = (await jobService.GetJobAsync(Job2.Id))!;
 
 			Agent1Name = "testAgent1";
-			Agent1 = await agentService.CreateAgentAsync(Agent1Name, false, "");
+			Agent1 = await agentService.CreateAgentAsync(new CreateAgentOptions(new AgentId(Agent1Name), false, ""));
 		}
 
 		private class FixtureGraph : IGraph
diff --git a/Engine/Source/Programs/Horde/Plugins/Build/HordeServer.Build.Tests/Jobs/JobServiceTests.cs b/Engine/Source/Programs/Horde/Plugins/Build/HordeServer.Build.Tests/Jobs/JobServiceTests.cs
--- a/Engine/Source/Programs/Horde/Plugins/Build/HordeServer.Build.Tests/Jobs/JobServiceTests.cs
+++ b/Engine/Source/Programs/Horde/Plugins/Build/HordeServer.Build.Tests/Jobs/JobServiceTests.cs
@@ -2,6 +2,7 @@
 
 using System.Text;
 using EpicGames.Core;
+using EpicGames.Horde.Agents;
 using EpicGames.Horde.Agents.Pools;
 using EpicGames.Horde.Commits;
 using EpicGames.Horde.Jobs;
@@ -18,7 +19,6 @@ using HordeServer.Projects;
 using HordeServer.Server;
 using HordeServer.Streams;
 using HordeServer.Users;
-using Microsoft.CodeAnalysis;
 
 namespace HordeServer.Tests.Jobs
 {
@@ -246,7 +246,7 @@ namespace HordeServer.Tests.Jobs
 
 			// ----
 
-			IAgent? agent = await AgentService.CreateAgentAsync("TestAgent", false, "");
+			IAgent? agent = await AgentService.CreateAgentAsync(new CreateAgentOptions(new AgentId("TestAgent"), false, ""));
 			Assert.IsNotNull(agent);
 
 			agent = await agent.TryUpdateAsync(new UpdateAgentOptions { Enabled = true, ExplicitPools = new List<PoolId> { new PoolId("win") } });
diff --git a/Engine/Source/Programs/Horde/Plugins/Build/HordeServer.Build.Tests/Jobs/JobTaskSourceTests.cs b/Engine/Source/Programs/Horde/Plugins/Build/HordeServer.Build.Tests/Jobs/JobTaskSourceTests.cs
--- a/Engine/Source/Programs/Horde/Plugins/Build/HordeServer.Build.Tests/Jobs/JobTaskSourceTests.cs
+++ b/Engine/Source/Programs/Horde/Plugins/Build/HordeServer.Build.Tests/Jobs/JobTaskSourceTests.cs
@@ -1,5 +1,6 @@
 // Copyright Epic Games, Inc. All Rights Reserved.
 
+using EpicGames.Horde.Agents;
 using EpicGames.Horde.Agents.Pools;
 using EpicGames.Horde.Commits;
 using EpicGames.Horde.Jobs;
@@ -158,7 +159,7 @@ namespace HordeServer.Tests.Jobs
 
 			if (shouldCreateAgent)
 			{
-				IAgent? agent = await AgentService.CreateAgentAsync("TestAgent", false, "");
+				IAgent? agent = await AgentService.CreateAgentAsync(new CreateAgentOptions(new AgentId("TestAgent"), false, ""));
 				Assert.IsNotNull(agent);
 
 				agent = await agent.TryUpdateAsync(new UpdateAgentOptions { Enabled = isAgentEnabled, ExplicitPools = new List<PoolId> { pool.Id } });
diff --git a/Engine/Source/Programs/Horde/Plugins/Compute/HordeServer.Compute.Tests/Agents/AgentCollectionTests.cs b/Engine/Source/Programs/Horde/Plugins/Compute/HordeServer.Compute.Tests/Agents/AgentCollectionTests.cs
--- a/Engine/Source/Programs/Horde/Plugins/Compute/HordeServer.Compute.Tests/Agents/AgentCollectionTests.cs
+++ b/Engine/Source/Programs/Horde/Plugins/Compute/HordeServer.Compute.Tests/Agents/AgentCollectionTests.cs
@@ -31,7 +31,7 @@ public class AgentCollectionTests : ComputeTestSetup
 	[TestInitialize]
 	public async Task SetupAsync()
 	{
-		_agent = await AgentCollection.AddAsync(new AgentId("test"), ephemeral: true, enrollmentKey: "");
+		_agent = await AgentCollection.AddAsync(new CreateAgentOptions(new AgentId("test"), true, ""));
 	}
 
 	[TestMethod]
diff --git a/Engine/Source/Programs/Horde/Plugins/Compute/HordeServer.Compute/Agents/AgentCollection.cs b/Engine/Source/Programs/Horde/Plugins/Compute/HordeServer.Compute/Agents/AgentCollection.cs
--- a/Engine/Source/Programs/Horde/Plugins/Compute/HordeServer.Compute/Agents/AgentCollection.cs
+++ b/Engine/Source/Programs/Horde/Plugins/Compute/HordeServer.Compute/Agents/AgentCollection.cs
@@ -58,6 +58,7 @@ namespace HordeServer.Agents
 			bool IAgent.Deleted => _document.Deleted;
 			string? IAgent.Version => _document.Version;
 			string? IAgent.Comment => _document.Comment;
+			IReadOnlyList<string> IAgent.ServerDefinedProperties => _document.ServerDefinedProperties ?? DefaultProperties;
 			IReadOnlyList<string> IAgent.Properties => _document.Properties ?? DefaultProperties;
 			IReadOnlyDictionary<string, int> IAgent.Resources => _document.Resources ?? DefaultResources;
 			string? IAgent.LastUpgradeVersion => _document.LastUpgradeVersion;
@@ -162,7 +163,8 @@ namespace HordeServer.Agents
 
 			[BsonElement("Version2")]
 			public string? Version { get; set; }
-
+			
+			public List<string>? ServerDefinedProperties { get; set; }
 			public List<string>? Properties { get; set; }
 			public Dictionary<string, int>? Resources { get; set; }
 
@@ -545,12 +547,15 @@ namespace HordeServer.Agents
 		}
 
 		/// <inheritdoc/>
-		public async Task<IAgent> AddAsync(AgentId id, bool ephemeral, string enrollmentKey, CancellationToken cancellationToken)
+		public async Task<IAgent> AddAsync(CreateAgentOptions options, CancellationToken cancellationToken)
 		{
-			AgentDocument document = new AgentDocument(id, ephemeral, enrollmentKey);
-			document.LastOnlineTime = _clock.UtcNow;
-			document.DocumentVersion = 2;
-
+			AgentDocument document = new (options.Id, options.Ephemeral, options.EnrollmentKey)
+			{
+				DocumentVersion = 2,
+				ServerDefinedProperties = options.ServerDefinedProperties?.ToList(),
+				LastOnlineTime = _clock.UtcNow,
+			};
+			
 			await _agentCollection.InsertOneAsync(document, null, cancellationToken);
 			return new Agent(this, document, null);
 		}
@@ -736,7 +741,6 @@ namespace HordeServer.Agents
 			return await _agentCollection.FindOneAndUpdateAsync<AgentDocument>(filter, updateWithIndex, new FindOneAndUpdateOptions<AgentDocument, AgentDocument> { ReturnDocument = ReturnDocument.After }, cancellationToken);
 		}
 
-		/// <inheritdoc/>
 		async Task<Agent?> TryUpdateSettingsAsync(Agent agent, UpdateAgentOptions options, CancellationToken cancellationToken)
 		{
 			// Update the database
@@ -812,7 +816,6 @@ namespace HordeServer.Agents
 			return await CreateAgentObjectAsync(newDocument, agent.Session, cancellationToken);
 		}
 
-		/// <inheritdoc/>
 		async Task<Agent?> TryUpdateSessionAsync(Agent agent, UpdateSessionOptions options, CancellationToken cancellationToken)
 		{
 			AgentId agentId = agent.Document.Id;
diff --git a/Engine/Source/Programs/Horde/Plugins/Compute/HordeServer.Compute/Agents/AgentService.cs b/Engine/Source/Programs/Horde/Plugins/Compute/HordeServer.Compute/Agents/AgentService.cs
--- a/Engine/Source/Programs/Horde/Plugins/Compute/HordeServer.Compute/Agents/AgentService.cs
+++ b/Engine/Source/Programs/Horde/Plugins/Compute/HordeServer.Compute/Agents/AgentService.cs
@@ -188,35 +188,20 @@ namespace HordeServer.Agents
 		/// <summary>
 		/// Register a new agent
 		/// </summary>
-		/// <param name="name">Name of the agent</param>
-		/// <param name="ephemeral">Whether the agent is ephemeral or not</param>
-		/// <param name="enrollmentKey">Key for enrolling the agent</param>
+		/// <param name="options">Parameters for new agent</param>
 		/// <param name="cancellationToken">Cancellation token for the operation</param>
-		/// <returns>Unique id for the agent</returns>
-		public Task<IAgent> CreateAgentAsync(string name, bool ephemeral, string enrollmentKey, CancellationToken cancellationToken = default)
-		{
-			return CreateAgentAsync(new AgentId(name), ephemeral, enrollmentKey, cancellationToken);
-		}
-
-		/// <summary>
-		/// Register a new agent
-		/// </summary>
-		/// <param name="agentId">Agent id</param>
-		/// <param name="ephemeral">Whether the agent is ephemeral or not</param>
-		/// <param name="enrollmentKey">Key for enrolling the agent</param>
-		/// <param name="cancellationToken">Cancellation token for the operation</param>
-		/// <returns>Unique id for the agent</returns>
-		public async Task<IAgent> CreateAgentAsync(AgentId agentId, bool ephemeral, string enrollmentKey, CancellationToken cancellationToken = default)
+		/// <returns>Unique ID for the agent</returns>
+		public async Task<IAgent> CreateAgentAsync(CreateAgentOptions options, CancellationToken cancellationToken = default)
 		{
 			for (; ; )
 			{
-				IAgent? agent = await Agents.GetAsync(agentId, cancellationToken);
+				IAgent? agent = await Agents.GetAsync(options.Id, cancellationToken);
 				if (agent == null)
 				{
-					return await Agents.AddAsync(agentId, ephemeral, enrollmentKey, cancellationToken);
+					return await Agents.AddAsync(options, cancellationToken);
 				}
 
-				agent = await agent.TryResetAsync(ephemeral, enrollmentKey, cancellationToken);
+				agent = await agent.TryResetAsync(options.Ephemeral, options.EnrollmentKey, cancellationToken);
 				if (agent != null)
 				{
 					return agent;
@@ -387,6 +372,7 @@ namespace HordeServer.Agents
 		{
 			RpcAgentCapabilities caps = new (capabilities);
 			caps.Properties.Clear();
+			caps.Properties.AddRange(agent.ServerDefinedProperties);
 			caps.Properties.AddRange(AgentExtensions.RemoveServerDefinedProperties(capabilities.Properties));
 			return caps;
 		}
diff --git a/Engine/Source/Programs/Horde/Plugins/Compute/HordeServer.Compute/Agents/Enrollment/EnrollmentRpc.cs b/Engine/Source/Programs/Horde/Plugins/Compute/HordeServer.Compute/Agents/Enrollment/EnrollmentRpc.cs
--- a/Engine/Source/Programs/Horde/Plugins/Compute/HordeServer.Compute/Agents/Enrollment/EnrollmentRpc.cs
+++ b/Engine/Source/Programs/Horde/Plugins/Compute/HordeServer.Compute/Agents/Enrollment/EnrollmentRpc.cs
@@ -58,7 +58,7 @@ namespace HordeServer.Agents.Enrollment
 						return;
 					}
 
-					IAgent agent = await _agentService.CreateAgentAsync(agentId, false, request.Key, cancellationSource.Token);
+					IAgent agent = await _agentService.CreateAgentAsync(new CreateAgentOptions(agentId, false, request.Key), cancellationSource.Token);
 
 					List<AclClaimConfig> claims = new List<AclClaimConfig>();
 					claims.Add(new AclClaimConfig(HordeClaimTypes.Agent, agentId.ToString()));
diff --git a/Engine/Source/Programs/Horde/Plugins/Compute/HordeServer.Compute/Agents/IAgent.cs b/Engine/Source/Programs/Horde/Plugins/Compute/HordeServer.Compute/Agents/IAgent.cs
--- a/Engine/Source/Programs/Horde/Plugins/Compute/HordeServer.Compute/Agents/IAgent.cs
+++ b/Engine/Source/Programs/Horde/Plugins/Compute/HordeServer.Compute/Agents/IAgent.cs
@@ -78,7 +78,12 @@ namespace HordeServer.Agents
 		/// Arbitrary comment for the agent (useful for disable reasons etc)
 		/// </summary>
 		public string? Comment { get; }
-
+		
+		/// <summary>
+		/// List of server-defined properties (will overwrite and merge with agent-reported properties)
+		/// </summary>
+		public IReadOnlyList<string> ServerDefinedProperties { get; }
+		
 		/// <summary>
 		/// List of properties for this agent
 		/// </summary>
@@ -284,8 +289,17 @@ namespace HordeServer.Agents
 	/// <summary>
 	/// Information for a new lease
 	/// </summary>
-	public record class CreateLeaseOptions(LeaseId Id, LeaseId? ParentId, string Name, StreamId? StreamId, PoolId? PoolId, LogId? LogId, IReadOnlyDictionary<string, int>? Resources, bool Exclusive, IMessage Payload);
-
+	public record CreateLeaseOptions(LeaseId Id, LeaseId? ParentId, string Name, StreamId? StreamId, PoolId? PoolId, LogId? LogId, IReadOnlyDictionary<string, int>? Resources, bool Exclusive, IMessage Payload);
+	
+	/// <summary>
+	/// Options for creating an agent
+	/// </summary>
+	/// <param name="Id">ID of the agent</param>
+	/// <param name="Ephemeral">Whether the agent is ephemeral or not</param>
+	/// <param name="EnrollmentKey">Key used to identify a unique enrollment for the agent with this ID</param>
+	/// <param name="ServerDefinedProperties">Server-defined properties</param>
+	public record CreateAgentOptions(AgentId Id, bool Ephemeral, string EnrollmentKey, IReadOnlyList<string>? ServerDefinedProperties = null);
+	
 	/// <summary>
 	/// Options for updating an agent
 	/// </summary>
@@ -298,7 +312,7 @@ namespace HordeServer.Agents
 	/// <param name="ShutdownReason">The reason for shutting down agent, ex. Autoscaler/Manual/Unexpected</param>
 	/// <param name="ExplicitPools">List of pools for the agent</param>
 	/// <param name="Comment">New comment</param>
-	public record class UpdateAgentOptions(bool? Enabled = null, bool? RequestConform = null, bool? RequestFullConform = null, bool? RequestRestart = null, bool? RequestShutdown = null, bool? RequestForceRestart = null, string? ShutdownReason = null, List<PoolId>? ExplicitPools = null, string? Comment = null);
+	public record UpdateAgentOptions(bool? Enabled = null, bool? RequestConform = null, bool? RequestFullConform = null, bool? RequestRestart = null, bool? RequestShutdown = null, bool? RequestForceRestart = null, string? ShutdownReason = null, List<PoolId>? ExplicitPools = null, string? Comment = null);
 
 	/// <summary>
 	/// Options for starting a new agent session
@@ -306,7 +320,7 @@ namespace HordeServer.Agents
 	/// <param name="Capabilities">Capabilities for the agent</param>
 	/// <param name="DynamicPools">New list of dynamic pools for the agent</param>
 	/// <param name="Version">Current version of the agent software</param>
-	public record class CreateSessionOptions(RpcAgentCapabilities Capabilities, IReadOnlyList<PoolId> DynamicPools, string? Version);
+	public record CreateSessionOptions(RpcAgentCapabilities Capabilities, IReadOnlyList<PoolId> DynamicPools, string? Version);
 
 	/// <summary>
 	/// Options for updating a new agent session
@@ -315,7 +329,7 @@ namespace HordeServer.Agents
 	/// <param name="Capabilities">Capbilities for the session</param>
 	/// <param name="DynamicPools">New list of dynamic pools for the agent</param>
 	/// <param name="Leases">New set of leases</param>
-	public record class UpdateSessionOptions(AgentStatus? Status = null, RpcAgentCapabilities? Capabilities = null, IReadOnlyList<PoolId>? DynamicPools = null, IEnumerable<RpcLease>? Leases = null);
+	public record UpdateSessionOptions(AgentStatus? Status = null, RpcAgentCapabilities? Capabilities = null, IReadOnlyList<PoolId>? DynamicPools = null, IEnumerable<RpcLease>? Leases = null);
 
 	/// <summary>
 	/// Extension methods for IAgent
diff --git a/Engine/Source/Programs/Horde/Plugins/Compute/HordeServer.Compute/Agents/IAgentCollection.cs b/Engine/Source/Programs/Horde/Plugins/Compute/HordeServer.Compute/Agents/IAgentCollection.cs
--- a/Engine/Source/Programs/Horde/Plugins/Compute/HordeServer.Compute/Agents/IAgentCollection.cs
+++ b/Engine/Source/Programs/Horde/Plugins/Compute/HordeServer.Compute/Agents/IAgentCollection.cs
@@ -16,11 +16,9 @@ namespace HordeServer.Agents
 		/// <summary>
 		/// Adds a new agent with the given properties
 		/// </summary>
-		/// <param name="id">Id for the new agent</param>
-		/// <param name="ephemeral">Whether the agent is ephemeral or not</param>
-		/// <param name="enrollmentKey">Key used to identify a unique enrollment for the agent with this id</param>
+		/// <param name="options">Parameters for new agent</param>
 		/// <param name="cancellationToken">Cancellation token for the operation</param>
-		Task<IAgent> AddAsync(AgentId id, bool ephemeral, string enrollmentKey, CancellationToken cancellationToken = default);
+		Task<IAgent> AddAsync(CreateAgentOptions options, CancellationToken cancellationToken = default);
 
 		/// <summary>
 		/// Gets an agent by ID
diff --git a/Engine/Source/Programs/Horde/Plugins/Compute/HordeServer.Compute/Server/RpcService.cs b/Engine/Source/Programs/Horde/Plugins/Compute/HordeServer.Compute/Server/RpcService.cs
--- a/Engine/Source/Programs/Horde/Plugins/Compute/HordeServer.Compute/Server/RpcService.cs
+++ b/Engine/Source/Programs/Horde/Plugins/Compute/HordeServer.Compute/Server/RpcService.cs
@@ -149,6 +149,7 @@ namespace HordeServer.Server
 
 		/// <summary>
 		/// Creates a new agent
+		/// Assumes user has permission to create new agents as it skips the enrollment process
 		/// </summary>
 		/// <param name="request">Request to create a new agent</param>
 		/// <param name="context">Context for the RPC call</param>
@@ -162,7 +163,9 @@ namespace HordeServer.Server
 				throw new StructuredRpcException(StatusCode.PermissionDenied, "User is not authenticated to create new agents");
 			}
 
-			IAgent agent = await _agentService.CreateAgentAsync(new AgentId(request.Name), request.Ephemeral, "", context.CancellationToken);
+			// If the agent/user has permission to create an agent, it's assumed the agent is also trusted
+			CreateAgentOptions options = new(new AgentId(request.Name), request.Ephemeral, "", [$"{KnownPropertyNames.Trusted}=true"]);
+			IAgent agent = await _agentService.CreateAgentAsync(options, context.CancellationToken);
 
 			List<AclClaimConfig> claims = new List<AclClaimConfig>();
 			claims.Add(new AclClaimConfig(HordeClaimTypes.Agent, agent.Id.ToString()));
@@ -202,7 +205,7 @@ namespace HordeServer.Server
 					throw new StructuredRpcException(StatusCode.PermissionDenied, "User is not authenticated to create new agents");
 				}
 
-				agent = await _agentService.CreateAgentAsync(agentId, false, "");
+				agent = await _agentService.CreateAgentAsync(new CreateAgentOptions(agentId, false, ""));
 			}
 
 			// Check the enrollment key in the user token matches
