Horde: Add timeouts for available port checking in tests

diff --git a/Engine/Source/Programs/Horde/HordeServer.Tests.Shared/DatabaseRunner.cs b/Engine/Source/Programs/Horde/HordeServer.Tests.Shared/DatabaseRunner.cs
--- a/Engine/Source/Programs/Horde/HordeServer.Tests.Shared/DatabaseRunner.cs
+++ b/Engine/Source/Programs/Horde/HordeServer.Tests.Shared/DatabaseRunner.cs
@@ -129,16 +129,18 @@ public abstract class DatabaseRunner : IDisposable
 		return temp;
 	}
 
-	private static int GetAvailablePort()
+	private static int GetAvailablePort(int timeoutMs = 5000)
 	{
 		using TcpListener listener = new(IPAddress.Loopback, 0);
+		listener.Server.ReceiveTimeout = timeoutMs;
+		listener.Server.SendTimeout = timeoutMs;
 		listener.Start();
 		int port = ((IPEndPoint)listener.LocalEndpoint).Port;
 		listener.Stop();
 		return port;
 	}
 
-	public static bool IsPortAvailable(int port)
+	public static bool IsPortAvailable(int port, int timeoutMs = 5000)
 	{
 		if (RuntimeInformation.IsOSPlatform(OSPlatform.Windows))
 		{
@@ -155,9 +157,13 @@ public abstract class DatabaseRunner : IDisposable
 
 		try
 		{
-			using TcpListener listenerAny = new TcpListener(IPAddress.Loopback, port);
+			using TcpListener listenerAny = new (IPAddress.Loopback, port);
+			listenerAny.Server.ReceiveTimeout = timeoutMs;
+			listenerAny.Server.SendTimeout = timeoutMs;
 			listenerAny.Start();
-			using TcpListener listenerLoopback = new TcpListener(IPAddress.Any, port);
+			using TcpListener listenerLoopback = new (IPAddress.Any, port);
+			listenerLoopback.Server.ReceiveTimeout = timeoutMs;
+			listenerLoopback.Server.SendTimeout = timeoutMs;
 			listenerLoopback.Start();
 			return true;
 		}
