Horde: Fix issue with check in stream utils
#rnx
#skipci

diff --git a/Engine/Source/Programs/Horde/HordeDashboard/src/base/utilities/streamUtils.ts b/Engine/Source/Programs/Horde/HordeDashboard/src/base/utilities/streamUtils.ts
--- a/Engine/Source/Programs/Horde/HordeDashboard/src/base/utilities/streamUtils.ts
+++ b/Engine/Source/Programs/Horde/HordeDashboard/src/base/utilities/streamUtils.ts
@@ -12,7 +12,7 @@ export const getActiveStreamId = (): string | undefined => {
         return undefined;
     }
 
-    if (![projectStore.streamById(path[1])]) {
+    if (!projectStore.streamById(path[1])) {
         return undefined;
     }
 
