Horde: Increase default GC frequency to 6m.

#rnx

diff --git a/Engine/Source/Programs/Horde/Plugins/Storage/HordeServer.Storage/StorageConfig.cs b/Engine/Source/Programs/Horde/Plugins/Storage/HordeServer.Storage/StorageConfig.cs
--- a/Engine/Source/Programs/Horde/Plugins/Storage/HordeServer.Storage/StorageConfig.cs
+++ b/Engine/Source/Programs/Horde/Plugins/Storage/HordeServer.Storage/StorageConfig.cs
@@ -326,7 +326,7 @@ namespace HordeServer.Storage
 		/// <summary>
 		/// How frequently to run garbage collection, in hours.
 		/// </summary>
-		public double GcFrequencyHrs { get; set; } = 2.0;
+		public double GcFrequencyHrs { get; set; } = 0.1;
 
 		/// <summary>
 		/// How long to keep newly uploaded orphanned objects before allowing them to be deleted, in hours.
