Horde: Add settings menu to tray app, as well as functionality for downloading tools from Horde.

diff --git a/Engine/Source/Programs/Horde/BuildHorde.xml b/Engine/Source/Programs/Horde/BuildHorde.xml
--- a/Engine/Source/Programs/Horde/BuildHorde.xml
+++ b/Engine/Source/Programs/Horde/BuildHorde.xml
@@ -140,15 +140,15 @@
 			<Delete Files="$(ServerToolsStagingDir_HordeAgentWin64)/..."/>
 			
 			<!-- Build the agent -->
-			<DotNet Arguments="publish &quot;$(RootDir)/Engine/Source/Programs/Horde/HordeAgent/HordeAgent.csproj&quot; --output &quot;$(ServerToolsStagingDir_HordeAgentWin64)&quot; --runtime win7-x64 --self-contained $(VersionArguments)"/>
+			<DotNet Arguments="publish &quot;$(RootDir)/Engine/Source/Programs/Horde/HordeAgent/HordeAgent.csproj&quot; --output &quot;$(ServerToolsStagingDir_HordeAgentWin64)&quot; --runtime win-x64 --self-contained $(VersionArguments)"/>
 			<Sign Files="$(ServerToolsStagingDir_HordeAgentWin64)/HordeAgent.exe;$(ServerToolsStagingDir_HordeAgentWin64)/HordeAgent.dll;$(ServerToolsStagingDir_HordeAgentWin64)/EpicGames.*.dll" If="'$(SignOutputs)' == 'true'"/>
 
 			<!-- Build the driver -->
-			<DotNet Arguments="publish &quot;$(RootDir)/Engine/Source/Programs/Horde/Drivers/JobDriver/JobDriver.csproj&quot; --output &quot;$(ServerToolsStagingDir_HordeAgentWin64)/JobDriver&quot; --runtime win7-x64 --self-contained $(VersionArguments)"/>
+			<DotNet Arguments="publish &quot;$(RootDir)/Engine/Source/Programs/Horde/Drivers/JobDriver/JobDriver.csproj&quot; --output &quot;$(ServerToolsStagingDir_HordeAgentWin64)/JobDriver&quot; --runtime win-x64 --self-contained $(VersionArguments)"/>
 			<Sign Files="$(ServerToolsStagingDir_HordeAgentWin64)/JobDriver/JobDriver.exe;$(ServerToolsStagingDir_HordeAgentWin64)/JobDriver/JobDriver.dll;$(ServerToolsStagingDir_HordeAgentWin64)/JobDriver/EpicGames.*.dll" If="'$(SignOutputs)' == 'true'"/>
 
 			<!-- Build the tray app -->
-			<DotNet Arguments="publish &quot;$(RootDir)/Engine/Source/Programs/Horde/HordeTrayApp/HordeTrayApp.csproj&quot; --output &quot;$(ServerToolsStagingDir_HordeAgentWin64)/TrayApp&quot; --runtime win7-x64 --self-contained $(VersionArguments)"/>
+			<DotNet Arguments="publish &quot;$(RootDir)/Engine/Source/Programs/Horde/HordeTrayApp/HordeTrayApp.csproj&quot; --output &quot;$(ServerToolsStagingDir_HordeAgentWin64)/TrayApp&quot; --runtime win-x64 --self-contained $(VersionArguments)"/>
 			<Sign Files="$(ServerToolsStagingDir_HordeAgentWin64)/TrayApp/Horde.Agent.TrayApp.exe;$(ServerToolsStagingDir_HordeAgentWin64)/TrayApp/EpicGames.*.dll" If="'$(SignOutputs)' == 'true'"/>
 
 			<Delete Files="$(ServerToolsStagingDir_HordeAgentWin64)/appsettings.Local.json"/>
diff --git a/Engine/Source/Programs/Horde/HordeTrayApp/AgentPlugin.cs b/Engine/Source/Programs/Horde/HordeTrayApp/AgentPlugin.cs
deleted file mode 100644
diff --git a/Engine/Source/Programs/Horde/HordeTrayApp/App.axaml.cs b/Engine/Source/Programs/Horde/HordeTrayApp/App.axaml.cs
--- a/Engine/Source/Programs/Horde/HordeTrayApp/App.axaml.cs
+++ b/Engine/Source/Programs/Horde/HordeTrayApp/App.axaml.cs
@@ -1,10 +1,14 @@
 // Copyright Epic Games, Inc. All Rights Reserved.
 
+using System.Diagnostics;
 using Avalonia;
 using Avalonia.Controls;
 using Avalonia.Controls.ApplicationLifetimes;
 using Avalonia.Markup.Xaml;
 using Avalonia.Threading;
+using EpicGames.Core;
+using HordeTrayApp.Plugins.HordeAgent;
+using Microsoft.Extensions.DependencyInjection;
 
 namespace HordeTrayApp
 {
@@ -13,18 +17,31 @@ namespace HordeTrayApp
 	/// </summary>
 	public sealed partial class App : Application, ITrayAppHost
 	{
-		readonly List<ITrayAppPlugin> _plugins = new List<ITrayAppPlugin>();
+		readonly ServiceProvider _serviceProvider;
 
 		WindowIcon? _normalIcon;
 		WindowIcon? _busyIcon;
 		WindowIcon? _pausedIcon;
 		WindowIcon? _errorIcon;
 
+		SettingsWindow? _settingsWindow;
+
 		/// <summary>
 		/// Constructor
 		/// </summary>
 		public App()
 		{
+			ServiceCollection serviceCollection = new ServiceCollection();
+			serviceCollection.AddLogging(builder => builder.AddEpicDefault());
+			serviceCollection.AddSingleton<IHordeClientProvider, HordeClientProvider>();
+			serviceCollection.AddSingleton<ToolCatalog>();
+			serviceCollection.AddSingleton<IToolCatalog, ToolCatalog>(sp => sp.GetRequiredService<ToolCatalog>());
+			serviceCollection.AddSingleton<ITrayAppHost>(this);
+
+			// Configure plugins
+			serviceCollection.AddSingleton<ITrayAppPlugin, HordeAgentPlugin>();
+
+			_serviceProvider = serviceCollection.BuildServiceProvider();
 		}
 
 		/// <inheritdoc/>
@@ -37,6 +54,18 @@ namespace HordeTrayApp
 			_pausedIcon = (WindowIcon)Resources["StatusPaused"]!;
 			_errorIcon = (WindowIcon)Resources["StatusError"]!;
 
+			UpdateMenu();
+
+			ToolCatalog toolCatalog = _serviceProvider.GetRequiredService<ToolCatalog>();
+			toolCatalog.OnItemsChanged += UpdateMenu;
+			toolCatalog.Start();
+		}
+
+		private void UpdateMenu()
+		{
+			NativeMenuItem settingsMenuItem = new NativeMenuItem("Settings...");
+			settingsMenuItem.Click += TrayIcon_Settings;
+
 			NativeMenuItem exitMenuItem = new NativeMenuItem("Exit");
 			exitMenuItem.Click += TrayIcon_Exit;
 
@@ -45,27 +74,109 @@ namespace HordeTrayApp
 			TrayIcon trayIcon = TrayIcon.GetIcons(this)![0];
 			trayIcon.Menu = contextMenu;
 
-//			_icon = new TrayIcon();
-//			_icon.Icon = _normalIcon;
-//			_icon.Menu = contextMenu;
-//			_icon.ToolTipText = "Hello world";
-
-//			TrayIcons icons = new TrayIcons();
-//			icons.Add(_icon);
+			foreach (ITrayAppPlugin plugin in _serviceProvider.GetServices<ITrayAppPlugin>())
+			{
+				plugin.PopulateContextMenu(contextMenu);
+			}
 
-//			TrayIcon.SetIcons(this, icons);
+			int numItems = contextMenu.Items.Count;
+			if (numItems > 0)
+			{
+				contextMenu.Items.Add(new NativeMenuItemSeparator());
+				numItems++;
+			}
 
-			_plugins.Add(new AgentPlugin(this));
+			IToolCatalog toolCatalog = _serviceProvider.GetRequiredService<IToolCatalog>();
+			foreach (IToolCatalogItem item in toolCatalog.Items)
+			{
+				CurrentToolDeploymentInfo? current = item.Current;
+				if (current != null)
+				{
+					ToolConfig? toolConfig = current.Config;
+					if (toolConfig != null && toolConfig.PopupMenu != null)
+					{
+						AddMenuItems(item, toolConfig.PopupMenu, contextMenu);
+					}
+				}
+			}
 
-//			NativeMenu contextMenu = trayIcon.Menu!;
-			foreach (ITrayAppPlugin plugin in _plugins)
+			if (contextMenu.Items.Count > numItems)
 			{
-				plugin.PopulateContextMenu(contextMenu);
+				contextMenu.Items.Add(new NativeMenuItemSeparator());
 			}
-			contextMenu.Items.Add(new NativeMenuItemSeparator());
+			contextMenu.Items.Add(settingsMenuItem);
 			contextMenu.Items.Add(exitMenuItem);
 		}
 
+		static void AddMenuItems(IToolCatalogItem item, ToolMenuItem toolMenuItem, NativeMenu menu)
+		{
+			if (!String.IsNullOrEmpty(toolMenuItem.Label))
+			{
+				NativeMenuItem menuItem = new NativeMenuItem(toolMenuItem.Label);
+				if (toolMenuItem.Children != null && toolMenuItem.Children.Count > 0)
+				{
+					NativeMenu subMenu = new NativeMenu();
+					foreach (ToolMenuItem child in toolMenuItem.Children)
+					{
+						AddMenuItems(item, child, subMenu);
+					}
+					menuItem.Menu = subMenu;
+				}
+				else
+				{
+					if (toolMenuItem.Command != null)
+					{
+						menuItem.Click += (_, _) => TrayIcon_RunCommand(item, toolMenuItem.Command);
+					}
+				}
+				menu.Items.Add(menuItem);
+			}
+		}
+
+		static void TrayIcon_RunCommand(IToolCatalogItem item, ToolMenuCommand command)
+		{
+			ProcessStartInfo startInfo = new ProcessStartInfo(command.FileName);
+			if (command.Arguments != null)
+			{
+				foreach (string argument in command.Arguments)
+				{
+					startInfo.ArgumentList.Add(argument);
+				}
+			}
+			startInfo.WorkingDirectory = item.Current!.Dir.FullName;
+			Process.Start(startInfo);
+		}
+
+		private void TrayIcon_Click(object? sender, EventArgs e)
+		{
+			OpenSettings();
+		}
+
+		private void TrayIcon_Settings(object? sender, EventArgs e)
+		{
+			OpenSettings();
+		}
+
+		private void OpenSettings()
+		{
+			if (_settingsWindow != null)
+			{
+				_settingsWindow.BringIntoView();
+				_settingsWindow.Show();
+			}
+			else
+			{
+				_settingsWindow = new SettingsWindow(_serviceProvider);
+				_settingsWindow.Closed += SettingsWindow_Closed;
+				_settingsWindow.Show();
+			}
+		}
+
+		private void SettingsWindow_Closed(object? sender, EventArgs e)
+		{
+			_settingsWindow = null;
+		}
+
 		private void TrayIcon_Exit(object? sender, EventArgs e)
 		{
 			((IClassicDesktopStyleApplicationLifetime)ApplicationLifetime!).Shutdown();
@@ -82,7 +193,7 @@ namespace HordeTrayApp
 			TrayAppPluginState state = TrayAppPluginState.Undefined;
 			List<string> messages = new List<string>();
 
-			foreach (ITrayAppPlugin plugin in _plugins)
+			foreach (ITrayAppPlugin plugin in _serviceProvider.GetServices<ITrayAppPlugin>())
 			{
 				TrayAppPluginStatus status = plugin.GetStatus();
 				if (status.State >= state)
diff --git a/Engine/Source/Programs/Horde/HordeTrayApp/GeneralSettingsPage.axaml.cs b/Engine/Source/Programs/Horde/HordeTrayApp/GeneralSettingsPage.axaml.cs
new file mode 100644
--- /dev/null
+++ b/Engine/Source/Programs/Horde/HordeTrayApp/GeneralSettingsPage.axaml.cs
@@ -0,0 +1,30 @@
+// Copyright Epic Games, Inc. All Rights Reserved.
+
+using Avalonia.Controls;
+
+namespace HordeTrayApp
+{
+	partial class GeneralSettingsPage : UserControl, IDisposable
+	{
+		readonly GeneralSettingsViewModel _viewModel;
+
+		public GeneralSettingsPage()
+			: this(SettingsContext.Default)
+		{
+		}
+
+		public GeneralSettingsPage(SettingsContext context)
+		{
+			InitializeComponent();
+
+			_viewModel = new GeneralSettingsViewModel(context);
+
+			DataContext = _viewModel;
+		}
+
+		public void Dispose()
+		{
+			_viewModel.Dispose();
+		}
+	}
+}
diff --git a/Engine/Source/Programs/Horde/HordeTrayApp/GeneralSettingsViewModel.cs b/Engine/Source/Programs/Horde/HordeTrayApp/GeneralSettingsViewModel.cs
new file mode 100644
--- /dev/null
+++ b/Engine/Source/Programs/Horde/HordeTrayApp/GeneralSettingsViewModel.cs
@@ -0,0 +1,164 @@
+// Copyright Epic Games, Inc. All Rights Reserved.
+
+using System.Diagnostics;
+using Avalonia.Controls;
+using CommunityToolkit.Mvvm.ComponentModel;
+using EpicGames.Horde;
+using FluentAvalonia.UI.Controls;
+using Microsoft.Extensions.DependencyInjection;
+
+namespace HordeTrayApp
+{
+	partial class GeneralSettingsViewModel : ObservableObject, IDisposable
+	{
+		record StateRecord(Uri ServerUrl, bool Connected);
+
+		readonly SettingsContext _context;
+		readonly IHordeClientProvider _hordeClientProvider;
+
+		[ObservableProperty]
+		string _serverUrl = String.Empty;
+
+		[ObservableProperty]
+		string _serverStatus = String.Empty;
+
+		[ObservableProperty]
+		bool? _isLoginEnabled;
+
+		public ToolCatalogViewModel ToolCatalog { get; }
+
+		public GeneralSettingsViewModel()
+			: this(SettingsContext.Default)
+		{
+		}
+
+		public GeneralSettingsViewModel(SettingsContext context)
+		{
+			_context = context;
+
+			_hordeClientProvider = context.ServiceProvider.GetRequiredService<IHordeClientProvider>();
+			_hordeClientProvider.OnStateChanged += OnStateChangedAsync;
+			_hordeClientProvider.OnAccessTokenStateChanged += OnStateChangedAsync;
+
+			ToolCatalog = new ToolCatalogViewModel(context.ServiceProvider.GetService<IToolCatalog>());
+
+			OnStateChangedAsync();
+		}
+
+		public void Dispose()
+		{
+			_hordeClientProvider.OnStateChanged -= OnStateChangedAsync;
+			_hordeClientProvider.OnAccessTokenStateChanged -= OnStateChangedAsync;
+		}
+
+		async void OnStateChangedAsync()
+		{
+			using IHordeClientRef? hordeClient = await _hordeClientProvider.GetClientRefAsync();
+			if (hordeClient == null)
+			{
+				ServerUrl = "No server configured";
+				ServerStatus = "Status unavailable";
+				IsLoginEnabled = false;
+			}
+			else
+			{
+				ServerUrl = hordeClient.Client.ServerUrl.ToString();
+				ServerStatus = hordeClient.Client.HasValidAccessToken() ? "Connected" : "Session Expired";
+				IsLoginEnabled = !hordeClient.Client.HasValidAccessToken();
+			}
+		}
+
+		public void OpenServer()
+		{
+			OpenBrowser(ServerUrl);
+		}
+
+		public async Task LoginAsync()
+		{
+			using IHordeClientRef? hordeClientRef = await _hordeClientProvider.GetClientRefAsync();
+			if (hordeClientRef != null)
+			{
+				await hordeClientRef.Client.LoginAsync(true, CancellationToken.None);
+			}
+		}
+
+		public async Task ConfigureServerAsync()
+		{
+			string serverUrl = ServerUrl;
+
+			string? errorMessage = null;
+			for (; ; )
+			{
+				TextBox serverUrlTextBox = new TextBox();
+				serverUrlTextBox.Text = serverUrl;
+				serverUrlTextBox.SelectionStart = serverUrl.Length;
+				serverUrlTextBox.SelectionEnd = serverUrl.Length;
+
+				Control content = serverUrlTextBox;
+				if (!String.IsNullOrEmpty(errorMessage))
+				{
+					TextBlock errorBlock = new TextBlock();
+					errorBlock.Text = errorMessage;
+
+					StackPanel stackPanel = new StackPanel() { Spacing = 10 };
+					stackPanel.Children.Add(serverUrlTextBox);
+					stackPanel.Children.Add(errorBlock);
+
+					content = stackPanel;
+				}
+
+				ContentDialog dialog = new ContentDialog()
+				{
+					Title = "Connect to Server",
+					Content = content,
+					PrimaryButtonText = "Connect",
+					CloseButtonText = "Cancel"
+				};
+
+				ContentDialogResult result = await dialog.ShowAsync(_context.SettingsWindow);
+				if (result == ContentDialogResult.None)
+				{
+					return;
+				}
+
+				serverUrl = serverUrlTextBox.Text;
+
+				try
+				{
+					HordeOptions.SetDefaultServerUrl(new Uri(serverUrl));
+
+					await _hordeClientProvider.RecreateAsync();
+
+					using IHordeClientRef? clientRef = await _hordeClientProvider.GetClientRefAsync();
+					if (clientRef != null)
+					{
+						await clientRef.Client.LoginAsync(true, CancellationToken.None);
+					}
+
+					return;
+				}
+				catch (Exception ex)
+				{
+					errorMessage = ex.ToString();
+				}
+			}
+		}
+
+		static void OpenBrowser(string url)
+		{
+			if (OperatingSystem.IsWindows())
+			{
+				string escapedUrl = url.Replace("&", "^&", StringComparison.Ordinal);
+				Process.Start(new ProcessStartInfo("cmd", $"/c start {escapedUrl}") { CreateNoWindow = true });
+			}
+			else if (OperatingSystem.IsLinux())
+			{
+				Process.Start("xdg-open", url);
+			}
+			else if (OperatingSystem.IsMacOS())
+			{
+				Process.Start("open", url);
+			}
+		}
+	}
+}
diff --git a/Engine/Source/Programs/Horde/HordeTrayApp/HordeTrayApp.csproj b/Engine/Source/Programs/Horde/HordeTrayApp/HordeTrayApp.csproj
--- a/Engine/Source/Programs/Horde/HordeTrayApp/HordeTrayApp.csproj
+++ b/Engine/Source/Programs/Horde/HordeTrayApp/HordeTrayApp.csproj
@@ -2,37 +2,52 @@
 
   <PropertyGroup>
     <OutputType>WinExe</OutputType>
-    <TargetFramework>net6.0-windows</TargetFramework>
+    <TargetFramework>net8.0-windows</TargetFramework>
     <Nullable>enable</Nullable>
     <ImplicitUsings>enable</ImplicitUsings>
     <NeutralLanguage>en-US</NeutralLanguage>
     <Configurations>Debug;Release;Analyze</Configurations>
 	<AssemblyName>Horde.Agent.TrayApp</AssemblyName>
     <AvaloniaUseCompiledBindingsByDefault>true</AvaloniaUseCompiledBindingsByDefault>
+    <NoWarn>IDE0005</NoWarn>
   </PropertyGroup>
 
   <PropertyGroup Condition=" '$(Configuration)|$(Platform)' == 'Analyze|AnyCPU' ">
     <RunAnalyzersDuringBuild>True</RunAnalyzersDuringBuild>
     <TreatWarningsAsErrors>false</TreatWarningsAsErrors>
     <EnforceCodeStyleInBuild>true</EnforceCodeStyleInBuild>
-    <GenerateDocumentationFile>True</GenerateDocumentationFile>
+    <GenerateDocumentationFile>true</GenerateDocumentationFile>
     <DocumentationFile></DocumentationFile>
+    <NoWarn>CA1501</NoWarn>
   </PropertyGroup>
 
   <ItemGroup>
-    <Compile Include="..\HordeAgent\IpcMessages.cs" Link="IpcMessages.cs" />
+    <Compile Include="..\HordeAgent\IpcMessages.cs" Link="Plugins\HordeAgent\IpcMessages.cs" />
   </ItemGroup>
 
   <ItemGroup>
     <ProjectReference Include="..\..\Shared\EpicGames.Core\EpicGames.Core.csproj" />
+    <ProjectReference Include="..\..\Shared\EpicGames.Horde\EpicGames.Horde.csproj" />
   </ItemGroup>
 
   <ItemGroup>
+    <Compile Update="GeneralSettingsPage.axaml.cs">
+      <DependentUpon>GeneralSettingsPage.axaml</DependentUpon>
+    </Compile>
+    <Compile Update="Plugins\HordeAgent\HordeAgentSettingsPage.axaml.cs">
+      <DependentUpon>HordeAgentSettingsPage.axaml</DependentUpon>
+    </Compile>
     <Compile Update="Properties\Resources.Designer.cs">
       <DesignTime>True</DesignTime>
       <AutoGen>True</AutoGen>
       <DependentUpon>Resources.resx</DependentUpon>
     </Compile>
+    <Compile Update="HordeServerSettings.axaml.cs">
+      <DependentUpon>HordeServerSettings.axaml</DependentUpon>
+    </Compile>
+    <Compile Update="SettingsWindow.axaml.cs">
+      <DependentUpon>SettingsWindow.axaml</DependentUpon>
+    </Compile>
   </ItemGroup>
 
   <ItemGroup>
@@ -53,6 +68,8 @@
     <PackageReference Include="Avalonia.Fonts.Inter" Version="11.1.0" />
     <!--Condition below is needed to remove Avalonia.Diagnostics package from build output in Release configuration.-->
     <PackageReference Condition="'$(Configuration)' == 'Debug'" Include="Avalonia.Diagnostics" Version="11.1.0" />
+    <PackageReference Include="FluentAvaloniaUI" Version="2.1.0" />
+    <PackageReference Include="CommunityToolkit.Mvvm" Version="8.2.1" />
   </ItemGroup>
 
 </Project>
\ No newline at end of file
diff --git a/Engine/Source/Programs/Horde/HordeTrayApp/IToolCatalog.cs b/Engine/Source/Programs/Horde/HordeTrayApp/IToolCatalog.cs
new file mode 100644
--- /dev/null
+++ b/Engine/Source/Programs/Horde/HordeTrayApp/IToolCatalog.cs
@@ -0,0 +1,20 @@
+// Copyright Epic Games, Inc. All Rights Reserved.
+
+namespace HordeTrayApp
+{
+	/// <summary>
+	/// Interface for the tool catalog
+	/// </summary>
+	public interface IToolCatalog
+	{
+		/// <summary>
+		/// List of available tools
+		/// </summary>
+		IReadOnlyList<IToolCatalogItem> Items { get; }
+
+		/// <summary>
+		/// Notification that some property of the items list has changed
+		/// </summary>
+		event Action? OnItemsChanged;
+	}
+}
diff --git a/Engine/Source/Programs/Horde/HordeTrayApp/IToolCatalogItem.cs b/Engine/Source/Programs/Horde/HordeTrayApp/IToolCatalogItem.cs
new file mode 100644
--- /dev/null
+++ b/Engine/Source/Programs/Horde/HordeTrayApp/IToolCatalogItem.cs
@@ -0,0 +1,87 @@
+// Copyright Epic Games, Inc. All Rights Reserved.
+
+using EpicGames.Core;
+using EpicGames.Horde.Tools;
+
+namespace HordeTrayApp
+{
+	/// <summary>
+	/// Describes a tool available to be installed
+	/// </summary>
+	public interface IToolCatalogItem
+	{
+		/// <summary>
+		/// Identifier for the tool
+		/// </summary>
+		ToolId Id { get; }
+
+		/// <summary>
+		/// Name of the tool
+		/// </summary>
+		string Name { get; }
+
+		/// <summary>
+		/// Description for the tool
+		/// </summary>
+		string Description { get; }
+
+		/// <summary>
+		/// Latest available version of the tool
+		/// </summary>
+		ToolDeploymentInfo? Latest { get; }
+
+		/// <summary>
+		/// The current tool deployment
+		/// </summary>
+		CurrentToolDeploymentInfo? Current { get; }
+
+		/// <summary>
+		/// Pending tool deployment, if an install or uninstall operation 
+		/// </summary>
+		PendingToolDeploymentInfo? Pending { get; }
+
+		/// <summary>
+		/// Callback for the item state changing
+		/// </summary>
+		event Action? OnItemChanged;
+
+		/// <summary>
+		/// Cancel any pending installation
+		/// </summary>
+		void Cancel();
+
+		/// <summary>
+		/// Start installing this tool or updating it to the latest version
+		/// </summary>
+		void Install();
+
+		/// <summary>
+		/// Start uninstalling this tool
+		/// </summary>
+		void Uninstall();
+	}
+
+	/// <summary>
+	/// Information about a particular tool deployment
+	/// </summary>
+	/// <param name="Id">Identifier for the deployment</param>
+	/// <param name="Version">Version number for this deployment</param>
+	public record class ToolDeploymentInfo(ToolDeploymentId Id, string Version);
+
+	/// <summary>
+	/// Information about an installed tool deployment
+	/// </summary>
+	/// <param name="Id">Identifier for the deployment</param>
+	/// <param name="Version">Version number for this deployment</param>
+	/// <param name="Dir">Directory containing the tool</param>
+	/// <param name="Config">Configuration for the tool</param>
+	public record class CurrentToolDeploymentInfo(ToolDeploymentId Id, string Version, DirectoryReference Dir, ToolConfig Config) : ToolDeploymentInfo(Id, Version);
+
+	/// <summary>
+	/// Information about a pending tool installation
+	/// </summary>
+	/// <param name="Failed">Whether the installation failed</param>
+	/// <param name="Message">Current status message</param>
+	/// <param name="Deployment">Information about the pending deployment</param>
+	public record class PendingToolDeploymentInfo(bool Failed, string? Message, ToolDeploymentInfo? Deployment);
+}
diff --git a/Engine/Source/Programs/Horde/HordeTrayApp/ITrayAppPlugin.cs b/Engine/Source/Programs/Horde/HordeTrayApp/ITrayAppPlugin.cs
--- a/Engine/Source/Programs/Horde/HordeTrayApp/ITrayAppPlugin.cs
+++ b/Engine/Source/Programs/Horde/HordeTrayApp/ITrayAppPlugin.cs
@@ -1,6 +1,7 @@
 // Copyright Epic Games, Inc. All Rights Reserved.
 
 using Avalonia.Controls;
+using FluentAvalonia.UI.Controls;
 
 namespace HordeTrayApp
 {
@@ -16,44 +17,39 @@ namespace HordeTrayApp
 	/// <summary>
 	/// Reported status of a plugin
 	/// </summary>
-	record class TrayAppPluginStatus(TrayAppPluginState State, string? Message);
+	record class TrayAppPluginStatus(TrayAppPluginState State, string? Message)
+	{
+		public static TrayAppPluginStatus Default { get; } = new TrayAppPluginStatus(TrayAppPluginState.Undefined, null);
+	}
 
 	/// <summary>
 	/// Plugin for the tray app
 	/// </summary>
 	interface ITrayAppPlugin : IAsyncDisposable
 	{
+		/// <summary>
+		/// Name of the plugin to show in the settings page
+		/// </summary>
+		string Name { get; }
+
+		/// <summary>
+		/// Icon to show on the settings page
+		/// </summary>
+		IconSource Icon { get; }
+
 		/// <summary>
 		/// Get the current status of this plugin
 		/// </summary>
 		TrayAppPluginStatus GetStatus();
 
+		/// <summary>
+		/// Create a settings page for this plugin
+		/// </summary>
+		Control CreateSettingsPage(SettingsContext context);
+
 		/// <summary>
 		/// Allow the plugin to customize the tray icon context menu
 		/// </summary>
 		void PopulateContextMenu(NativeMenu contextMenu);
 	}
-
-	/// <summary>
-	/// Plugin for the tray app
-	/// </summary>
-	abstract class TrayAppPluginBase : ITrayAppPlugin
-	{
-		/// <inheritdoc/>
-		public virtual TrayAppPluginStatus GetStatus()
-		{
-			return new TrayAppPluginStatus(TrayAppPluginState.Undefined, null);
-		}
-
-		/// <inheritdoc/>
-		public virtual void PopulateContextMenu(NativeMenu contextMenu)
-		{
-		}
-
-		/// <inheritdoc/>
-		public virtual ValueTask DisposeAsync()
-		{
-			return default;
-		}
-	}
 }
diff --git a/Engine/Source/Programs/Horde/HordeTrayApp/Plugins/HordeAgent/HordeAgentPlugin.cs b/Engine/Source/Programs/Horde/HordeTrayApp/Plugins/HordeAgent/HordeAgentPlugin.cs
new file mode 100644
--- /dev/null
+++ b/Engine/Source/Programs/Horde/HordeTrayApp/Plugins/HordeAgent/HordeAgentPlugin.cs
@@ -0,0 +1,477 @@
+// Copyright Epic Games, Inc. All Rights Reserved.
+
+using System.Diagnostics;
+using System.IO.Pipes;
+using System.Runtime.InteropServices;
+using System.Text.Json;
+using System.Text.Json.Serialization;
+using Avalonia.Controls;
+using EpicGames.Core;
+using FluentAvalonia.UI.Controls;
+using HordeAgent;
+using Microsoft.Win32;
+
+namespace HordeTrayApp.Plugins.HordeAgent
+{
+	class HordeAgentPlugin : ITrayAppPlugin
+	{
+		record struct IdleStat(string Name, long Value, long MinValue);
+
+		readonly ITrayAppHost _host;
+
+		readonly BackgroundTask _clientTask;
+		readonly BackgroundTask _tickPauseStateTask;
+
+		readonly FileReference _settingsFile;
+
+		byte[] _settingsData = Array.Empty<byte>();
+		HordeAgentSettings _settings = new HordeAgentSettings();
+
+		public HordeAgentSettings Settings => _settings;
+
+		AgentSettingsMessage? _agentSettings;
+
+		public string Name => "Horde Agent";
+
+		public IconSource Icon => new SymbolIconSource() { Symbol = Symbol.Globe };
+
+		public Control CreateSettingsPage(SettingsContext context)
+			=> new HordeAgentSettingsPage(context, this);
+
+		public void UpdateSettings(HordeAgentSettings settings)
+		{
+			byte[] data = JsonSerializer.SerializeToUtf8Bytes(settings, GetJsonSerializerOptions());
+			if (!data.SequenceEqual(_settingsData))
+			{
+				_settings = settings;
+				_settingsData = data;
+
+				DirectoryReference.CreateDirectory(_settingsFile.Directory);
+				FileReference.WriteAllBytes(_settingsFile, data);
+
+				_statusChangedEvent.Set();
+			}
+		}
+
+		void EnrollWithServer()
+		{
+			Uri? serverUrl = _agentSettings?.ServerUrl;
+			if (serverUrl != null)
+			{
+				Process.Start(new ProcessStartInfo(new Uri(serverUrl, "agents/registration").ToString()) { UseShellExecute = true });
+			}
+		}
+
+		public HordeAgentPlugin(ITrayAppHost host)
+		{
+			_host = host;
+
+			DirectoryReference? settingsRoot = DirectoryReference.GetSpecialFolder(Environment.SpecialFolder.CommonApplicationData);
+			settingsRoot ??= DirectoryReference.GetCurrentDirectory();
+			_settingsFile = FileReference.Combine(settingsRoot, "Epic", "Horde", "TrayApp", "Settings.json");
+
+			LoadSettings();
+
+			_clientTask = BackgroundTask.StartNew(StatusTaskAsync);
+			_tickPauseStateTask = BackgroundTask.StartNew(ctx => TickPauseStateAsync(ctx));
+		}
+
+		static JsonSerializerOptions GetJsonSerializerOptions()
+		{
+			JsonSerializerOptions options = new JsonSerializerOptions();
+			options.PropertyNamingPolicy = JsonNamingPolicy.CamelCase;
+			options.PropertyNameCaseInsensitive = true;
+			options.DefaultIgnoreCondition = JsonIgnoreCondition.WhenWritingNull;
+			options.AllowTrailingCommas = true;
+			options.WriteIndented = true;
+			options.Converters.Add(new JsonStringEnumConverter());
+			return options;
+		}
+
+		void LoadSettings()
+		{
+			if (FileReference.Exists(_settingsFile))
+			{
+				try
+				{
+					byte[] data = FileReference.ReadAllBytes(_settingsFile);
+					if (!data.SequenceEqual(_settingsData))
+					{
+						_settings = JsonSerializer.Deserialize<HordeAgentSettings>(data, GetJsonSerializerOptions())!;
+						if (_settings.Mode == null)
+						{
+							_settings.Mode = ReadLegacyMode();
+							UpdateSettings(_settings);
+						}
+						_settingsData = data;
+						return;
+					}
+				}
+				catch (Exception)
+				{
+				}
+			}
+
+			UpdateSettings(new HordeAgentSettings());
+		}
+
+		static AgentMode? ReadLegacyMode()
+		{
+			if (!OperatingSystem.IsWindows())
+			{
+				return null;
+			}
+
+			const string RegistryKey = "HKEY_CURRENT_USER\\Software\\Epic Games\\Horde\\TrayApp";
+			const string RegistryStatusValue = "Status";
+
+			int? status = Registry.GetValue(RegistryKey, RegistryStatusValue, null) as int?;
+			if (status == null)
+			{
+				return null;
+			}
+
+			return status.Value switch
+			{
+				0 => AgentMode.Dedicated,
+				1 => AgentMode.Disabled,
+				2 => AgentMode.Workstation,
+				_ => null
+			};
+		}
+
+		public void PopulateContextMenu(NativeMenu contextMenu)
+		{
+			NativeMenuItem enrollMenuItem = new NativeMenuItem("Enroll with Server...");
+			enrollMenuItem.Click += (s, e) => EnrollWithServer();
+			contextMenu.Items.Add(enrollMenuItem);
+		}
+
+		public async ValueTask DisposeAsync()
+		{
+			await _tickPauseStateTask.DisposeAsync();
+			await _clientTask.DisposeAsync();
+		}
+
+		private void OnOpenLogs(object? sender, EventArgs e)
+		{
+			DirectoryReference? programDataDir = DirectoryReference.GetSpecialFolder(Environment.SpecialFolder.CommonApplicationData);
+			if (programDataDir != null)
+			{
+				DirectoryReference logsDir = DirectoryReference.Combine(programDataDir, "Epic", "Horde", "Agent");
+				if (DirectoryReference.Exists(logsDir))
+				{
+					Process.Start(new ProcessStartInfo { FileName = logsDir.FullName, UseShellExecute = true });
+				}
+			}
+		}
+
+		TrayAppPluginStatus? _status;
+
+		void SetStatus(AgentStatusMessage status)
+		{
+			if (!status.Healthy)
+			{
+				string message = String.IsNullOrEmpty(status.Detail) ? "Error. Check logs." : status.Detail.Length > 100 ? status.Detail.Substring(0, 100) : status.Detail;
+				_status = new TrayAppPluginStatus(TrayAppPluginState.Error, message);
+			}
+			else if (status.NumLeases > 0)
+			{
+				string message = status.NumLeases == 1 ? "Currently handling 1 lease" : $"Currently handling {status.NumLeases} leases";
+				_status = new TrayAppPluginStatus(TrayAppPluginState.Busy, message);
+			}
+			else if (_enabled)
+			{
+				string message = "Agent is operating normally";
+				_status = new TrayAppPluginStatus(TrayAppPluginState.Ok, message);
+			}
+			else
+			{
+				string message = "Agent is paused";
+				_status = new TrayAppPluginStatus(TrayAppPluginState.Paused, message);
+			}
+			_host.UpdateStatus();
+		}
+
+		public TrayAppPluginStatus GetStatus()
+			=> _status ?? TrayAppPluginStatus.Default;
+
+		async Task StatusTaskAsync(CancellationToken cancellationToken)
+		{
+			SetStatus(AgentStatusMessage.Starting);
+			for (; ; )
+			{
+				try
+				{
+					await PollForStatusUpdatesAsync(cancellationToken);
+				}
+				catch (OperationCanceledException) when (cancellationToken.IsCancellationRequested)
+				{
+					break;
+				}
+				catch
+				{
+					SetStatus(new AgentStatusMessage(false, 0, "Unable to connect to Agent. Check the service is running."));
+					await Task.Delay(TimeSpan.FromSeconds(5.0), cancellationToken);
+				}
+			}
+		}
+
+#pragma warning disable IDE1006
+		[StructLayout(LayoutKind.Sequential)]
+		struct LASTINPUTINFO
+		{
+			public int cbSize;
+			public uint dwTime;
+		}
+
+		[DllImport("user32.dll")]
+		static extern bool GetLastInputInfo(ref LASTINPUTINFO plii);
+
+		[DllImport("kernel32.dll")]
+		static extern uint GetTickCount();
+
+		[StructLayout(LayoutKind.Sequential)]
+		struct FILETIME
+		{
+			public uint dwLowDateTime;
+			public uint dwHighDateTime;
+
+			public readonly ulong Total => dwLowDateTime | (ulong)dwHighDateTime << 32;
+		};
+
+		[StructLayout(LayoutKind.Sequential)]
+		struct MEMORYSTATUSEX
+		{
+			public int dwLength;
+			public uint dwMemoryLoad;
+			public ulong ullTotalPhys;
+			public ulong ullAvailPhys;
+			public ulong ullTotalPageFile;
+			public ulong ullAvailPageFile;
+			public ulong ullTotalVirtual;
+			public ulong ullAvailVirtual;
+			public ulong ullAvailExtendedVirtual;
+		}
+
+		[DllImport("kernel32.dll", CharSet = CharSet.Auto, SetLastError = true)]
+		static extern bool GlobalMemoryStatusEx(ref MEMORYSTATUSEX lpBuffer);
+
+		[DllImport("kernel32.dll", SetLastError = true)]
+		static extern bool GetSystemTimes(out FILETIME lpIdleTime, out FILETIME lpKernelTime, out FILETIME lpUserTime);
+#pragma warning restore IDE1006
+
+		bool _enabled;
+		readonly AsyncEvent _statusChangedEvent = new AsyncEvent();
+		readonly AsyncEvent _enabledChangedEvent = new AsyncEvent();
+
+		async Task TickPauseStateAsync(CancellationToken cancellationToken)
+		{
+			await using BackgroundTask cpuStatsTask = BackgroundTask.StartNew(ctx => TickCpuStatsAsync(ctx));
+			await using BackgroundTask criticalProcessTask = BackgroundTask.StartNew(ctx => TickCriticalProcessAsync(ctx));
+
+			TimeSpan pollInterval = TimeSpan.FromSeconds(0.25);
+
+			Stopwatch stateChangeTimer = Stopwatch.StartNew();
+			while (!cancellationToken.IsCancellationRequested)
+			{
+				Task statusChangedTask = _statusChangedEvent.Task;
+
+				AgentMode mode = _settings.Mode ?? AgentMode.Disabled;
+				if (mode == AgentMode.Dedicated)
+				{
+					if (!_enabled)
+					{
+						_enabled = true;
+						_enabledChangedEvent.Set();
+					}
+				}
+				else if (mode == AgentMode.Disabled)
+				{
+					if (_enabled)
+					{
+						_enabled = false;
+						_enabledChangedEvent.Set();
+					}
+				}
+
+				DateTime utcNow = DateTime.UtcNow;
+				IEnumerable<IdleStat> idleStats = GetIdleStats();
+
+				bool idle = idleStats.All(x => x.Value >= x.MinValue);
+				if (idle == _enabled)
+				{
+					stateChangeTimer.Restart();
+				}
+
+				const int WakeTimeSecs = 2;
+				const int IdleTimeSecs = 30;
+				int stateChangeTime = (int)stateChangeTimer.Elapsed.TotalSeconds;
+				int stateChangeMaxTime = _enabled ? WakeTimeSecs : IdleTimeSecs;
+				//				_idleForm?.TickStats(_enabled, stateChangeTime, stateChangeMaxTime, idleStats);
+
+				if (mode == AgentMode.Workstation && stateChangeTime >= stateChangeMaxTime)
+				{
+					_enabled ^= true;
+					_enabledChangedEvent.Set();
+					stateChangeTimer.Restart();
+				}
+
+				await Task.WhenAny(statusChangedTask, Task.Delay(pollInterval, cancellationToken));
+			}
+		}
+
+		IEnumerable<IdleStat> GetIdleStats()
+		{
+			// Check there has been no input for a while
+			LASTINPUTINFO lastInputInfo = new LASTINPUTINFO();
+			lastInputInfo.cbSize = Marshal.SizeOf<LASTINPUTINFO>();
+
+			if (GetLastInputInfo(ref lastInputInfo))
+			{
+				yield return new IdleStat("LastInputTime", (GetTickCount() - lastInputInfo.dwTime) / 1000, _settings.Idle.MinIdleTimeSecs);
+			}
+
+			// Check that no critical processes are running
+			if (_settings.Idle.CriticalProcesses.Any())
+			{
+				yield return new IdleStat("CriticalProcCount", -_idleCriticalProcessCount, 0);
+			}
+
+			// Only look at memory/CPU usage if we're not paused; executing jobs will increase them
+			if (!_enabled)
+			{
+				// Check the CPU usage doesn't exceed the limit
+				yield return new IdleStat("IdleCpuPct", _idleCpuPct, _settings.Idle.MinIdleCpuPct);
+
+				// Check there's enough available virtual memory 
+				MEMORYSTATUSEX memoryStatus = new MEMORYSTATUSEX();
+				memoryStatus.dwLength = Marshal.SizeOf<MEMORYSTATUSEX>();
+
+				if (GlobalMemoryStatusEx(ref memoryStatus))
+				{
+					yield return new IdleStat("VirtualMemMb", (long)(memoryStatus.ullAvailPhys + memoryStatus.ullAvailPageFile) / (1024 * 1024), _settings.Idle.MinFreeVirtualMemMb);
+				}
+			}
+		}
+
+		int _idleCpuPct = 0;
+		int _idleCriticalProcessCount = 0;
+
+		async Task TickCpuStatsAsync(CancellationToken cancellationToken)
+		{
+			const int NumSamples = 10;
+			TimeSpan sampleInterval = TimeSpan.FromSeconds(0.2);
+			(ulong IdleTime, ulong TotalTime)[] samples = new (ulong IdleTime, ulong TotalTime)[NumSamples];
+
+			int sampleIdx = 0;
+			for (; ; )
+			{
+				if (GetSystemTimes(out FILETIME idleTime, out FILETIME kernelTime, out FILETIME userTime))
+				{
+					(ulong prevIdleTime, ulong prevTotalTime) = samples[sampleIdx];
+					(ulong nextIdleTime, ulong nextTotalTime) = (idleTime.Total, kernelTime.Total + userTime.Total);
+
+					samples[sampleIdx] = (nextIdleTime, nextTotalTime);
+					sampleIdx = (sampleIdx + 1) % NumSamples;
+
+					if (prevTotalTime > 0 && nextTotalTime > prevTotalTime)
+					{
+						_idleCpuPct = (int)((nextIdleTime - prevIdleTime) * 100 / (nextTotalTime - prevTotalTime));
+					}
+				}
+				await Task.Delay(sampleInterval, cancellationToken);
+			}
+		}
+
+		async Task TickCriticalProcessAsync(CancellationToken cancellationToken)
+		{
+			TimeSpan sampleInterval = TimeSpan.FromSeconds(1.0);
+
+			for (; ; )
+			{
+				try
+				{
+					if (RuntimeInformation.IsOSPlatform(OSPlatform.Windows) && _settings.Idle.CriticalProcesses.Any())
+					{
+						IEnumerable<int> hordeProcessIds = Process.GetProcessesByName("HordeAgent").Select(x => x.Id);
+						IEnumerable<Process> criticalProcesses = _settings.Idle.CriticalProcesses
+							.Select(x => Path.GetFileNameWithoutExtension(x).ToUpperInvariant())
+							.Distinct()
+							.SelectMany(x => Process.GetProcessesByName(x))
+							.Where(x => !x.HasExited);
+
+						// Ignore processes that are descendants of HordeAgent
+						if (hordeProcessIds.Any() && criticalProcesses.Any())
+						{
+							criticalProcesses = criticalProcesses
+								.Where(x => !ProcessUtils.GetAncestorProcesses(x)
+								.Select(x => x.Id).Intersect(hordeProcessIds).Any())
+								.Where(x => !x.HasExited);
+						}
+
+						_idleCriticalProcessCount = criticalProcesses.Count();
+					}
+				}
+				catch (InvalidOperationException)
+				{
+					// If a process stops running Process.Id will throw an exception
+				}
+				await Task.Delay(sampleInterval, cancellationToken);
+			}
+		}
+
+		async Task PollForStatusUpdatesAsync(CancellationToken cancellationToken)
+		{
+			AgentMessageBuffer message = new AgentMessageBuffer();
+			using (NamedPipeClientStream pipeClient = new NamedPipeClientStream(".", AgentMessagePipe.PipeName, PipeDirection.InOut))
+			{
+				SetStatus(new AgentStatusMessage(false, 0, "Connecting to agent..."));
+				await pipeClient.ConnectAsync(cancellationToken);
+
+				SetStatus(new AgentStatusMessage(false, 0, "Waiting for status update."));
+				for (; ; )
+				{
+					Task idleChangeTask = _enabledChangedEvent.Task;
+
+					bool enabled = _enabled;
+					message.Set(AgentMessageType.SetEnabledRequest, new AgentEnabledMessage(enabled));
+					await message.SendAsync(pipeClient, cancellationToken);
+
+					if (_agentSettings == null)
+					{
+						message.Set(AgentMessageType.GetSettingsRequest);
+						await message.SendAsync(pipeClient, cancellationToken);
+
+						if (!await message.TryReadAsync(pipeClient, cancellationToken))
+						{
+							break;
+						}
+						if (message.Type == AgentMessageType.GetSettingsResponse)
+						{
+							_agentSettings = message.Parse<AgentSettingsMessage>();
+						}
+					}
+
+					message.Set(AgentMessageType.GetStatusRequest);
+					await message.SendAsync(pipeClient, cancellationToken);
+
+					if (!await message.TryReadAsync(pipeClient, cancellationToken))
+					{
+						break;
+					}
+
+					switch (message.Type)
+					{
+						case AgentMessageType.GetStatusResponse:
+							AgentStatusMessage status = message.Parse<AgentStatusMessage>();
+							SetStatus(status);
+							break;
+					}
+
+					await Task.WhenAny(idleChangeTask, Task.Delay(TimeSpan.FromSeconds(5.0), cancellationToken));
+				}
+			}
+		}
+	}
+}
diff --git a/Engine/Source/Programs/Horde/HordeTrayApp/Plugins/HordeAgent/HordeAgentSettings.cs b/Engine/Source/Programs/Horde/HordeTrayApp/Plugins/HordeAgent/HordeAgentSettings.cs
new file mode 100644
--- /dev/null
+++ b/Engine/Source/Programs/Horde/HordeTrayApp/Plugins/HordeAgent/HordeAgentSettings.cs
@@ -0,0 +1,30 @@
+// Copyright Epic Games, Inc. All Rights Reserved.
+
+#pragma warning disable CS1591
+
+namespace HordeTrayApp.Plugins.HordeAgent
+{
+	public enum AgentMode
+	{
+		Workstation = 0,
+		Dedicated = 1,
+		Disabled = 2,
+	}
+
+	[Serializable]
+	public class HordeAgentSettings
+	{
+		public AgentMode? Mode { get; set; }
+		public IdleSettings Idle { get; set; } = new();
+	}
+
+	[Serializable]
+	public class IdleSettings
+	{
+		public int MinIdleTimeSecs { get; set; } = 2;
+		public int MinIdleCpuPct { get; set; } = 70;
+		public int MinFreeVirtualMemMb { get; set; } = 256;
+
+		public string[] CriticalProcesses { get; set; } = Array.Empty<string>();
+	}
+}
diff --git a/Engine/Source/Programs/Horde/HordeTrayApp/Plugins/HordeAgent/HordeAgentSettingsPage.axaml.cs b/Engine/Source/Programs/Horde/HordeTrayApp/Plugins/HordeAgent/HordeAgentSettingsPage.axaml.cs
new file mode 100644
--- /dev/null
+++ b/Engine/Source/Programs/Horde/HordeTrayApp/Plugins/HordeAgent/HordeAgentSettingsPage.axaml.cs
@@ -0,0 +1,21 @@
+// Copyright Epic Games, Inc. All Rights Reserved.
+
+using Avalonia.Controls;
+
+namespace HordeTrayApp.Plugins.HordeAgent
+{
+	partial class HordeAgentSettingsPage : UserControl
+	{
+		public HordeAgentSettingsPage()
+			: this(SettingsContext.Default, null)
+		{
+		}
+
+		public HordeAgentSettingsPage(SettingsContext context, HordeAgentPlugin? plugin)
+		{
+			InitializeComponent();
+
+			DataContext = new HordeAgentSettingsViewModel(context, plugin);
+		}
+	}
+}
\ No newline at end of file
diff --git a/Engine/Source/Programs/Horde/HordeTrayApp/Plugins/HordeAgent/HordeAgentSettingsViewModel.cs b/Engine/Source/Programs/Horde/HordeTrayApp/Plugins/HordeAgent/HordeAgentSettingsViewModel.cs
new file mode 100644
--- /dev/null
+++ b/Engine/Source/Programs/Horde/HordeTrayApp/Plugins/HordeAgent/HordeAgentSettingsViewModel.cs
@@ -0,0 +1,130 @@
+// Copyright Epic Games, Inc. All Rights Reserved.
+
+using System.ComponentModel;
+using System.Diagnostics;
+using Avalonia.Platform.Storage;
+using CommunityToolkit.Mvvm.ComponentModel;
+
+namespace HordeTrayApp.Plugins.HordeAgent
+{
+	partial class HordeAgentSettingsViewModel : ObservableObject
+	{
+		readonly SettingsContext _context;
+		readonly HordeAgentPlugin? _plugin;
+
+		AgentMode _mode;
+
+		static string GetPropertyName(AgentMode mode)
+			=> mode switch
+			{
+				AgentMode.Dedicated => nameof(IsDedicatedMode),
+				AgentMode.Workstation => nameof(IsWorkstationMode),
+				_ => nameof(IsDisabled)
+			};
+
+		void SetModeIfTrue(AgentMode mode, bool value)
+		{
+			if (value && mode != _mode)
+			{
+				string oldName = GetPropertyName(_mode);
+				string newName = GetPropertyName(mode);
+
+				OnPropertyChanging(oldName);
+				OnPropertyChanging(newName);
+
+				_mode = mode;
+
+				OnPropertyChanged(oldName);
+				OnPropertyChanged(newName);
+			}
+		}
+
+		public bool IsDedicatedMode
+		{
+			get => _mode == AgentMode.Dedicated;
+			set => SetModeIfTrue(AgentMode.Dedicated, value);
+		}
+
+		public bool IsWorkstationMode
+		{
+			get => _mode == AgentMode.Workstation;
+			set => SetModeIfTrue(AgentMode.Workstation, value);
+		}
+
+		public bool IsDisabled
+		{
+			get => _mode == AgentMode.Disabled;
+			set => SetModeIfTrue(AgentMode.Disabled, value);
+		}
+
+		[ObservableProperty]
+		int _minIdleTimeSecs;
+
+		[ObservableProperty]
+		int _minFreeMemoryMb;
+
+		[ObservableProperty]
+		int _minFreeCpuPct;
+
+		[ObservableProperty]
+		string _sandboxDir = "C:\\Dump"; // TODO
+
+		public HordeAgentSettingsViewModel(SettingsContext context, HordeAgentPlugin? plugin)
+		{
+			_context = context;
+			_plugin = plugin;
+
+			HordeAgentSettings settings = plugin?.Settings ?? new HordeAgentSettings();
+			_mode = settings.Mode ?? AgentMode.Disabled;
+			_minIdleTimeSecs = settings.Idle.MinIdleTimeSecs;
+			_minFreeMemoryMb = settings.Idle.MinFreeVirtualMemMb;
+			_minFreeCpuPct = settings.Idle.MinIdleCpuPct;
+		}
+
+		protected override void OnPropertyChanged(PropertyChangedEventArgs e)
+		{
+			base.OnPropertyChanged(e);
+
+			if (_plugin != null)
+			{
+				HordeAgentSettings settings = new HordeAgentSettings();
+
+				settings.Mode = _mode;
+				settings.Idle.MinIdleTimeSecs = MinIdleTimeSecs;
+				settings.Idle.MinFreeVirtualMemMb = MinFreeMemoryMb;
+				settings.Idle.MinIdleCpuPct = MinFreeCpuPct;
+
+				_plugin.UpdateSettings(settings);
+			}
+		}
+
+		public void OpenSandbox()
+		{
+			ProcessStartInfo startInfo = new ProcessStartInfo("explorer.exe");
+			startInfo.ArgumentList.Add(SandboxDir);
+			Process.Start(startInfo);
+		}
+
+		public async Task MoveSandboxAsync()
+		{
+			if (_context.SettingsWindow != null)
+			{
+				IStorageProvider storageProvider = _context.SettingsWindow.StorageProvider;
+
+				FolderPickerOpenOptions options = new FolderPickerOpenOptions();
+				options.AllowMultiple = false;
+				options.SuggestedStartLocation = await storageProvider.TryGetFolderFromPathAsync(SandboxDir);
+
+				IReadOnlyList<IStorageFolder> folders = await storageProvider.OpenFolderPickerAsync(options);
+				if (folders.Count > 0)
+				{
+					string? path = folders[0].TryGetLocalPath();
+					if (path != null)
+					{
+						SandboxDir = path;
+					}
+				}
+			}
+		}
+	}
+}
diff --git a/Engine/Source/Programs/Horde/HordeTrayApp/Program.cs b/Engine/Source/Programs/Horde/HordeTrayApp/Program.cs
--- a/Engine/Source/Programs/Horde/HordeTrayApp/Program.cs
+++ b/Engine/Source/Programs/Horde/HordeTrayApp/Program.cs
@@ -69,7 +69,7 @@ namespace HordeTrayApp
 			}
 
 			BuildAvaloniaApp()
-				.StartWithClassicDesktopLifetime(Array.Empty<string>(), ShutdownMode.OnExplicitShutdown);
+				.StartWithClassicDesktopLifetime(args, ShutdownMode.OnExplicitShutdown);
 
 			return 0;
 		}
diff --git a/Engine/Source/Programs/Horde/HordeTrayApp/Settings.cs b/Engine/Source/Programs/Horde/HordeTrayApp/Settings.cs
deleted file mode 100644
diff --git a/Engine/Source/Programs/Horde/HordeTrayApp/SettingsWindow.axaml.cs b/Engine/Source/Programs/Horde/HordeTrayApp/SettingsWindow.axaml.cs
new file mode 100644
--- /dev/null
+++ b/Engine/Source/Programs/Horde/HordeTrayApp/SettingsWindow.axaml.cs
@@ -0,0 +1,68 @@
+// Copyright Epic Games, Inc. All Rights Reserved.
+
+using Avalonia.Controls;
+using FluentAvalonia.UI.Controls;
+using Microsoft.Extensions.DependencyInjection;
+
+namespace HordeTrayApp
+{
+	record class SettingsContext(Window? SettingsWindow, IServiceProvider ServiceProvider)
+	{
+		public static SettingsContext Default { get; } = new SettingsContext(null!, new ServiceCollection().BuildServiceProvider());
+	}
+
+	partial class SettingsWindow : Window
+	{
+		readonly IServiceProvider? _serviceProvider;
+		readonly Dictionary<string, ITrayAppPlugin> _typeToPlugin = new Dictionary<string, ITrayAppPlugin>();
+
+		public SettingsWindow()
+			: this(null!)
+		{ }
+
+		public SettingsWindow(IServiceProvider serviceProvider)
+		{
+			InitializeComponent();
+
+			_serviceProvider = serviceProvider;
+
+			foreach (ITrayAppPlugin plugin in _serviceProvider.GetServices<ITrayAppPlugin>())
+			{
+				string typeName = plugin.GetType().FullName!;
+				_navView.MenuItems.Add(new NavigationViewItem(){ Content = plugin.Name, IconSource = plugin.Icon, Tag = typeName });
+				_typeToPlugin.Add(typeName, plugin);
+			}
+
+			_navView.SelectedItem = _navView.MenuItems[0];
+			_navView.SelectionChanged += NavView_SelectionChanged;
+
+			NavView_UpdateContent();
+		}
+
+		private void NavView_SelectionChanged(object? sender, NavigationViewSelectionChangedEventArgs e)
+		{
+			if (sender == _navView)
+			{
+				NavView_UpdateContent();
+			}
+		}
+
+		private void NavView_UpdateContent()
+		{
+			if (_navView.SelectedItem is NavigationViewItem nvi)
+			{
+				SettingsContext context = new SettingsContext(this, _serviceProvider!);
+
+				string? typeName = (string?)nvi.Tag;
+				if (typeName != null && _typeToPlugin.TryGetValue(typeName, out ITrayAppPlugin? plugin))
+				{
+					_navView.Content = plugin.CreateSettingsPage(context);
+				}
+				else
+				{
+					_navView.Content = new GeneralSettingsPage(context);
+				}
+			}
+		}
+	}
+}
diff --git a/Engine/Source/Programs/Horde/HordeTrayApp/ToolCatalog.cs b/Engine/Source/Programs/Horde/HordeTrayApp/ToolCatalog.cs
new file mode 100644
--- /dev/null
+++ b/Engine/Source/Programs/Horde/HordeTrayApp/ToolCatalog.cs
@@ -0,0 +1,709 @@
+// Copyright Epic Games, Inc. All Rights Reserved.
+
+using System.Diagnostics;
+using System.Text.Json;
+using EpicGames.Core;
+using EpicGames.Horde;
+using EpicGames.Horde.Storage.Nodes;
+using EpicGames.Horde.Tools;
+using Microsoft.Extensions.Logging;
+
+#nullable enable
+
+namespace HordeTrayApp
+{
+	class ToolCatalog : IToolCatalog, IAsyncDisposable
+	{
+		record class ItemState(string Name, string Description, CurrentToolDeploymentInfo? Current, PendingToolDeploymentInfo? Pending, IToolDeployment? Latest);
+
+		[DebuggerDisplay("{Id}")]
+		class Item : IToolCatalogItem
+		{
+			public ToolId Id { get; }
+
+			readonly ToolCatalog _catalog;
+
+			// Keep a separate copy of the state on the main thread; we will update it explicitly. This prevents tearing of updates on UI threads.
+			public ItemState _mainThreadState;
+			public ToolDeploymentInfo? _mainThreadLatest;
+
+			// Public properties only expose main thread data
+			string IToolCatalogItem.Name => _mainThreadState.Name;
+			string IToolCatalogItem.Description => _mainThreadState.Description;
+			ToolDeploymentInfo? IToolCatalogItem.Latest => _mainThreadLatest;
+			CurrentToolDeploymentInfo? IToolCatalogItem.Current => _mainThreadState.Current;
+			PendingToolDeploymentInfo? IToolCatalogItem.Pending => _mainThreadState.Pending;
+
+			public ItemState _state;
+			public DeploymentTask? _deploymentTask;
+			public event Action? OnItemChanged;
+
+			public Item(ToolCatalog catalog, ToolId id, ItemState state)
+			{
+				_catalog = catalog;
+				Id = id;
+				_state = state;
+
+				_mainThreadState = null!;
+				UpdateMainThreadState();
+			}
+
+			public void UpdateMainThreadState()
+			{
+				_mainThreadState = _state;
+				_mainThreadLatest = (_state.Latest == null) ? null : new ToolDeploymentInfo(_state.Latest.Id, _state.Latest.Version);
+			}
+
+			public void NotifyChanged()
+			{
+				OnItemChanged?.Invoke();
+			}
+
+			public void SetState(ItemState state)
+			{
+				_state = state;
+				OnItemChanged?.Invoke();
+			}
+
+			public void Cancel()
+				=> _catalog.CancelDeployment(this, _mainThreadState.Pending?.Deployment);
+
+			// Installing or uninstalling changes the pending task state immediately, and anyone copying the pending state will always see the value of that state.
+			// UI things may wait on that state, and will get notified from that state.
+			// So we should basically post a notice to the main thread to say we've queued up the next state.
+
+			public void Install()
+				=> _catalog.UpdateDeployment(this, _mainThreadState.Latest);
+
+			public void Uninstall()
+				=> _catalog.UpdateDeployment(this, null);
+		}
+
+#pragma warning disable CA1001 // Cancellation source is disposed by the task itself
+		class DeploymentTask
+#pragma warning restore CA1001
+		{
+			public ToolDeploymentInfo Deployment { get; }
+
+			CancellationTokenSource? _cancellationSource;
+			readonly Task _task;
+
+			public DeploymentTask(ToolDeploymentInfo deployment, Func<CancellationToken, Task> func, DeploymentTask? lastDeploymentTask)
+			{
+				Deployment = deployment;
+
+				CancellationTokenSource cancellationSource = new CancellationTokenSource();
+				_task = Task.Run(() => RunAsync(func, lastDeploymentTask, cancellationSource), CancellationToken.None);
+
+				_cancellationSource = cancellationSource;
+			}
+
+			async Task RunAsync(Func<CancellationToken, Task> func, DeploymentTask? lastDeploymentTask, CancellationTokenSource cancellationSource)
+			{
+				CancellationToken cancellationToken = cancellationSource!.Token;
+				try
+				{
+					if (lastDeploymentTask != null)
+					{
+						lastDeploymentTask.Cancel();
+						await lastDeploymentTask.WaitAsync(CancellationToken.None);
+					}
+					await func(cancellationToken);
+				}
+				catch (OperationCanceledException)
+				{
+				}
+				finally
+				{
+					lock (_task)
+					{
+						cancellationSource.Dispose(); // Note; the instance passed into this function, not the class member.
+						_cancellationSource = null;
+					}
+				}
+			}
+
+			public Task WaitAsync(CancellationToken cancellationToken)
+				=> _task.WaitAsync(cancellationToken);
+
+			public void Cancel()
+			{
+				lock (_task)
+				{
+					_cancellationSource?.Cancel();
+					_cancellationSource = null;
+				}
+			}
+		}
+
+		#region Json serialization
+
+		class JsonState
+		{
+			public List<JsonItemState> Items { get; set; } = new List<JsonItemState>();
+		}
+
+		class JsonItemState
+		{
+			public ToolId Id { get; set; }
+			public string Name { get; set; }
+			public string Description { get; set; }
+			public JsonDeploymentState? Current { get; set; }
+
+			public JsonItemState()
+			{
+				Name = String.Empty;
+				Description = String.Empty;
+			}
+
+			public JsonItemState(ToolId id, ItemState state)
+			{
+				Id = id;
+				Name = state.Name;
+				Description = state.Description;
+				Current = (state.Current == null) ? null : new JsonDeploymentState(state.Current);
+			}
+		}
+
+		class JsonDeploymentState
+		{
+			public ToolDeploymentId Id { get; set; }
+			public string Version { get; set; }
+
+			public JsonDeploymentState()
+			{
+				Version = String.Empty;
+			}
+
+			public JsonDeploymentState(ToolDeploymentInfo current)
+			{
+				Id = current.Id;
+				Version = current.Version;
+			}
+
+			public ToolDeploymentInfo ToToolDeploymentInfo()
+				=> new ToolDeploymentInfo(Id, Version);
+		}
+
+		#endregion
+
+		readonly IHordeClientProvider _hordeClientProvider;
+		readonly ILogger _logger;
+		readonly object _workerThreadLockObject = new object();
+		readonly DirectoryReference _baseDir;
+		readonly FileReference _configFile;
+		readonly AsyncEvent _updateEvent;
+		readonly BackgroundTask _updateTask;
+		readonly AsyncEvent _cleanupEvent;
+		readonly BackgroundTask _cleanupTask;
+
+		readonly SynchronizationContext? _mainThreadSynchronizationContext;
+
+		IReadOnlyDictionary<ToolId, Item> _mainThreadItems;
+		IReadOnlyList<IToolCatalogItem> _mainThreadItemList;
+		IReadOnlyList<IToolCatalogItem> IToolCatalog.Items => _mainThreadItemList;
+
+		IReadOnlyDictionary<ToolId, Item> _items;
+
+		Uri? _serverUri;
+
+		public DirectoryReference BaseDir => _baseDir;
+
+		public event Action? OnItemsChanged;
+
+		static readonly JsonSerializerOptions s_jsonSerializerOptions = new JsonSerializerOptions { AllowTrailingCommas = true, PropertyNameCaseInsensitive = true, PropertyNamingPolicy = JsonNamingPolicy.CamelCase, WriteIndented = true };
+
+		public ToolCatalog(IHordeClientProvider hordeClientProvider, ILogger<ToolCatalog> logger)
+		{
+			DirectoryReference? toolsDir = DirectoryReference.GetSpecialFolder(Environment.SpecialFolder.LocalApplicationData);
+			toolsDir ??= DirectoryReference.GetCurrentDirectory();
+
+			_hordeClientProvider = hordeClientProvider;
+			_logger = logger;
+			_baseDir = DirectoryReference.Combine(toolsDir, "Horde", "Launcher");
+			_configFile = FileReference.Combine(_baseDir, "state.json");
+			_updateEvent = new AsyncEvent();
+			_updateTask = new BackgroundTask(UpdateAsync);
+			_cleanupEvent = new AsyncEvent();
+			_cleanupTask = new BackgroundTask(CleanupAsync);
+			_mainThreadSynchronizationContext = SynchronizationContext.Current;
+
+			_items = LoadState();
+			_mainThreadItems = null!;
+			_mainThreadItemList = null!;
+
+			UpdateMainThread();
+		}
+
+		void PostMainThreadUpdate()
+		{
+			if (_mainThreadSynchronizationContext == null)
+			{
+				UpdateMainThread();
+			}
+			else
+			{
+				_mainThreadSynchronizationContext.Post(_ => UpdateMainThread(), null);
+			}
+		}
+
+		void UpdateMainThread()
+		{
+			bool updatedCollection = false;
+			List<Item> updatedItems = new List<Item>();
+
+			lock (_workerThreadLockObject)
+			{
+				if (_items != _mainThreadItems)
+				{
+					UpdateMainThreadItemsState();
+					updatedCollection = true;
+				}
+
+				foreach (Item item in _items.Values)
+				{
+					if (item._state != item._mainThreadState)
+					{
+						item.UpdateMainThreadState();
+						updatedItems.Add(item);
+					}
+				}
+			}
+
+			if (updatedCollection)
+			{
+				OnItemsChanged?.Invoke();
+			}
+			foreach (Item updatedItem in updatedItems)
+			{
+				updatedItem.NotifyChanged();
+			}
+		}
+
+		void UpdateMainThreadItemsState()
+		{
+			_mainThreadItems = _items;
+			_mainThreadItemList = _items.Values.OrderBy(x => x._mainThreadState.Name).ToArray();
+		}
+
+		public void Start()
+		{
+			_updateTask.Start();
+			_cleanupTask.Start();
+
+			_hordeClientProvider.OnStateChanged += OnServerChange;
+		}
+
+		public async Task StopAsync(CancellationToken cancellationToken)
+		{
+			_hordeClientProvider.OnStateChanged -= OnServerChange;
+
+			foreach (Item item in _items.Values)
+			{
+				DeploymentTask? deploymentTask = item._deploymentTask;
+				if (deploymentTask != null)
+				{
+					deploymentTask.Cancel();
+					await deploymentTask.WaitAsync(cancellationToken);
+				}
+			}
+
+			await _updateTask.StopAsync(cancellationToken);
+			await _cleanupTask.StopAsync(cancellationToken);
+		}
+
+		public async ValueTask DisposeAsync()
+		{
+			await StopAsync(CancellationToken.None);
+			await _updateTask.DisposeAsync();
+			await _cleanupTask.DisposeAsync();
+		}
+
+		void OnServerChange()
+		{
+			_updateEvent.Pulse();
+		}
+
+		DirectoryReference GetToolDir(ToolId toolId, ToolDeploymentId deploymentId)
+		{
+			return DirectoryReference.Combine(_baseDir, toolId.ToString(), deploymentId.ToString());
+		}
+
+		ToolConfig LoadToolConfig(DirectoryReference toolDir)
+		{
+			FileReference configFile = FileReference.Combine(toolDir, "Launcher.json");
+			if (FileReference.Exists(configFile))
+			{
+				JsonSerializerOptions jsonSerializerOptions = new JsonSerializerOptions();
+				jsonSerializerOptions.PropertyNameCaseInsensitive = true;
+				jsonSerializerOptions.AllowTrailingCommas = true;
+				jsonSerializerOptions.ReadCommentHandling = JsonCommentHandling.Skip;
+
+				try
+				{
+					byte[] data = FileReference.ReadAllBytes(configFile);
+					return JsonSerializer.Deserialize<ToolConfig>(data, jsonSerializerOptions) ?? new ToolConfig();
+				}
+				catch (Exception ex)
+				{
+					_logger.LogError(ex, "Error reading {File}: {Message}", configFile, ex.Message);
+				}
+			}
+			return new ToolConfig();
+		}
+
+		IReadOnlyDictionary<ToolId, Item> LoadState()
+		{
+			Dictionary<ToolId, Item> items = new Dictionary<ToolId, Item>();
+
+			byte[]? data = FileTransaction.ReadAllBytes(_configFile);
+			if (data != null)
+			{
+				JsonState? state = JsonSerializer.Deserialize<JsonState>(data, s_jsonSerializerOptions);
+				if (state != null)
+				{
+					foreach (JsonItemState jsonItem in state.Items)
+					{
+						CurrentToolDeploymentInfo? current = null;
+						if (jsonItem.Current != null)
+						{
+							DirectoryReference toolDir = GetToolDir(jsonItem.Id, jsonItem.Current.Id);
+							ToolConfig toolConfig = LoadToolConfig(toolDir);
+							current = new CurrentToolDeploymentInfo(jsonItem.Current.Id, jsonItem.Current.Version, toolDir, toolConfig);
+						}
+
+						ItemState itemState = new ItemState(jsonItem.Name, jsonItem.Description, current, null, null);
+						items.Add(jsonItem.Id, new Item(this, jsonItem.Id, itemState));
+					}
+				}
+			}
+
+			return items;
+		}
+
+		async Task SaveStateAsync(CancellationToken cancellationToken)
+		{
+			JsonState state;
+			lock (_workerThreadLockObject)
+			{
+				state = new JsonState();
+				foreach (Item item in _items.Values)
+				{
+					if (item._state.Current != null)
+					{
+						state.Items.Add(new JsonItemState(item.Id, item._state));
+					}
+				}
+			}
+
+			byte[] data = JsonSerializer.SerializeToUtf8Bytes(state, s_jsonSerializerOptions);
+			await FileTransaction.WriteAllBytesAsync(_configFile, data, cancellationToken);
+		}
+
+		void CancelDeployment(Item item, ToolDeploymentInfo? deployment)
+		{
+			if (deployment != null)
+			{
+				lock (_workerThreadLockObject)
+				{
+					if (item._deploymentTask != null && item._deploymentTask.Deployment == deployment)
+					{
+						item._deploymentTask.Cancel();
+					}
+				}
+			}
+		}
+
+		void UpdateDeployment(Item item, IToolDeployment? deployment)
+		{
+			lock (_workerThreadLockObject)
+			{
+				DeploymentTask? lastDeploymentTask = item._deploymentTask;
+				if (deployment != null)
+				{
+					ToolDeploymentInfo deploymentInfo = new ToolDeploymentInfo(deployment.Id, deployment.Version);
+					item._deploymentTask = new DeploymentTask(deploymentInfo, ctx => UpdateCurrentDeploymentAsync(item, deployment, deploymentInfo, ctx), lastDeploymentTask);
+				}
+				else if (item._state.Current != null)
+				{
+					ToolDeploymentInfo deploymentInfo = item._state.Current;
+					item._deploymentTask = new DeploymentTask(deploymentInfo, ctx => UpdateCurrentDeploymentAsync(item, deployment, deploymentInfo, ctx), lastDeploymentTask);
+				}
+			}
+		}
+
+		async Task UpdateCurrentDeploymentAsync(Item item, IToolDeployment? deployment, ToolDeploymentInfo deploymentInfo, CancellationToken cancellationToken)
+		{
+			try
+			{
+				if (deployment == null)
+				{
+					UpdateStateAndNotify(item, state => state with { Current = null });
+					_cleanupEvent.Set();
+				}
+				else
+				{
+					UpdateStateAndNotify(item, state => state with { Pending = new PendingToolDeploymentInfo(false, $"Updating to {deployment.Version}", deploymentInfo) });
+
+					DirectoryReference targetDir = DirectoryReference.Combine(_baseDir, item.Id.ToString(), deployment.Id.ToString());
+					DirectoryReference.CreateDirectory(targetDir);
+					await deployment.Content.ExtractAsync(targetDir.ToDirectoryInfo(), _logger, cancellationToken);
+
+					DirectoryReference toolDir = GetToolDir(item.Id, deploymentInfo.Id);
+					ToolConfig toolConfig = LoadToolConfig(toolDir);
+					CurrentToolDeploymentInfo current = new CurrentToolDeploymentInfo(deploymentInfo.Id, deploymentInfo.Version, toolDir, toolConfig);
+
+					UpdateStateAndNotify(item, state => state with { Pending = null, Current = current });
+				}
+
+				await SaveStateAsync(cancellationToken);
+			}
+			catch (OperationCanceledException)
+			{
+				_logger.LogInformation("Deployment cancelled");
+				UpdateStateAndNotify(item, state => state with { Pending = new PendingToolDeploymentInfo(true, $"Cancelled.", deploymentInfo) });
+			}
+			catch (Exception ex)
+			{
+				_logger.LogError(ex, "Error updating deployment: {Message}", ex.Message);
+				UpdateStateAndNotify(item, state => state with { Pending = new PendingToolDeploymentInfo(true, $"Error updating. See log.", deploymentInfo) });
+			}
+		}
+
+		void UpdateStateAndNotify(Item item, Func<ItemState, ItemState> update)
+		{
+			lock (_workerThreadLockObject)
+			{
+				item._state = update(item._state);
+			}
+			PostMainThreadUpdate();
+		}
+
+		#region Cleanup
+
+		async Task CleanupAsync(CancellationToken cancellationToken)
+		{
+			while (!cancellationToken.IsCancellationRequested)
+			{
+				Task cleanupTask = _cleanupEvent.Task;
+
+				bool success = false;
+				try
+				{
+					success = CleanupOnce();
+				}
+				catch (OperationCanceledException)
+				{
+					break;
+				}
+				catch (Exception ex)
+				{
+					_logger.LogError(ex, "Exception while running cleanup over installed tools: {Message}", ex.Message);
+				}
+
+				Task waitTask = cleanupTask.WaitAsync(cancellationToken);
+				if (!success)
+				{
+					waitTask = Task.WhenAny(waitTask, Task.Delay(TimeSpan.FromMinutes(1.0), cancellationToken));
+				}
+
+				await waitTask;
+			}
+		}
+
+		bool CleanupOnce()
+		{
+			bool result = true;
+
+			List<DirectoryInfo> deleteToolDirs = new List<DirectoryInfo>();
+			List<DirectoryInfo> deleteDeploymentDirs = new List<DirectoryInfo>();
+
+			DirectoryInfo baseDirInfo = _baseDir.ToDirectoryInfo();
+			if (baseDirInfo.Exists)
+			{
+				foreach (DirectoryInfo toolDir in baseDirInfo.EnumerateDirectories())
+				{
+					if (!toolDir.Name.StartsWith('.'))
+					{
+						DirectoryInfo[] deploymentDirs = toolDir.GetDirectories();
+
+						ToolId? toolId;
+						try
+						{
+							toolId = new ToolId(toolDir.Name);
+						}
+						catch
+						{
+							toolId = null;
+						}
+
+						lock (_workerThreadLockObject)
+						{
+							ItemState? itemState = null;
+							if (toolId != null && _items.TryGetValue(toolId.Value, out Item? item))
+							{
+								itemState = item._state;
+							}
+
+							bool deleteToolDir = true;
+							for (int idx = 0; idx < deploymentDirs.Length; idx++)
+							{
+								DirectoryInfo deploymentDir = deploymentDirs[idx];
+								if (!deploymentDir.Name.StartsWith('.'))
+								{
+									if (!BinaryId.TryParse(deploymentDir.Name, out BinaryId deploymentId))
+									{
+										deleteDeploymentDirs.Add(deploymentDir);
+									}
+									else if (itemState != null && itemState.Current != null && itemState.Current.Id.Id == deploymentId)
+									{
+										deleteToolDir = false;
+									}
+									else if (itemState != null && itemState.Pending?.Deployment != null && itemState.Pending.Deployment.Id.Id == deploymentId)
+									{
+										deleteToolDir = false;
+									}
+									else if (TryRenameDir(deploymentDir, deploymentDir.FullName + ".dead"))
+									{
+										deleteDeploymentDirs.Add(deploymentDir);
+									}
+									else
+									{
+										result = deleteToolDir = false;
+									}
+								}
+							}
+
+							if (deleteToolDir)
+							{
+								deleteToolDirs.Add(toolDir);
+							}
+						}
+					}
+				}
+			}
+
+			// Try to delete all the deployment directories
+			foreach (DirectoryInfo deleteDeploymentDir in deleteDeploymentDirs)
+			{
+				result &= TryDeleteDir(deleteDeploymentDir, true);
+			}
+
+			// Try to delete all the tool directories if they're empty
+			foreach (DirectoryInfo deleteToolDir in deleteToolDirs)
+			{
+				result &= TryDeleteDir(deleteToolDir, false);
+			}
+
+			return result;
+		}
+
+		bool TryRenameDir(DirectoryInfo dirInfo, string newName)
+		{
+			try
+			{
+				dirInfo.MoveTo(newName);
+				return true;
+			}
+			catch (Exception ex)
+			{
+				_logger.LogWarning(ex, "Unable to rename directory {DirName}: {Message}", dirInfo.FullName, ex.Message);
+				return false;
+			}
+		}
+
+		bool TryDeleteDir(DirectoryInfo dirInfo, bool recursive)
+		{
+			try
+			{
+				dirInfo.Delete(recursive);
+				return true;
+			}
+			catch (Exception ex)
+			{
+				_logger.LogWarning(ex, "Unable to delete directory {DirName}: {Message}", dirInfo.FullName, ex.Message);
+				return false;
+			}
+		}
+
+		#endregion
+
+		#region Updates
+
+		async Task UpdateAsync(CancellationToken cancellationToken)
+		{
+			while (!cancellationToken.IsCancellationRequested)
+			{
+				Task updateTask = _updateEvent.Task;
+				try
+				{
+					await PollForUpdatesOnceAsync(cancellationToken);
+				}
+				catch (OperationCanceledException)
+				{
+					break;
+				}
+				catch (Exception ex)
+				{
+					_logger.LogError(ex, "Exception while checking for tool updates: {Message}", ex.Message);
+				}
+				await Task.WhenAny(updateTask, Task.Delay(TimeSpan.FromMinutes(5.0), cancellationToken));
+			}
+		}
+
+		async Task PollForUpdatesOnceAsync(CancellationToken cancellationToken)
+		{
+			using IHordeClientRef? hordeClientRef = await _hordeClientProvider.GetClientRefAsync();
+			if (hordeClientRef == null)
+			{
+				SetTools(Enumerable.Empty<ITool>());
+			}
+			else
+			{
+				if (_serverUri != hordeClientRef.Client.ServerUrl)
+				{
+					SetTools(Enumerable.Empty<ITool>());
+					_serverUri = hordeClientRef.Client.ServerUrl;
+				}
+
+				IEnumerable<ITool> tools = await hordeClientRef.Client.Tools.GetAllAsync(cancellationToken);
+				SetTools(tools);
+			}
+		}
+
+		void SetTools(IEnumerable<ITool> tools)
+		{
+			lock (_workerThreadLockObject)
+			{
+				Dictionary<ToolId, Item> newItems = new Dictionary<ToolId, Item>();
+				foreach (ITool tool in tools)
+				{
+					//				if (tool.ShowInUgs)
+					{
+						IToolDeployment? latestDeployment = (tool.Deployments.Count > 0) ? tool.Deployments[^1] : null;
+						if (_items.TryGetValue(tool.Id, out Item? item))
+						{
+							item._state = new ItemState(tool.Name, tool.Description, item._state.Current, item._state.Pending, latestDeployment);
+						}
+						else
+						{
+							item = new Item(this, tool.Id, new ItemState(tool.Name, tool.Description, null, null, latestDeployment));
+						}
+						newItems[item.Id] = item;
+					}
+				}
+				foreach (Item existingItem in _items.Values)
+				{
+					if (existingItem._state.Current != null && !newItems.ContainsKey(existingItem.Id))
+					{
+						newItems.Add(existingItem.Id, existingItem);
+					}
+				}
+				_items = newItems;
+			}
+			PostMainThreadUpdate();
+		}
+
+		#endregion
+	}
+}
diff --git a/Engine/Source/Programs/Horde/HordeTrayApp/ToolCatalogViewModel.cs b/Engine/Source/Programs/Horde/HordeTrayApp/ToolCatalogViewModel.cs
new file mode 100644
--- /dev/null
+++ b/Engine/Source/Programs/Horde/HordeTrayApp/ToolCatalogViewModel.cs
@@ -0,0 +1,158 @@
+// Copyright Epic Games, Inc. All Rights Reserved.
+
+using System.Collections.ObjectModel;
+using CommunityToolkit.Mvvm.ComponentModel;
+
+namespace HordeTrayApp
+{
+	partial class ToolCatalogViewModel : ObservableObject, IDisposable
+	{
+		readonly IToolCatalog? _catalog;
+		readonly Dictionary<IToolCatalogItem, ToolCatalogItemViewModel> _items = new Dictionary<IToolCatalogItem, ToolCatalogItemViewModel>();
+
+		public ObservableCollection<ToolCatalogItemViewModel> Items { get; } = new ObservableCollection<ToolCatalogItemViewModel>();
+
+		public bool ShowTools => Items.Count > 0;
+
+		public ToolCatalogViewModel()
+			: this(null)
+		{ }
+
+		public ToolCatalogViewModel(IToolCatalog? catalog)
+		{
+			if (catalog != null)
+			{
+				_catalog = catalog;
+				_catalog.OnItemsChanged += RefreshItems;
+
+				Items.CollectionChanged += (_, _) => OnPropertyChanged(nameof(ShowTools));
+
+				RefreshItems();
+			}
+		}
+
+		public void Dispose()
+		{
+			foreach (ToolCatalogItemViewModel item in _items.Values)
+			{
+				item.Dispose();
+			}
+			if (_catalog != null)
+			{
+				_catalog.OnItemsChanged -= RefreshItems;
+			}
+		}
+
+		void RefreshItems()
+		{
+			if (_catalog != null)
+			{
+				Items.Clear();
+
+				HashSet<IToolCatalogItem> newItemSet = new HashSet<IToolCatalogItem>(_catalog.Items);
+				foreach ((IToolCatalogItem item, ToolCatalogItemViewModel itemViewModel) in _items.ToArray())
+				{
+					if (!newItemSet.Contains(item))
+					{
+						_items.Remove(item);
+						itemViewModel.Dispose();
+					}
+				}
+
+				foreach (IToolCatalogItem item in _catalog.Items)
+				{
+					ToolCatalogItemViewModel? itemViewModel;
+					if (!_items.TryGetValue(item, out itemViewModel))
+					{
+						itemViewModel = new ToolCatalogItemViewModel(item);
+						_items.Add(item, itemViewModel);
+					}
+					Items.Add(itemViewModel);
+				}
+			}
+		}
+	}
+
+	partial class ToolCatalogItemViewModel : ObservableObject, IDisposable
+	{
+		readonly IToolCatalogItem _item;
+
+		public ToolCatalogItemViewModel(IToolCatalogItem item)
+		{
+			_item = item;
+			_item.OnItemChanged += UpdateFromModel;
+			UpdateFromModel();
+		}
+
+		public void Dispose()
+		{
+			_item.OnItemChanged -= UpdateFromModel;
+		}
+
+		[ObservableProperty]
+		string _name = String.Empty;
+
+		[ObservableProperty]
+		string _description = String.Empty;
+
+		[ObservableProperty]
+		string _status = String.Empty;
+
+		[ObservableProperty]
+		bool _showAdd;
+
+		[ObservableProperty]
+		bool _showUpdate;
+
+		[ObservableProperty]
+		bool _showRemove;
+
+		internal void UpdateFromModel()
+		{
+			Name = _item.Name;
+			Description = _item.Description;
+
+			PendingToolDeploymentInfo? pending = _item.Pending;
+			if (pending != null)
+			{
+				if (pending.Message != null)
+				{
+					Status = pending.Message;
+				}
+				else if (pending.Deployment == null)
+				{
+					Status = "Removing...";
+				}
+				else
+				{
+					Status = $"Updating to {pending.Deployment.Version}...";
+				}
+			}
+			else
+			{
+				ToolDeploymentInfo? current = _item.Current;
+				if (current != null)
+				{
+					Status = current.Version;
+				}
+				else
+				{
+					Status = String.Empty;
+				}
+			}
+
+			ShowAdd = _item.Latest != null && _item.Current == null;
+			ShowUpdate = _item.Latest != null && _item.Current != null && _item.Latest.Id != _item.Current.Id;
+			ShowRemove = _item.Current != null;
+		}
+
+		public void Add()
+			=> _item.Install();
+
+		public void Remove()
+			=> _item.Uninstall();
+
+		public void Update()
+			=> _item.Install();
+	}
+}
diff --git a/Engine/Source/Programs/Horde/HordeTrayApp/ToolConfig.cs b/Engine/Source/Programs/Horde/HordeTrayApp/ToolConfig.cs
new file mode 100644
--- /dev/null
+++ b/Engine/Source/Programs/Horde/HordeTrayApp/ToolConfig.cs
@@ -0,0 +1,52 @@
+// Copyright Epic Games, Inc. All Rights Reserved.
+
+namespace HordeTrayApp
+{
+	/// <summary>
+	/// Configuration for a tool that can be shown in the launcher
+	/// </summary>
+	public class ToolConfig
+	{
+		/// <summary>
+		/// Popup menu item for the tool
+		/// </summary>
+		public ToolMenuItem? PopupMenu { get; set; }
+	}
+
+	/// <summary>
+	/// Menu item for the tool
+	/// </summary>
+	public class ToolMenuItem
+	{
+		/// <summary>
+		/// Text to display for the menu item. If null, or a sequence of hyphens, will create a menu separator.
+		/// </summary>
+		public string? Label { get; set; }
+
+		/// <summary>
+		/// Command to run when the menu item is clicked. Will be executed in the root directory of the downloaded tool.
+		/// </summary>
+		public ToolMenuCommand? Command { get; set; }
+
+		/// <summary>
+		/// Child menu items. Cannot be specified at the same time as a tool.
+		/// </summary>
+		public List<ToolMenuItem>? Children { get; set; }
+	}
+
+	/// <summary>
+	/// Command to run for a menu item.
+	/// </summary>
+	public class ToolMenuCommand
+	{
+		/// <summary>
+		/// Executable to run
+		/// </summary>
+		public string FileName { get; set; } = "cmd.exe";
+
+		/// <summary>
+		/// Command line arguments for the tool
+		/// </summary>
+		public List<string>? Arguments { get; set; }
+	}
+}
diff --git a/Engine/Source/Programs/Shared/EpicGames.Horde.Tests/ServerLogPacketBuilderTests.cs b/Engine/Source/Programs/Shared/EpicGames.Horde.Tests/ServerLogPacketBuilderTests.cs
--- a/Engine/Source/Programs/Shared/EpicGames.Horde.Tests/ServerLogPacketBuilderTests.cs
+++ b/Engine/Source/Programs/Shared/EpicGames.Horde.Tests/ServerLogPacketBuilderTests.cs
@@ -211,7 +211,7 @@ namespace EpicGames.Horde.Tests
 			}
 		}
 
-		class FakeHordeClient : IHordeClient
+		class FakeHordeClient : IHordeClient, IAsyncDisposable
 		{
 			public FakeLogRpcClient LogRpc { get; } = new FakeLogRpcClient();
 			public Dictionary<string, BundleStorageNamespace> StorageNamespaces { get; } = new Dictionary<string, BundleStorageNamespace>();
@@ -224,6 +224,8 @@ namespace EpicGames.Horde.Tests
 			public ISecretCollection Secrets => throw new NotImplementedException();
 			public IToolCollection Tools => throw new NotImplementedException();
 
+			public event Action? OnAccessTokenStateChanged;
+
 			public Task<bool> LoginAsync(bool allowLogin, CancellationToken cancellationToken)
 				=> throw new NotImplementedException();
 
diff --git a/Engine/Source/Programs/Shared/EpicGames.Horde/HordeClient.cs b/Engine/Source/Programs/Shared/EpicGames.Horde/HordeClient.cs
--- a/Engine/Source/Programs/Shared/EpicGames.Horde/HordeClient.cs
+++ b/Engine/Source/Programs/Shared/EpicGames.Horde/HordeClient.cs
@@ -60,6 +60,9 @@ namespace EpicGames.Horde
 		/// <inheritdoc/>
 		public IToolCollection Tools { get; }
 
+		/// <inheritdoc/>
+		public event Action? OnAccessTokenStateChanged;
+
 		/// <summary>
 		/// Accessor for the logger instance
 		/// </summary>
@@ -96,6 +99,14 @@ namespace EpicGames.Horde
 			GC.SuppressFinalize(this);
 		}
 
+		/// <summary>
+		/// Notify listeners that the auth state has changed
+		/// </summary>
+		protected void NotifyAuthStateChanged()
+		{
+			OnAccessTokenStateChanged?.Invoke();
+		}
+
 		/// <inheritdoc/>
 		public abstract Task<bool> LoginAsync(bool allowLogin, CancellationToken cancellationToken);
 
@@ -331,11 +342,15 @@ namespace EpicGames.Horde
 
 			_authHandlerState = new HordeHttpAuthHandlerState(_baseHttpMessageHandler.Instance, serverUrl, hordeOptions, loggerFactory.CreateLogger<HordeHttpAuthHandlerState>());
 			_authHttpMessageHandler = new HordeHttpAuthHandler(_baseHttpMessageHandler.Instance, _authHandlerState, hordeOptions);
+
+			_authHandlerState.OnStateChanged += NotifyAuthStateChanged;
 		}
 
 		/// <inheritdoc/>
 		public override async ValueTask DisposeAsync()
 		{
+			_authHandlerState.OnStateChanged -= NotifyAuthStateChanged;
+
 			await _authHandlerState.DisposeAsync();
 			_authHttpMessageHandler.Dispose();
 			await base.DisposeAsync();
diff --git a/Engine/Source/Programs/Shared/EpicGames.Horde/HordeHttpAuthHandler.cs b/Engine/Source/Programs/Shared/EpicGames.Horde/HordeHttpAuthHandler.cs
--- a/Engine/Source/Programs/Shared/EpicGames.Horde/HordeHttpAuthHandler.cs
+++ b/Engine/Source/Programs/Shared/EpicGames.Horde/HordeHttpAuthHandler.cs
@@ -107,6 +107,11 @@ namespace EpicGames.Horde
 		readonly IOidcTokenManager? _oidcTokenManager;
 		readonly IClock _clock;
 
+		/// <summary>
+		/// Event handler for the auth state changing
+		/// </summary>
+		public event Action? OnStateChanged;
+
 		/// <summary>
 		/// Constructor
 		/// </summary>
@@ -262,6 +267,7 @@ namespace EpicGames.Horde
 					if (_authStateTask == null || _authStateTask == authStateTask)
 					{
 						_authStateTask = Task.Run(() => GetAuthStateInternalAsync(interactive, _cancellationTokenSource.Token), _cancellationTokenSource.Token);
+						_authStateTask.ContinueWith(_ => OnStateChanged?.Invoke(), CancellationToken.None, TaskContinuationOptions.None, TaskScheduler.Default);
 					}
 					authStateTask = _authStateTask;
 				}
diff --git a/Engine/Source/Programs/Shared/EpicGames.Horde/IHordeClient.cs b/Engine/Source/Programs/Shared/EpicGames.Horde/IHordeClient.cs
--- a/Engine/Source/Programs/Shared/EpicGames.Horde/IHordeClient.cs
+++ b/Engine/Source/Programs/Shared/EpicGames.Horde/IHordeClient.cs
@@ -16,12 +16,12 @@ using Microsoft.Extensions.Logging;
 namespace EpicGames.Horde
 {
 	/// <summary>
-	/// Base interface for Horde functionality.
+	/// Interface for Horde functionality.
 	/// </summary>
 	public interface IHordeClient
 	{
 		/// <summary>
-		/// Base URL of the horde server
+		/// URL of the horde server
 		/// </summary>
 		Uri ServerUrl { get; }
 
@@ -50,6 +50,11 @@ namespace EpicGames.Horde
 		/// </summary>
 		IToolCollection Tools { get; }
 
+		/// <summary>
+		/// Event triggered whenever the access token state changes
+		/// </summary>
+		event Action? OnAccessTokenStateChanged;
+
 		/// <summary>
 		/// Connect to the Horde server
 		/// </summary>
