Horde: Filter out launcher apps which are not valid for the current platform.

#rnx

diff --git a/Engine/Source/Programs/Horde/HordeTrayApp/ToolCatalog.cs b/Engine/Source/Programs/Horde/HordeTrayApp/ToolCatalog.cs
--- a/Engine/Source/Programs/Horde/HordeTrayApp/ToolCatalog.cs
+++ b/Engine/Source/Programs/Horde/HordeTrayApp/ToolCatalog.cs
@@ -1,6 +1,7 @@
 // Copyright Epic Games, Inc. All Rights Reserved.
 
 using System.Diagnostics;
+using System.Runtime.InteropServices;
 using System.Text.Json;
 using EpicGames.Core;
 using EpicGames.Horde;
@@ -678,7 +679,7 @@ namespace HordeTrayApp
 				Dictionary<ToolId, Item> newItems = new Dictionary<ToolId, Item>();
 				foreach (ITool tool in tools)
 				{
-					//				if (tool.ShowInUgs)
+					if (ShouldIncludeTool(tool))
 					{
 						IToolDeployment? latestDeployment = (tool.Deployments.Count > 0) ? tool.Deployments[^1] : null;
 						if (_items.TryGetValue(tool.Id, out Item? item))
@@ -704,6 +705,66 @@ namespace HordeTrayApp
 			PostMainThreadUpdate();
 		}
 
+		static IReadOnlySet<string> HostPlatformRids { get; } = GetHostPlatformRids();
+
+		static bool ShouldIncludeTool(ITool tool)
+		{
+			if (tool.Platforms != null && !tool.Platforms.Any(x => HostPlatformRids.Contains(x)))
+			{
+				return false;
+			}
+			return true; // TODO: ShowInUgs?
+		}
+
+		static HashSet<string> GetHostPlatformRids()
+		{
+			HashSet<string> platforms = new HashSet<string>(StringComparer.OrdinalIgnoreCase);
+			platforms.Add("any");
+			if (OperatingSystem.IsWindows())
+			{
+				platforms.Add("win");
+				switch (RuntimeInformation.OSArchitecture)
+				{
+					case Architecture.X64:
+						platforms.Add("win-x64");
+						break;
+					case Architecture.Arm64:
+						platforms.Add("win-arm64");
+						break;
+				}
+			}
+			else if (OperatingSystem.IsMacOS())
+			{
+				platforms.Add("osx");
+				switch (RuntimeInformation.OSArchitecture)
+				{
+					case Architecture.X64:
+						platforms.Add("osx-x64");
+						break;
+					case Architecture.Arm64:
+						platforms.Add("osx-arm64");
+						break;
+				}
+			}
+			else if (OperatingSystem.IsLinux())
+			{
+				platforms.Add("linux");
+				platforms.Add("unix");
+				switch (RuntimeInformation.OSArchitecture)
+				{
+					case Architecture.X64:
+						platforms.Add("linux-x64");
+						platforms.Add("unix-x64");
+						break;
+					case Architecture.Arm64:
+						platforms.Add("linux-arm64");
+						platforms.Add("unix-arm64");
+						break;
+				}
+			}
+			return platforms;
+		}
+
 		#endregion
 	}
 }
