Horde: Build fix.

#rnx

diff --git a/Engine/Source/Programs/Horde/Drivers/JobDriver.Tests/Execution/WorkspaceExecutorTest.cs b/Engine/Source/Programs/Horde/Drivers/JobDriver.Tests/Execution/WorkspaceExecutorTest.cs
--- a/Engine/Source/Programs/Horde/Drivers/JobDriver.Tests/Execution/WorkspaceExecutorTest.cs
+++ b/Engine/Source/Programs/Horde/Drivers/JobDriver.Tests/Execution/WorkspaceExecutorTest.cs
@@ -39,6 +39,12 @@ namespace JobDriver.Tests.Execution
 			private readonly Dictionary<StreamId, RpcGetStreamResponse> _streamIdToStreamResponse = new();
 			private readonly Dictionary<JobId, RpcGetJobResponse> _jobIdToJobResponse = new();
 
+			public event Action? OnAccessTokenStateChanged
+			{
+				add { }
+				remove { }
+			}
+
 			public void AddStream(StreamId streamId, string streamName)
 			{
 				if (_streamIdToStreamResponse.ContainsKey(streamId))
