Horde: Add a proxy server to the Horde tray app, allowing users to connect to the Horde server via an unauthenticated connection bound to localhost.

diff --git a/Engine/Source/Programs/Horde/HordeTrayApp/App.axaml.cs b/Engine/Source/Programs/Horde/HordeTrayApp/App.axaml.cs
--- a/Engine/Source/Programs/Horde/HordeTrayApp/App.axaml.cs
+++ b/Engine/Source/Programs/Horde/HordeTrayApp/App.axaml.cs
@@ -8,6 +8,7 @@ using Avalonia.Markup.Xaml;
 using Avalonia.Threading;
 using EpicGames.Core;
 using HordeTrayApp.Plugins.HordeAgent;
+using HordeTrayApp.Plugins.HordeProxy;
 using Microsoft.Extensions.DependencyInjection;
 
 namespace HordeTrayApp
@@ -40,6 +41,7 @@ namespace HordeTrayApp
 
 			// Configure plugins
 			serviceCollection.AddSingleton<ITrayAppPlugin, HordeAgentPlugin>();
+			serviceCollection.AddSingleton<ITrayAppPlugin, HordeProxyPlugin>();
 
 			_serviceProvider = serviceCollection.BuildServiceProvider();
 		}
diff --git a/Engine/Source/Programs/Horde/HordeTrayApp/HordeTrayApp.csproj b/Engine/Source/Programs/Horde/HordeTrayApp/HordeTrayApp.csproj
--- a/Engine/Source/Programs/Horde/HordeTrayApp/HordeTrayApp.csproj
+++ b/Engine/Source/Programs/Horde/HordeTrayApp/HordeTrayApp.csproj
@@ -1,5 +1,9 @@
 ﻿<Project Sdk="Microsoft.NET.Sdk">
 
+  <ItemGroup>
+    <FrameworkReference Include="Microsoft.AspNetCore.App" />
+  </ItemGroup>
+
   <PropertyGroup>
     <OutputType>WinExe</OutputType>
     <TargetFramework>net8.0-windows</TargetFramework>
@@ -9,7 +13,7 @@
     <Configurations>Debug;Release;Analyze</Configurations>
 	<AssemblyName>Horde.Agent.TrayApp</AssemblyName>
     <AvaloniaUseCompiledBindingsByDefault>true</AvaloniaUseCompiledBindingsByDefault>
-    <NoWarn>IDE0005</NoWarn>
+    <NoWarn>IDE0005;CA1501</NoWarn>
   </PropertyGroup>
 
   <PropertyGroup Condition=" '$(Configuration)|$(Platform)' == 'Analyze|AnyCPU' ">
@@ -18,7 +22,6 @@
     <EnforceCodeStyleInBuild>true</EnforceCodeStyleInBuild>
     <GenerateDocumentationFile>true</GenerateDocumentationFile>
     <DocumentationFile></DocumentationFile>
-    <NoWarn>CA1501</NoWarn>
   </PropertyGroup>
 
   <ItemGroup>
diff --git a/Engine/Source/Programs/Horde/HordeTrayApp/Plugins/HordeAgent/HordeAgentPlugin.cs b/Engine/Source/Programs/Horde/HordeTrayApp/Plugins/HordeAgent/HordeAgentPlugin.cs
--- a/Engine/Source/Programs/Horde/HordeTrayApp/Plugins/HordeAgent/HordeAgentPlugin.cs
+++ b/Engine/Source/Programs/Horde/HordeTrayApp/Plugins/HordeAgent/HordeAgentPlugin.cs
@@ -33,7 +33,7 @@ namespace HordeTrayApp.Plugins.HordeAgent
 
 		public string Name => "Horde Agent";
 
-		public IconSource Icon => new SymbolIconSource() { Symbol = Symbol.Globe };
+		public IconSource Icon => new SymbolIconSource() { Symbol = Symbol.People };
 
 		public Control CreateSettingsPage(SettingsContext context)
 			=> new HordeAgentSettingsPage(context, this);
diff --git a/Engine/Source/Programs/Horde/HordeTrayApp/Plugins/HordeProxy/HordeProxyPlugin.cs b/Engine/Source/Programs/Horde/HordeTrayApp/Plugins/HordeProxy/HordeProxyPlugin.cs
new file mode 100644
--- /dev/null
+++ b/Engine/Source/Programs/Horde/HordeTrayApp/Plugins/HordeProxy/HordeProxyPlugin.cs
@@ -0,0 +1,210 @@
+// Copyright Epic Games, Inc. All Rights Reserved.
+
+using System.Net;
+using System.Text.Json;
+using System.Text.Json.Serialization;
+using Avalonia.Controls;
+using EpicGames.Core;
+using EpicGames.Horde;
+using FluentAvalonia.UI.Controls;
+using Microsoft.AspNetCore.Builder;
+using Microsoft.AspNetCore.Hosting;
+using Microsoft.AspNetCore.Http;
+using Microsoft.Extensions.DependencyInjection;
+using Microsoft.Extensions.Hosting;
+
+namespace HordeTrayApp.Plugins.HordeProxy
+{
+	class HordeProxyPlugin : ITrayAppPlugin
+	{
+		readonly FileReference _settingsFile;
+		readonly BackgroundTask _serverTask;
+		readonly IHordeClientProvider _hordeClientProvider;
+		readonly AsyncEvent _restartServerEvent = new AsyncEvent();
+		readonly SynchronizationContext? _synchronizationContext;
+
+		string _status = "Starting...";
+
+		byte[] _settingsData = Array.Empty<byte>();
+		HordeProxySettings _settings = new HordeProxySettings();
+
+		public HordeProxySettings Settings => _settings;
+
+		public event Action? OnStatusChanged;
+
+		public string Name => "Horde Proxy";
+
+		public string Status => _status;
+
+		public IconSource Icon => new SymbolIconSource() { Symbol = Symbol.Filter };
+
+		public HordeProxyPlugin(IHordeClientProvider hordeClientProvider)
+		{
+			_hordeClientProvider = hordeClientProvider;
+			_hordeClientProvider.OnStateChanged += OnStateChanged;
+			_synchronizationContext = SynchronizationContext.Current;
+
+			DirectoryReference ? settingsRoot = DirectoryReference.GetSpecialFolder(Environment.SpecialFolder.CommonApplicationData);
+			settingsRoot ??= DirectoryReference.GetCurrentDirectory();
+			_settingsFile = FileReference.Combine(settingsRoot, "Epic", "Horde", "Proxy", "Settings.json");
+
+			LoadSettings();
+
+			_serverTask = BackgroundTask.StartNew(RunServerAsync);
+		}
+
+		public Control CreateSettingsPage(SettingsContext context)
+			=> new HordeProxySettingsPage(this);
+
+		public void UpdateSettings(HordeProxySettings settings)
+		{
+			byte[] data = JsonSerializer.SerializeToUtf8Bytes(settings, GetJsonSerializerOptions());
+			if (!data.SequenceEqual(_settingsData))
+			{
+				_settings = settings;
+				_settingsData = data;
+				_restartServerEvent.Pulse();
+
+				DirectoryReference.CreateDirectory(_settingsFile.Directory);
+				FileReference.WriteAllBytes(_settingsFile, data);
+			}
+		}
+
+		public void SetStatus(string status)
+		{
+			_status = status;
+
+			if (_synchronizationContext == null)
+			{
+				OnStatusChanged?.Invoke();
+			}
+			else
+			{
+				_synchronizationContext.Post(_ => OnStatusChanged?.Invoke(), null);
+			}
+		}
+
+		public async ValueTask DisposeAsync()
+		{
+			_hordeClientProvider.OnStateChanged -= OnStateChanged;
+			await _serverTask.DisposeAsync();
+		}
+
+		void OnStateChanged()
+		{
+			_restartServerEvent.Pulse();
+		}
+
+		static JsonSerializerOptions GetJsonSerializerOptions()
+		{
+			JsonSerializerOptions options = new JsonSerializerOptions();
+			options.PropertyNamingPolicy = JsonNamingPolicy.CamelCase;
+			options.PropertyNameCaseInsensitive = true;
+			options.DefaultIgnoreCondition = JsonIgnoreCondition.WhenWritingNull;
+			options.AllowTrailingCommas = true;
+			options.WriteIndented = true;
+			options.Converters.Add(new JsonStringEnumConverter());
+			return options;
+		}
+
+		void LoadSettings()
+		{
+			if (FileReference.Exists(_settingsFile))
+			{
+				try
+				{
+					byte[] data = FileReference.ReadAllBytes(_settingsFile);
+					if (!data.SequenceEqual(_settingsData))
+					{
+						_settings = JsonSerializer.Deserialize<HordeProxySettings>(data, GetJsonSerializerOptions())!;
+						_settingsData = data;
+						_restartServerEvent.Pulse();
+						return;
+					}
+				}
+				catch (Exception)
+				{
+				}
+			}
+
+			UpdateSettings(new HordeProxySettings());
+		}
+
+		public void PopulateContextMenu(NativeMenu contextMenu)
+		{
+		}
+
+		public TrayAppPluginStatus GetStatus()
+			=> TrayAppPluginStatus.Default;
+
+		async Task RunServerAsync(CancellationToken cancellationToken)
+		{
+			while (!cancellationToken.IsCancellationRequested)
+			{
+				Task restartServerTask = _restartServerEvent.Task;
+
+				HordeProxySettings settings = _settings;
+				if (settings.Enabled)
+				{
+					using CancellationTokenSource cancellationSource = CancellationTokenSource.CreateLinkedTokenSource(cancellationToken);
+					Task stoppedTask = _restartServerEvent.Task.ContinueWith(_ => cancellationSource.Cancel(), cancellationSource.Token, TaskContinuationOptions.None, TaskScheduler.Default);
+
+					try
+					{
+						SetStatus($"Running on http://localhost:{settings.Port}");
+
+						WebApplicationBuilder builder = WebApplication.CreateBuilder();
+						builder.Services.AddHttpClient();
+						builder.Services.AddHorde();
+						builder.WebHost.ConfigureKestrel(options => options.ListenLocalhost(settings.Port));
+
+						await using WebApplication app = builder.Build();
+						app.Map("/{*Route}", ForwardRequestAsync);
+						await app.RunAsync(cancellationSource.Token);
+					}
+					catch (OperationCanceledException)
+					{
+						throw;
+					}
+					catch (Exception ex)
+					{
+						SetStatus($"Error: {ex.Message}");
+					}
+
+					await stoppedTask;
+				}
+				else
+				{
+					SetStatus("Stopped");
+					await restartServerTask.WaitAsync(cancellationToken);
+				}
+				await Task.Delay(TimeSpan.FromSeconds(2.0), cancellationToken);
+			}
+		}
+
+		static async Task ForwardRequestAsync(HttpContext context)
+		{
+			IHordeClient hordeClient = context.RequestServices.GetRequiredService<IHordeClient>();
+			HttpClient httpClient = hordeClient.CreateHttpClient().HttpClient;
+
+			if (context.Request.Method != "GET")
+			{
+				context.Response.StatusCode = (int)HttpStatusCode.MethodNotAllowed;
+				return;
+			}
+
+			using HttpRequestMessage request = new HttpRequestMessage();
+			request.Method = HttpMethod.Parse(context.Request.Method);
+			request.RequestUri = new Uri(hordeClient.ServerUrl, $"{context.Request.Path}");
+			request.Content = new StreamContent(context.Request.Body);
+
+			using HttpResponseMessage response = await httpClient.SendAsync(request, context.RequestAborted);
+			context.Response.StatusCode = (int)response.StatusCode;
+			context.Response.ContentType = response.Content.Headers.ContentType?.MediaType;
+			context.Response.ContentLength = response.Content.Headers.ContentLength;
+
+			using Stream stream = await response.Content.ReadAsStreamAsync(context.RequestAborted);
+			await stream.CopyToAsync(context.Response.Body, context.RequestAborted);
+		}
+	}
+}
diff --git a/Engine/Source/Programs/Horde/HordeTrayApp/Plugins/HordeProxy/HordeProxySettings.cs b/Engine/Source/Programs/Horde/HordeTrayApp/Plugins/HordeProxy/HordeProxySettings.cs
new file mode 100644
--- /dev/null
+++ b/Engine/Source/Programs/Horde/HordeTrayApp/Plugins/HordeProxy/HordeProxySettings.cs
@@ -0,0 +1,13 @@
+// Copyright Epic Games, Inc. All Rights Reserved.
+
+#pragma warning disable CS1591
+
+namespace HordeTrayApp.Plugins.HordeProxy
+{
+	[Serializable]
+	public class HordeProxySettings
+	{
+		public bool Enabled { get; set; }
+		public int Port { get; set; } = 13344;
+	}
+}
diff --git a/Engine/Source/Programs/Horde/HordeTrayApp/Plugins/HordeProxy/HordeProxySettingsPage.axaml.cs b/Engine/Source/Programs/Horde/HordeTrayApp/Plugins/HordeProxy/HordeProxySettingsPage.axaml.cs
new file mode 100644
--- /dev/null
+++ b/Engine/Source/Programs/Horde/HordeTrayApp/Plugins/HordeProxy/HordeProxySettingsPage.axaml.cs
@@ -0,0 +1,21 @@
+// Copyright Epic Games, Inc. All Rights Reserved.
+
+using Avalonia.Controls;
+
+namespace HordeTrayApp.Plugins.HordeProxy
+{
+	partial class HordeProxySettingsPage : UserControl
+	{
+		public HordeProxySettingsPage()
+			: this(null)
+		{
+		}
+
+		public HordeProxySettingsPage(HordeProxyPlugin? plugin)
+		{
+			InitializeComponent();
+
+			DataContext = new HordeProxySettingsViewModel(plugin);
+		}
+	}
+}
\ No newline at end of file
diff --git a/Engine/Source/Programs/Horde/HordeTrayApp/Plugins/HordeProxy/HordeProxySettingsViewModel.cs b/Engine/Source/Programs/Horde/HordeTrayApp/Plugins/HordeProxy/HordeProxySettingsViewModel.cs
new file mode 100644
--- /dev/null
+++ b/Engine/Source/Programs/Horde/HordeTrayApp/Plugins/HordeProxy/HordeProxySettingsViewModel.cs
@@ -0,0 +1,55 @@
+// Copyright Epic Games, Inc. All Rights Reserved.
+
+using System.ComponentModel;
+using System.Diagnostics;
+using Avalonia.Platform.Storage;
+using CommunityToolkit.Mvvm.ComponentModel;
+
+namespace HordeTrayApp.Plugins.HordeProxy
+{
+	partial class HordeProxySettingsViewModel : ObservableObject
+	{
+		readonly HordeProxyPlugin? _plugin;
+
+		[ObservableProperty]
+		bool _enabled;
+
+		[ObservableProperty]
+		int _port;
+
+		[ObservableProperty]
+		string _status = String.Empty;
+
+		public HordeProxySettingsViewModel(HordeProxyPlugin? plugin)
+		{
+			_plugin = plugin;
+			if (_plugin != null)
+			{
+				Status = _plugin.Status;
+				_plugin.OnStatusChanged += OnStatusChanged;
+			}
+
+			HordeProxySettings settings = plugin?.Settings ?? new HordeProxySettings();
+			_enabled = settings.Enabled;
+			_port = settings.Port;
+		}
+
+		void OnStatusChanged()
+			=> Status = _plugin?.Status ?? String.Empty;
+
+		protected override void OnPropertyChanged(PropertyChangedEventArgs e)
+		{
+			base.OnPropertyChanged(e);
+
+			if (_plugin != null)
+			{
+				HordeProxySettings settings = new HordeProxySettings();
+
+				settings.Enabled = Enabled;
+				settings.Port = Port;
+
+				_plugin.UpdateSettings(settings);
+			}
+		}
+	}
+}
diff --git a/Engine/Source/Programs/Shared/EpicGames.Horde/HordeHttpClient.cs b/Engine/Source/Programs/Shared/EpicGames.Horde/HordeHttpClient.cs
--- a/Engine/Source/Programs/Shared/EpicGames.Horde/HordeHttpClient.cs
+++ b/Engine/Source/Programs/Shared/EpicGames.Horde/HordeHttpClient.cs
@@ -65,6 +65,11 @@ namespace EpicGames.Horde
 		/// </summary>
 		public const string UploadRedirectHttpClientName = "HordeUploadRedirectHttpClient";
 
+		/// <summary>
+		/// Accessor for the inner http client
+		/// </summary>
+		public HttpClient HttpClient => _httpClient;
+
 		readonly HttpClient _httpClient;
 
 		internal static JsonSerializerOptions JsonSerializerOptions => HordeHttpRequest.JsonSerializerOptions;
