Horde: Add a delay before applying proxy settings.

#rnx

diff --git a/Engine/Source/Programs/Horde/HordeTrayApp/IHordeClientProvider.cs b/Engine/Source/Programs/Horde/HordeTrayApp/IHordeClientProvider.cs
new file mode 100644
--- /dev/null
+++ b/Engine/Source/Programs/Horde/HordeTrayApp/IHordeClientProvider.cs
@@ -0,0 +1,43 @@
+// Copyright Epic Games, Inc. All Rights Reserved.
+
+using EpicGames.Horde;
+
+namespace HordeTrayApp
+{
+	/// <summary>
+	/// Reference to a Horde client instance. Should be disposed of when done.
+	/// </summary>
+	interface IHordeClientRef : IDisposable
+	{
+		/// <summary>
+		/// The client instance
+		/// </summary>
+		IHordeClient Client { get; }
+	}
+
+	/// <summary>
+	/// Accessor for a <see cref="IHordeClient"/> instance which can be recreated in response to config changes. 
+	/// </summary>
+	interface IHordeClientProvider
+	{
+		/// <summary>
+		/// Event signalled whenever the connection state changes
+		/// </summary>
+		event Action? OnStateChanged;
+
+		/// <summary>
+		/// Event signalled whenever the access token state changes
+		/// </summary>
+		event Action? OnAccessTokenStateChanged;
+
+		/// <summary>
+		/// Gets a reference to the current client instance. These references should be kept as short as possible.
+		/// </summary>
+		Task<IHordeClientRef?> GetClientRefAsync();
+
+		/// <summary>
+		/// Recreates the client instance.
+		/// </summary>
+		Task RecreateAsync();
+	}
+}
diff --git a/Engine/Source/Programs/Horde/HordeTrayApp/JsonConfig.cs b/Engine/Source/Programs/Horde/HordeTrayApp/JsonConfig.cs
new file mode 100644
--- /dev/null
+++ b/Engine/Source/Programs/Horde/HordeTrayApp/JsonConfig.cs
@@ -0,0 +1,79 @@
+// Copyright Epic Games, Inc. All Rights Reserved.
+
+using System;
+using System.Collections.Generic;
+using System.Linq;
+using System.Text;
+using System.Text.Json;
+using System.Text.Json.Serialization;
+using System.Threading.Tasks;
+using EpicGames.Core;
+
+namespace HordeTrayApp
+{
+	class JsonConfig<T> where T : class, new()
+	{
+		readonly FileReference _file;
+		static readonly JsonSerializerOptions s_jsonSerializerOptions = GetJsonSerializerOptions();
+		T _current;
+		byte[] _data;
+
+		public T Current => _current;
+
+		public JsonConfig(FileReference file)
+		{
+			_file = file;
+			_current = new T();
+			_data = Array.Empty<byte>();
+		}
+
+		static JsonSerializerOptions GetJsonSerializerOptions()
+		{
+			JsonSerializerOptions options = new JsonSerializerOptions();
+			options.PropertyNamingPolicy = JsonNamingPolicy.CamelCase;
+			options.PropertyNameCaseInsensitive = true;
+			options.DefaultIgnoreCondition = JsonIgnoreCondition.WhenWritingNull;
+			options.AllowTrailingCommas = true;
+			options.WriteIndented = true;
+			options.Converters.Add(new JsonStringEnumConverter());
+			return options;
+		}
+
+		public bool LoadSettings()
+		{
+			if (FileReference.Exists(_file))
+			{
+				try
+				{
+					byte[] data = FileReference.ReadAllBytes(_file);
+					if (!data.SequenceEqual(_data))
+					{
+						_current = JsonSerializer.Deserialize<T>(data, s_jsonSerializerOptions)!;
+						_data = data;
+						return true;
+					}
+				}
+				catch (Exception)
+				{
+				}
+			}
+			return false;
+		}
+
+		public bool UpdateSettings(T settings)
+		{
+			byte[] data = JsonSerializer.SerializeToUtf8Bytes(settings, s_jsonSerializerOptions);
+			if (!data.SequenceEqual(_data))
+			{
+				_current = settings;
+				_data = data;
+
+				DirectoryReference.CreateDirectory(_file.Directory);
+				FileReference.WriteAllBytes(_file, data);
+
+				return true;
+			}
+			return false;
+		}
+	}
+}
diff --git a/Engine/Source/Programs/Horde/HordeTrayApp/Plugins/HordeProxy/HordeProxyPlugin.cs b/Engine/Source/Programs/Horde/HordeTrayApp/Plugins/HordeProxy/HordeProxyPlugin.cs
--- a/Engine/Source/Programs/Horde/HordeTrayApp/Plugins/HordeProxy/HordeProxyPlugin.cs
+++ b/Engine/Source/Programs/Horde/HordeTrayApp/Plugins/HordeProxy/HordeProxyPlugin.cs
@@ -17,18 +17,15 @@ namespace HordeTrayApp.Plugins.HordeProxy
 {
 	class HordeProxyPlugin : ITrayAppPlugin
 	{
-		readonly FileReference _settingsFile;
 		readonly BackgroundTask _serverTask;
 		readonly IHordeClientProvider _hordeClientProvider;
 		readonly AsyncEvent _restartServerEvent = new AsyncEvent();
 		readonly SynchronizationContext? _synchronizationContext;
+		readonly JsonConfig<HordeProxySettings> _settings;
 
 		string _status = "Starting...";
 
-		byte[] _settingsData = Array.Empty<byte>();
-		HordeProxySettings _settings = new HordeProxySettings();
-
-		public HordeProxySettings Settings => _settings;
+		public HordeProxySettings Settings => _settings.Current;
 
 		public event Action? OnStatusChanged;
 
@@ -46,9 +43,8 @@ namespace HordeTrayApp.Plugins.HordeProxy
 
 			DirectoryReference ? settingsRoot = DirectoryReference.GetSpecialFolder(Environment.SpecialFolder.CommonApplicationData);
 			settingsRoot ??= DirectoryReference.GetCurrentDirectory();
-			_settingsFile = FileReference.Combine(settingsRoot, "Epic", "Horde", "Proxy", "Settings.json");
-
-			LoadSettings();
+			_settings = new JsonConfig<HordeProxySettings>(FileReference.Combine(settingsRoot, "Epic", "Horde", "Proxy", "Settings.json"));
+			_settings.LoadSettings();
 
 			_serverTask = BackgroundTask.StartNew(RunServerAsync);
 		}
@@ -56,21 +52,7 @@ namespace HordeTrayApp.Plugins.HordeProxy
 		public Control CreateSettingsPage(SettingsContext context)
 			=> new HordeProxySettingsPage(this);
 
-		public void UpdateSettings(HordeProxySettings settings)
-		{
-			byte[] data = JsonSerializer.SerializeToUtf8Bytes(settings, GetJsonSerializerOptions());
-			if (!data.SequenceEqual(_settingsData))
-			{
-				_settings = settings;
-				_settingsData = data;
-				_restartServerEvent.Pulse();
-
-				DirectoryReference.CreateDirectory(_settingsFile.Directory);
-				FileReference.WriteAllBytes(_settingsFile, data);
-			}
-		}
-
-		public void SetStatus(string status)
+		void SetStatus(string status)
 		{
 			_status = status;
 
@@ -84,6 +66,14 @@ namespace HordeTrayApp.Plugins.HordeProxy
 			}
 		}
 
+		public void UpdateSettings(HordeProxySettings settings)
+		{
+			if (_settings.UpdateSettings(settings))
+			{
+				_restartServerEvent.Pulse();
+			}
+		}
+
 		public async ValueTask DisposeAsync()
 		{
 			_hordeClientProvider.OnStateChanged -= OnStateChanged;
@@ -95,41 +85,6 @@ namespace HordeTrayApp.Plugins.HordeProxy
 			_restartServerEvent.Pulse();
 		}
 
-		static JsonSerializerOptions GetJsonSerializerOptions()
-		{
-			JsonSerializerOptions options = new JsonSerializerOptions();
-			options.PropertyNamingPolicy = JsonNamingPolicy.CamelCase;
-			options.PropertyNameCaseInsensitive = true;
-			options.DefaultIgnoreCondition = JsonIgnoreCondition.WhenWritingNull;
-			options.AllowTrailingCommas = true;
-			options.WriteIndented = true;
-			options.Converters.Add(new JsonStringEnumConverter());
-			return options;
-		}
-
-		void LoadSettings()
-		{
-			if (FileReference.Exists(_settingsFile))
-			{
-				try
-				{
-					byte[] data = FileReference.ReadAllBytes(_settingsFile);
-					if (!data.SequenceEqual(_settingsData))
-					{
-						_settings = JsonSerializer.Deserialize<HordeProxySettings>(data, GetJsonSerializerOptions())!;
-						_settingsData = data;
-						_restartServerEvent.Pulse();
-						return;
-					}
-				}
-				catch (Exception)
-				{
-				}
-			}
-
-			UpdateSettings(new HordeProxySettings());
-		}
-
 		public void PopulateContextMenu(NativeMenu contextMenu)
 		{
 		}
@@ -143,7 +98,7 @@ namespace HordeTrayApp.Plugins.HordeProxy
 			{
 				Task restartServerTask = _restartServerEvent.Task;
 
-				HordeProxySettings settings = _settings;
+				HordeProxySettings settings = _settings.Current;
 				if (settings.Enabled)
 				{
 					using CancellationTokenSource cancellationSource = CancellationTokenSource.CreateLinkedTokenSource(cancellationToken);
diff --git a/Engine/Source/Programs/Horde/HordeTrayApp/Plugins/HordeProxy/HordeProxySettingsViewModel.cs b/Engine/Source/Programs/Horde/HordeTrayApp/Plugins/HordeProxy/HordeProxySettingsViewModel.cs
--- a/Engine/Source/Programs/Horde/HordeTrayApp/Plugins/HordeProxy/HordeProxySettingsViewModel.cs
+++ b/Engine/Source/Programs/Horde/HordeTrayApp/Plugins/HordeProxy/HordeProxySettingsViewModel.cs
@@ -3,6 +3,7 @@
 using System.ComponentModel;
 using System.Diagnostics;
 using Avalonia.Platform.Storage;
+using Avalonia.Threading;
 using CommunityToolkit.Mvvm.ComponentModel;
 
 namespace HordeTrayApp.Plugins.HordeProxy
@@ -11,6 +12,8 @@ namespace HordeTrayApp.Plugins.HordeProxy
 	{
 		readonly HordeProxyPlugin? _plugin;
 
+		int _pendingUpdate;
+
 		[ObservableProperty]
 		bool _enabled;
 
@@ -23,15 +26,16 @@ namespace HordeTrayApp.Plugins.HordeProxy
 		public HordeProxySettingsViewModel(HordeProxyPlugin? plugin)
 		{
 			_plugin = plugin;
-			if (_plugin != null)
-			{
-				Status = _plugin.Status;
-				_plugin.OnStatusChanged += OnStatusChanged;
-			}
 
 			HordeProxySettings settings = plugin?.Settings ?? new HordeProxySettings();
 			_enabled = settings.Enabled;
 			_port = settings.Port;
+
+			if (_plugin != null)
+			{
+				_status = _plugin.Status;
+				_plugin.OnStatusChanged += OnStatusChanged;
+			}
 		}
 
 		void OnStatusChanged()
@@ -43,11 +47,28 @@ namespace HordeTrayApp.Plugins.HordeProxy
 
 			if (_plugin != null)
 			{
-				HordeProxySettings settings = new HordeProxySettings();
+				int currentUpdate = Interlocked.Increment(ref _pendingUpdate);
+				ApplySettingsAsync(currentUpdate);
+			}
+		}
+
+		async void ApplySettingsAsync(int currentUpdate)
+		{
+			await Task.Delay(TimeSpan.FromSeconds(1.0));
+
+			if (Interlocked.CompareExchange(ref _pendingUpdate, 0, 0) == currentUpdate)
+			{
+				ApplySettings();
+			}
+		}
 
+		void ApplySettings()
+		{
+			if (_plugin != null)
+			{
+				HordeProxySettings settings = new HordeProxySettings();
 				settings.Enabled = Enabled;
 				settings.Port = Port;
-
 				_plugin.UpdateSettings(settings);
 			}
 		}
