Horde: Allow setting an install and uninstall command for launcher tools.

#rnx

diff --git a/Engine/Source/Programs/Horde/HordeTrayApp/App.axaml.cs b/Engine/Source/Programs/Horde/HordeTrayApp/App.axaml.cs
--- a/Engine/Source/Programs/Horde/HordeTrayApp/App.axaml.cs
+++ b/Engine/Source/Programs/Horde/HordeTrayApp/App.axaml.cs
@@ -135,7 +135,7 @@ namespace HordeTrayApp
 			}
 		}
 
-		static void TrayIcon_RunCommand(IToolCatalogItem item, ToolMenuCommand command)
+		static void TrayIcon_RunCommand(IToolCatalogItem item, ToolCommand command)
 		{
 			ProcessStartInfo startInfo = new ProcessStartInfo(command.FileName);
 			if (command.Arguments != null)
diff --git a/Engine/Source/Programs/Horde/HordeTrayApp/ToolCatalog.cs b/Engine/Source/Programs/Horde/HordeTrayApp/ToolCatalog.cs
--- a/Engine/Source/Programs/Horde/HordeTrayApp/ToolCatalog.cs
+++ b/Engine/Source/Programs/Horde/HordeTrayApp/ToolCatalog.cs
@@ -439,22 +439,18 @@ namespace HordeTrayApp
 			{
 				if (deployment == null)
 				{
-					UpdateStateAndNotify(item, state => state with { Current = null });
+					await UninstallAsync(item, cancellationToken);
 					_cleanupEvent.Set();
 				}
 				else
 				{
 					UpdateStateAndNotify(item, state => state with { Pending = new PendingToolDeploymentInfo(false, $"Updating to {deployment.Version}", deploymentInfo) });
 
-					DirectoryReference targetDir = DirectoryReference.Combine(_baseDir, item.Id.ToString(), deployment.Id.ToString());
-					DirectoryReference.CreateDirectory(targetDir);
-					await deployment.Content.ExtractAsync(targetDir.ToDirectoryInfo(), _logger, cancellationToken);
-
 					DirectoryReference toolDir = GetToolDir(item.Id, deploymentInfo.Id);
-					ToolConfig toolConfig = LoadToolConfig(toolDir);
-					CurrentToolDeploymentInfo current = new CurrentToolDeploymentInfo(deploymentInfo.Id, deploymentInfo.Version, toolDir, toolConfig);
+					DirectoryReference.CreateDirectory(toolDir);
+					await deployment.Content.ExtractAsync(toolDir.ToDirectoryInfo(), _logger, cancellationToken);
 
-					UpdateStateAndNotify(item, state => state with { Pending = null, Current = current });
+					await InstallAsync(item, toolDir, deploymentInfo, cancellationToken);
 				}
 
 				await SaveStateAsync(cancellationToken);
@@ -471,6 +467,80 @@ namespace HordeTrayApp
 			}
 		}
 
+		async Task InstallAsync(Item item, DirectoryReference toolDir, ToolDeploymentInfo deploymentInfo, CancellationToken cancellationToken)
+		{
+			ToolConfig toolConfig = LoadToolConfig(toolDir);
+
+			await UninstallAsync(item, cancellationToken);
+
+			if (toolConfig.InstallCommand != null)
+			{
+				int exitCode = await RunCommandAsync(item.Id.ToString(), toolConfig.InstallCommand, toolDir.FullName, cancellationToken);
+				if (exitCode != 0)
+				{
+					PendingToolDeploymentInfo pending = new PendingToolDeploymentInfo(true, $"Installation failed ({exitCode}). See log.", deploymentInfo);
+					UpdateStateAndNotify(item, state => state with { Pending = pending });
+					return;
+				}
+			}
+
+			CurrentToolDeploymentInfo current = new CurrentToolDeploymentInfo(deploymentInfo.Id, deploymentInfo.Version, toolDir, toolConfig);
+			UpdateStateAndNotify(item, state => state with { Pending = null, Current = current });
+		}
+
+		async Task UninstallAsync(Item item, CancellationToken cancellationToken)
+		{
+			CurrentToolDeploymentInfo? current = null;
+			UpdateStateAndNotify(item, state =>
+			{
+				current = state.Current;
+
+				PendingToolDeploymentInfo? pending = null;
+				if (current != null)
+				{
+					pending = new PendingToolDeploymentInfo(false, $"Removing {current.Version}", null);
+				}
+
+				return state with { Pending = pending, Current = null };
+			});
+
+			if (current != null)
+			{
+				ToolCommand? uninstallCommand = current.Config.UninstallCommand;
+				if (uninstallCommand != null)
+				{
+					await RunCommandAsync(current.Id.ToString(), uninstallCommand, current.Dir.FullName, cancellationToken);
+				}
+			}
+
+			UpdateStateAndNotify(item, state => state with { Pending = null, Current = null });
+		}
+
+		async Task<int> RunCommandAsync(string toolName, ToolCommand command, string workingDir, CancellationToken cancellationToken)
+		{
+			string arguments = String.Empty;
+			if (command.Arguments != null)
+			{
+				arguments = CommandLineArguments.Join(command.Arguments);
+			}
+
+			using (ManagedProcessGroup newProcessGroup = new ManagedProcessGroup())
+			using (ManagedProcess newProcess = new ManagedProcess(newProcessGroup, command.FileName, arguments, workingDir, null, null, ProcessPriorityClass.Normal))
+			{
+				for (; ; )
+				{
+					string? line = await newProcess.ReadLineAsync(cancellationToken);
+					if (line == null)
+					{
+						await newProcess.WaitForExitAsync(cancellationToken);
+						_logger.LogInformation("{ToolName}> Exit code {ExitCode}", toolName, newProcess.ExitCode);
+						return newProcess.ExitCode;
+					}
+					_logger.LogInformation("{ToolName}> {Line}", toolName, line);
+				}
+			}
+		}
+
 		void UpdateStateAndNotify(Item item, Func<ItemState, ItemState> update)
 		{
 			lock (_workerThreadLockObject)
diff --git a/Engine/Source/Programs/Horde/HordeTrayApp/ToolConfig.cs b/Engine/Source/Programs/Horde/HordeTrayApp/ToolConfig.cs
--- a/Engine/Source/Programs/Horde/HordeTrayApp/ToolConfig.cs
+++ b/Engine/Source/Programs/Horde/HordeTrayApp/ToolConfig.cs
@@ -11,6 +11,16 @@ namespace HordeTrayApp
 		/// Popup menu item for the tool
 		/// </summary>
 		public ToolMenuItem? PopupMenu { get; set; }
+
+		/// <summary>
+		/// Command to run when installing
+		/// </summary>
+		public ToolCommand? InstallCommand { get; set; }
+
+		/// <summary>
+		/// Command to run when uninstalling
+		/// </summary>
+		public ToolCommand? UninstallCommand { get; set; }
 	}
 
 	/// <summary>
@@ -26,7 +36,7 @@ namespace HordeTrayApp
 		/// <summary>
 		/// Command to run when the menu item is clicked. Will be executed in the root directory of the downloaded tool.
 		/// </summary>
-		public ToolMenuCommand? Command { get; set; }
+		public ToolCommand? Command { get; set; }
 
 		/// <summary>
 		/// Child menu items. Cannot be specified at the same time as a tool.
@@ -37,7 +47,7 @@ namespace HordeTrayApp
 	/// <summary>
 	/// Command to run for a menu item.
 	/// </summary>
-	public class ToolMenuCommand
+	public class ToolCommand
 	{
 		/// <summary>
 		/// Executable to run
