Horde: Updating schema documentation

#rnx

diff --git a/Engine/Source/Programs/Horde/Docs/Config/Schema/Globals.md b/Engine/Source/Programs/Horde/Docs/Config/Schema/Globals.md
--- a/Engine/Source/Programs/Horde/Docs/Config/Schema/Globals.md
+++ b/Engine/Source/Programs/Horde/Docs/Config/Schema/Globals.md
@@ -557,5 +557,6 @@ Name | Description
 `public` | `boolean`<br>Whether this tool should be exposed for download on a public endpoint without authentication
 `showInUgs` | `boolean`<br>Whether to show this tool for download in the UGS tools menu
 `showInDashboard` | `boolean`<br>Whether to show this tool for download in the dashboard
+`showInLauncher` | `boolean`<br>Whether to show this tool for download in the Horde launcher
 `namespaceId` | `string`<br>Default namespace for new deployments of this tool
 `acl` | [AclConfig](#aclconfig)<br>Permissions for the tool
