Horde: Move OpenTelemetry settings to EpicGames.Horde for sharing with agent

diff --git a/Engine/Source/Programs/Horde/HordeServer/ServerSettings.cs b/Engine/Source/Programs/Horde/HordeServer/ServerSettings.cs
--- a/Engine/Source/Programs/Horde/HordeServer/ServerSettings.cs
+++ b/Engine/Source/Programs/Horde/HordeServer/ServerSettings.cs
@@ -1,11 +1,10 @@
 // Copyright Epic Games, Inc. All Rights Reserved.
 
 using System;
-using System.Collections.Generic;
 using System.Linq;
 using System.Net;
-using EpicGames.Core;
 using EpicGames.Horde.Server;
+using EpicGames.Horde.Utilities;
 using HordeServer.Utilities;
 
 namespace HordeServer
@@ -21,68 +20,6 @@ namespace HordeServer
 	{
 	}
 
-	/// <summary>
-	/// OpenTelemetry configuration for collection and sending of traces and metrics.
-	/// </summary>
-	public class OpenTelemetrySettings
-	{
-		/// <summary>
-		/// Whether OpenTelemetry exporting is enabled
-		/// </summary>
-		public bool Enabled { get; set; } = false;
-
-		/// <summary>
-		/// Service name
-		/// </summary>
-		public string ServiceName { get; set; } = "HordeServer";
-
-		/// <summary>
-		/// Service namespace
-		/// </summary>
-		public string ServiceNamespace { get; set; } = "Horde";
-
-		/// <summary>
-		/// Service version
-		/// </summary>
-		public string? ServiceVersion { get; set; }
-
-		/// <summary>
-		/// Whether to enrich and format telemetry to fit presentation in Datadog
-		/// </summary>
-		public bool EnableDatadogCompatibility { get; set; } = false;
-
-		/// <summary>
-		/// Extra attributes to set
-		/// </summary>
-		public Dictionary<string, string> Attributes { get; set; } = new();
-
-		/// <summary>
-		/// Whether to enable the console exporter (for debugging purposes)
-		/// </summary>
-		public bool EnableConsoleExporter { get; set; } = false;
-
-		/// <summary>
-		/// Protocol exporters (key is a unique and arbitrary name) 
-		/// </summary>
-		public Dictionary<string, OpenTelemetryProtocolExporterSettings> ProtocolExporters { get; set; } = new();
-	}
-
-	/// <summary>
-	/// Configuration for an OpenTelemetry exporter
-	/// </summary>
-	public class OpenTelemetryProtocolExporterSettings
-	{
-		/// <summary>
-		/// Endpoint URL. Usually differs depending on protocol used.
-		/// </summary>
-		public Uri? Endpoint { get; set; }
-
-		/// <summary>
-		/// Protocol for the exporter ('grpc' or 'httpprotobuf')
-		/// </summary>
-		public string Protocol { get; set; } = "grpc";
-	}
-
 	/// <summary>
 	/// Global settings for the application
 	/// </summary>
diff --git a/Engine/Source/Programs/Horde/HordeServer/Utilities/OpenTelemetry.cs b/Engine/Source/Programs/Horde/HordeServer/Utilities/OpenTelemetry.cs
--- a/Engine/Source/Programs/Horde/HordeServer/Utilities/OpenTelemetry.cs
+++ b/Engine/Source/Programs/Horde/HordeServer/Utilities/OpenTelemetry.cs
@@ -7,6 +7,7 @@ using System.Diagnostics.Metrics;
 using System.Linq;
 using System.Net.Http;
 using System.Runtime.CompilerServices;
+using EpicGames.Horde.Utilities;
 using Microsoft.AspNetCore.Http;
 using Microsoft.Extensions.DependencyInjection;
 using Microsoft.Extensions.Logging;
diff --git a/Engine/Source/Programs/Shared/EpicGames.Horde/Utilities/OpenTelemetry.cs b/Engine/Source/Programs/Shared/EpicGames.Horde/Utilities/OpenTelemetry.cs
new file mode 100644
--- /dev/null
+++ b/Engine/Source/Programs/Shared/EpicGames.Horde/Utilities/OpenTelemetry.cs
@@ -0,0 +1,70 @@
+// Copyright Epic Games, Inc. All Rights Reserved.
+
+using System;
+using System.Collections.Generic;
+
+namespace EpicGames.Horde.Utilities;
+
+#pragma warning disable CA2227 // Collection properties should be read only
+
+/// <summary>
+/// OpenTelemetry configuration for collection and sending of traces and metrics.
+/// </summary>
+public class OpenTelemetrySettings
+{
+	/// <summary>
+	/// Whether OpenTelemetry exporting is enabled
+	/// </summary>
+	public bool Enabled { get; set; } = false;
+
+	/// <summary>
+	/// Service name
+	/// </summary>
+	public string ServiceName { get; set; } = "HordeServer";
+
+	/// <summary>
+	/// Service namespace
+	/// </summary>
+	public string ServiceNamespace { get; set; } = "Horde";
+
+	/// <summary>
+	/// Service version
+	/// </summary>
+	public string? ServiceVersion { get; set; }
+
+	/// <summary>
+	/// Whether to enrich and format telemetry to fit presentation in Datadog
+	/// </summary>
+	public bool EnableDatadogCompatibility { get; set; } = false;
+
+	/// <summary>
+	/// Extra attributes to set
+	/// </summary>
+	public Dictionary<string, string> Attributes { get; set; } = new();
+
+	/// <summary>
+	/// Whether to enable the console exporter (for debugging purposes)
+	/// </summary>
+	public bool EnableConsoleExporter { get; set; } = false;
+
+	/// <summary>
+	/// Protocol exporters (key is a unique and arbitrary name) 
+	/// </summary>
+	public Dictionary<string, OpenTelemetryProtocolExporterSettings> ProtocolExporters { get; set; } = new();
+}
+
+/// <summary>
+/// Configuration for an OpenTelemetry exporter
+/// </summary>
+public class OpenTelemetryProtocolExporterSettings
+{
+	/// <summary>
+	/// Endpoint URL. Usually differs depending on protocol used.
+	/// </summary>
+	public Uri? Endpoint { get; set; }
+
+	/// <summary>
+	/// Protocol for the exporter ('grpc' or 'httpprotobuf')
+	/// </summary>
+	public string Protocol { get; set; } = "grpc";
+}
\ No newline at end of file
