Horde: Replace OpenTracing with OpenTelemetry in agent

diff --git a/Engine/Source/Programs/Horde/HordeAgent/AgentApp.cs b/Engine/Source/Programs/Horde/HordeAgent/AgentApp.cs
--- a/Engine/Source/Programs/Horde/HordeAgent/AgentApp.cs
+++ b/Engine/Source/Programs/Horde/HordeAgent/AgentApp.cs
@@ -6,9 +6,6 @@ using System.Reflection;
 using System.Runtime.InteropServices;
 using System.Text.Json;
 using System.Text.Json.Serialization;
-using Datadog.Trace;
-using Datadog.Trace.Configuration;
-using Datadog.Trace.OpenTracing;
 using EpicGames.Core;
 using EpicGames.Horde;
 using HordeAgent.Leases;
@@ -19,15 +16,12 @@ using Microsoft.Extensions.Configuration;
 using Microsoft.Extensions.DependencyInjection;
 using Microsoft.Extensions.Logging;
 using Microsoft.Win32;
-using OpenTracing.Util;
 using Polly;
 using IConfigurationSource = Microsoft.Extensions.Configuration.IConfigurationSource;
 using JsonConfigurationSource = Microsoft.Extensions.Configuration.Json.JsonConfigurationSource;
 
 namespace HordeAgent
 {
-	using ITracer = OpenTracing.ITracer;
-
 	/// <summary>
 	/// Injectable list of existing services
 	/// </summary>
@@ -180,9 +174,7 @@ namespace HordeAgent
 				return 1;
 			}
 
-			ConfigureTracing(serverProfile.Environment, AgentApp.Version);
-
-			Logging.SetEnv(serverProfile.Environment);
+			OpenTelemetryHelper.Configure(services, settings.OpenTelemetry);
 
 			services.AddHorde(options =>
 			{
@@ -335,20 +327,6 @@ namespace HordeAgent
 			}
 		}
 
-		static void ConfigureTracing(string environment, string version)
-		{
-			TracerSettings settings = TracerSettings.FromDefaultSources();
-			settings.Environment = environment;
-			settings.ServiceName = "hordeagent";
-			settings.ServiceVersion = version;
-			settings.LogsInjectionEnabled = true;
-
-			Tracer.Configure(settings);
-
-			ITracer openTracer = OpenTracingTracerFactory.WrapTracer(Tracer.Instance);
-			GlobalTracer.Register(openTracer);
-		}
-
 		/// <summary>
 		/// Gets the application directory
 		/// </summary>
diff --git a/Engine/Source/Programs/Horde/HordeAgent/AgentSettings.cs b/Engine/Source/Programs/Horde/HordeAgent/AgentSettings.cs
--- a/Engine/Source/Programs/Horde/HordeAgent/AgentSettings.cs
+++ b/Engine/Source/Programs/Horde/HordeAgent/AgentSettings.cs
@@ -3,6 +3,7 @@
 using System.ComponentModel.DataAnnotations;
 using EpicGames.Core;
 using EpicGames.Horde;
+using EpicGames.Horde.Utilities;
 using HordeServer.Utilities;
 using Microsoft.Extensions.Configuration;
 
@@ -177,6 +178,11 @@ namespace HordeAgent
 		/// Incoming port for listening for compute work. Needs to be tied with a lease. Set port to 0 to disable incoming compute requests.
 		/// </summary>
 		public int ComputePort { get; set; } = 7000;
+		
+		/// <summary>
+		/// Options for OpenTelemetry
+		/// </summary>
+		public OpenTelemetrySettings OpenTelemetry { get; set; } = new ();
 
 		/// <summary>
 		/// Whether to send telemetry back to Horde server
diff --git a/Engine/Source/Programs/Horde/HordeAgent/HordeAgent.csproj b/Engine/Source/Programs/Horde/HordeAgent/HordeAgent.csproj
--- a/Engine/Source/Programs/Horde/HordeAgent/HordeAgent.csproj
+++ b/Engine/Source/Programs/Horde/HordeAgent/HordeAgent.csproj
@@ -36,7 +36,6 @@
     <PackageReference Include="AWSSDK.Core" Version="3.7.300.12" />
     <PackageReference Include="AWSSDK.EC2" Version="3.7.306.6" />
     <PackageReference Include="AWSSDK.SimpleSystemsManagement" Version="3.7.301.8" />
-    <PackageReference Include="Datadog.Trace.OpenTracing" Version="2.14.0" />
     <PackageReference Include="Google.Protobuf" Version="3.25.1" />
     <PackageReference Include="Grpc.Net.Client" Version="2.59.0" />
     <PackageReference Include="Grpc.Tools" Version="2.59.0">
@@ -60,7 +59,13 @@
       <PrivateAssets>all</PrivateAssets>
       <IncludeAssets>runtime; build; native; contentfiles; analyzers; buildtransitive</IncludeAssets>
     </PackageReference>
-    <PackageReference Include="OpenTracing" Version="0.12.1" />
+    <PackageReference Include="OpenTelemetry" Version="1.8.0" />
+    <PackageReference Include="OpenTelemetry.Exporter.Console" Version="1.8.0" />
+    <PackageReference Include="OpenTelemetry.Exporter.OpenTelemetryProtocol" Version="1.8.0" />
+    <PackageReference Include="OpenTelemetry.Extensions.Hosting" Version="1.8.0" />
+    <PackageReference Include="OpenTelemetry.Instrumentation.AspNetCore" Version="1.8.1" />
+    <PackageReference Include="OpenTelemetry.Instrumentation.GrpcNetClient" Version="1.8.0-beta.1" />
+    <PackageReference Include="OpenTelemetry.Instrumentation.Http" Version="1.8.1" />
     <PackageReference Include="Serilog" Version="3.1.1" />
     <PackageReference Include="Serilog.Extensions.Hosting" Version="8.0.0" />
     <PackageReference Include="Serilog.Extensions.Logging" Version="8.0.0" />
diff --git a/Engine/Source/Programs/Horde/HordeAgent/Leases/Handlers/ComputeHandler.cs b/Engine/Source/Programs/Horde/HordeAgent/Leases/Handlers/ComputeHandler.cs
--- a/Engine/Source/Programs/Horde/HordeAgent/Leases/Handlers/ComputeHandler.cs
+++ b/Engine/Source/Programs/Horde/HordeAgent/Leases/Handlers/ComputeHandler.cs
@@ -12,6 +12,7 @@ using HordeCommon.Rpc.Messages;
 using HordeCommon.Rpc.Tasks;
 using Microsoft.Extensions.Logging;
 using Microsoft.Extensions.Options;
+using OpenTelemetry.Trace;
 
 namespace HordeAgent.Leases.Handlers
 {
@@ -34,7 +35,7 @@ namespace HordeAgent.Leases.Handlers
 		}
 
 		/// <inheritdoc/>
-		protected override async Task<LeaseResult> ExecuteAsync(ISession session, LeaseId leaseId, ComputeTask computeTask, ILogger localLogger, CancellationToken cancellationToken)
+		protected override async Task<LeaseResult> ExecuteAsync(ISession session, LeaseId leaseId, ComputeTask computeTask, Tracer tracer, ILogger localLogger, CancellationToken cancellationToken)
 		{
 			await using IServerLogger? serverLogger = (computeTask.LogId != null) ? session.HordeClient.CreateServerLogger(LogId.Parse(computeTask.LogId), LogLevel.Trace).WithLocalLogger(localLogger) : null;
 			ILogger logger = serverLogger ?? localLogger;
diff --git a/Engine/Source/Programs/Horde/HordeAgent/Leases/Handlers/ConformHandler.cs b/Engine/Source/Programs/Horde/HordeAgent/Leases/Handlers/ConformHandler.cs
--- a/Engine/Source/Programs/Horde/HordeAgent/Leases/Handlers/ConformHandler.cs
+++ b/Engine/Source/Programs/Horde/HordeAgent/Leases/Handlers/ConformHandler.cs
@@ -9,6 +9,7 @@ using HordeAgent.Services;
 using HordeCommon.Rpc.Messages;
 using HordeCommon.Rpc.Tasks;
 using Microsoft.Extensions.Logging;
+using OpenTelemetry.Trace;
 
 namespace HordeAgent.Leases.Handlers
 {
@@ -19,7 +20,7 @@ namespace HordeAgent.Leases.Handlers
 		{ }
 
 		/// <inheritdoc/>
-		protected override async Task<LeaseResult> ExecuteAsync(ISession session, LeaseId leaseId, ConformTask conformTask, ILogger localLogger, CancellationToken cancellationToken)
+		protected override async Task<LeaseResult> ExecuteAsync(ISession session, LeaseId leaseId, ConformTask conformTask, Tracer tracer, ILogger localLogger, CancellationToken cancellationToken)
 		{
 			await using IServerLogger serverLogger = session.HordeClient.CreateServerLogger(LogId.Parse(conformTask.LogId)).WithLocalLogger(localLogger);
 
diff --git a/Engine/Source/Programs/Horde/HordeAgent/Leases/Handlers/JobHandler.cs b/Engine/Source/Programs/Horde/HordeAgent/Leases/Handlers/JobHandler.cs
--- a/Engine/Source/Programs/Horde/HordeAgent/Leases/Handlers/JobHandler.cs
+++ b/Engine/Source/Programs/Horde/HordeAgent/Leases/Handlers/JobHandler.cs
@@ -9,7 +9,7 @@ using HordeAgent.Services;
 using HordeCommon.Rpc.Messages;
 using HordeCommon.Rpc.Tasks;
 using Microsoft.Extensions.Logging;
-using OpenTracing.Util;
+using OpenTelemetry.Trace;
 
 namespace HordeAgent.Leases.Handlers
 {
@@ -20,11 +20,12 @@ namespace HordeAgent.Leases.Handlers
 		{ }
 
 		/// <inheritdoc/>
-		protected override async Task<LeaseResult> ExecuteAsync(ISession session, LeaseId leaseId, ExecuteJobTask executeTask, ILogger localLogger, CancellationToken cancellationToken)
+		protected override async Task<LeaseResult> ExecuteAsync(ISession session, LeaseId leaseId, ExecuteJobTask executeTask, Tracer tracer, ILogger localLogger, CancellationToken cancellationToken)
 		{
-			GlobalTracer.Instance.ActiveSpan?.SetTag("jobId", executeTask.JobId.ToString());
-			GlobalTracer.Instance.ActiveSpan?.SetTag("jobName", executeTask.JobName.ToString());
-			GlobalTracer.Instance.ActiveSpan?.SetTag("batchId", executeTask.BatchId.ToString());
+			using TelemetrySpan span = tracer.StartActiveSpan($"{nameof(JobHandler)}.{nameof(ExecuteAsync)}");
+			span.SetAttribute("horde.job.id", executeTask.JobId);
+			span.SetAttribute("horde.job.name", executeTask.JobName);
+			span.SetAttribute("horde.job.batch_id", executeTask.BatchId);
 
 			await using IServerLogger logger = session.HordeClient.CreateServerLogger(LogId.Parse(executeTask.LogId)).WithLocalLogger(localLogger);
 			try
@@ -42,6 +43,7 @@ namespace HordeAgent.Leases.Handlers
 
 				string driverName = String.IsNullOrEmpty(executeTask.JobOptions.Driver) ? "JobDriver" : executeTask.JobOptions.Driver;
 				FileReference driverAssembly = FileReference.Combine(new DirectoryReference(AppContext.BaseDirectory), driverName, $"{driverName}.dll");
+				span.SetAttribute("horde.job.driver_name", driverName);
 
 				Dictionary<string, string> environment = ManagedProcess.GetCurrentEnvVars();
 				environment[HordeHttpClient.HordeUrlEnvVarName] = session.HordeClient.ServerUrl.ToString();
diff --git a/Engine/Source/Programs/Horde/HordeAgent/Leases/Handlers/RestartHandler.cs b/Engine/Source/Programs/Horde/HordeAgent/Leases/Handlers/RestartHandler.cs
--- a/Engine/Source/Programs/Horde/HordeAgent/Leases/Handlers/RestartHandler.cs
+++ b/Engine/Source/Programs/Horde/HordeAgent/Leases/Handlers/RestartHandler.cs
@@ -6,6 +6,7 @@ using HordeAgent.Utility;
 using HordeCommon.Rpc.Messages;
 using HordeCommon.Rpc.Tasks;
 using Microsoft.Extensions.Logging;
+using OpenTelemetry.Trace;
 
 namespace HordeAgent.Leases.Handlers
 {
@@ -16,7 +17,7 @@ namespace HordeAgent.Leases.Handlers
 		{ }
 
 		/// <inheritdoc/>
-		protected override Task<LeaseResult> ExecuteAsync(ISession session, LeaseId leaseId, RestartTask task, ILogger logger, CancellationToken cancellationToken)
+		protected override Task<LeaseResult> ExecuteAsync(ISession session, LeaseId leaseId, RestartTask task, Tracer tracer, ILogger logger, CancellationToken cancellationToken)
 		{
 			logger.LogInformation("Scheduling restart task for agent {AgentId}", session.AgentId);
 			SessionResult result = new SessionResult((logger, ctx) => Shutdown.ExecuteAsync(true, logger, ctx));
diff --git a/Engine/Source/Programs/Horde/HordeAgent/Leases/Handlers/ShutdownHandler.cs b/Engine/Source/Programs/Horde/HordeAgent/Leases/Handlers/ShutdownHandler.cs
--- a/Engine/Source/Programs/Horde/HordeAgent/Leases/Handlers/ShutdownHandler.cs
+++ b/Engine/Source/Programs/Horde/HordeAgent/Leases/Handlers/ShutdownHandler.cs
@@ -6,6 +6,7 @@ using HordeAgent.Utility;
 using HordeCommon.Rpc.Messages;
 using HordeCommon.Rpc.Tasks;
 using Microsoft.Extensions.Logging;
+using OpenTelemetry.Trace;
 
 namespace HordeAgent.Leases.Handlers
 {
@@ -16,7 +17,7 @@ namespace HordeAgent.Leases.Handlers
 		{ }
 
 		/// <inheritdoc/>
-		protected override Task<LeaseResult> ExecuteAsync(ISession session, LeaseId leaseId, ShutdownTask task, ILogger logger, CancellationToken cancellationToken)
+		protected override Task<LeaseResult> ExecuteAsync(ISession session, LeaseId leaseId, ShutdownTask task, Tracer tracer, ILogger logger, CancellationToken cancellationToken)
 		{
 			logger.LogInformation("Scheduling shutdown task for agent {AgentId}", session.AgentId);
 			SessionResult result = new SessionResult((logger, ctx) => Shutdown.ExecuteAsync(false, logger, ctx));
diff --git a/Engine/Source/Programs/Horde/HordeAgent/Leases/Handlers/UpgradeHandler.cs b/Engine/Source/Programs/Horde/HordeAgent/Leases/Handlers/UpgradeHandler.cs
--- a/Engine/Source/Programs/Horde/HordeAgent/Leases/Handlers/UpgradeHandler.cs
+++ b/Engine/Source/Programs/Horde/HordeAgent/Leases/Handlers/UpgradeHandler.cs
@@ -13,6 +13,7 @@ using HordeCommon.Rpc;
 using HordeCommon.Rpc.Messages;
 using HordeCommon.Rpc.Tasks;
 using Microsoft.Extensions.Logging;
+using OpenTelemetry.Trace;
 
 namespace HordeAgent.Leases.Handlers
 {
@@ -23,7 +24,7 @@ namespace HordeAgent.Leases.Handlers
 		{ }
 
 		/// <inheritdoc/>
-		protected override async Task<LeaseResult> ExecuteAsync(ISession session, LeaseId leaseId, UpgradeTask task, ILogger localLogger, CancellationToken cancellationToken)
+		protected override async Task<LeaseResult> ExecuteAsync(ISession session, LeaseId leaseId, UpgradeTask task, Tracer tracer, ILogger localLogger, CancellationToken cancellationToken)
 		{
 			await using IServerLogger logger = session.HordeClient.CreateServerLogger(LogId.Parse(task.LogId)).WithLocalLogger(localLogger);
 
diff --git a/Engine/Source/Programs/Horde/HordeAgent/Leases/LeaseHandler.cs b/Engine/Source/Programs/Horde/HordeAgent/Leases/LeaseHandler.cs
--- a/Engine/Source/Programs/Horde/HordeAgent/Leases/LeaseHandler.cs
+++ b/Engine/Source/Programs/Horde/HordeAgent/Leases/LeaseHandler.cs
@@ -10,8 +10,7 @@ using HordeAgent.Services;
 using HordeCommon.Rpc;
 using HordeCommon.Rpc.Messages;
 using Microsoft.Extensions.Logging;
-using OpenTracing;
-using OpenTracing.Util;
+using OpenTelemetry.Trace;
 using ByteString = Google.Protobuf.ByteString;
 
 namespace HordeAgent.Leases
@@ -86,9 +85,9 @@ namespace HordeAgent.Leases
 		/// <summary>
 		/// Starts executing the lease
 		/// </summary>
-		public void Start(ISession session, ILogger logger, LeaseLoggerFactory leaseLoggerFactory)
+		public void Start(ISession session, Tracer tracer, ILogger logger, LeaseLoggerFactory leaseLoggerFactory)
 		{
-			Result = Task.Run(() => HandleLeaseAsync(session, logger, leaseLoggerFactory));
+			Result = Task.Run(() => HandleLeaseAsync(session, tracer, logger, leaseLoggerFactory));
 		}
 
 		/// <summary>
@@ -103,21 +102,20 @@ namespace HordeAgent.Leases
 		/// <summary>
 		/// Handle a lease request
 		/// </summary>
-		async Task<LeaseResult> HandleLeaseAsync(ISession session, ILogger logger, LeaseLoggerFactory leaseLoggerFactory)
+		async Task<LeaseResult> HandleLeaseAsync(ISession session, Tracer tracer, ILogger logger, LeaseLoggerFactory leaseLoggerFactory)
 		{
-			using IScope scope = GlobalTracer.Instance.BuildSpan("HandleLease").WithResourceName(RpcLease.Id.ToString()).StartActive();
-			scope.Span.SetTag("LeaseId", RpcLease.Id.ToString());
-			scope.Span.SetTag("AgentId", session.AgentId.ToString());
-			//			using IDisposable TraceProperty = LogContext.PushProperty("dd.trace_id", CorrelationIdentifier.TraceId.ToString());
-			//			using IDisposable SpanProperty = LogContext.PushProperty("dd.span_id", CorrelationIdentifier.SpanId.ToString());
-
+			using TelemetrySpan span = tracer.StartActiveSpan($"{nameof(LeaseHandler)}.{nameof(HandleLeaseAsync)}");
+			span.SetAttribute("horde.lease.id", RpcLease.Id.ToString());
+			span.SetAttribute("horde.agent.id", session.AgentId.ToString());
+			span.SetAttribute("horde.agent.session_id", session.SessionId.ToString());
 			logger.LogInformation("Handling lease {LeaseId}", Id);
 
 			// Get the lease outcome
 			LeaseResult result = LeaseResult.Failed;
 			try
 			{
-				result = await HandleLeasePayloadAsync(session, leaseLoggerFactory);
+				span.SetAttribute("horde.lease.task", RpcPayload.TypeUrl);
+				result = await HandleLeasePayloadAsync(session, tracer, leaseLoggerFactory);
 			}
 			catch (OperationCanceledException) when (_cancellationSource.IsCancellationRequested)
 			{
@@ -126,10 +124,14 @@ namespace HordeAgent.Leases
 			catch (InsufficientSpaceException ex)
 			{
 				logger.LogError(ex, "{Message}", ex.Message);
+				span.SetStatus(Status.Error);
+				span.RecordException(ex);
 			}
 			catch (Exception ex)
 			{
 				logger.LogError(ex, "Unhandled exception while executing lease {LeaseId}: {Message}", Id, ex.Message);
+				span.SetStatus(Status.Error);
+				span.RecordException(ex);
 			}
 
 			// Update the state of the lease
@@ -149,6 +151,8 @@ namespace HordeAgent.Leases
 			RpcLease = newRpcLease;
 
 			logger.LogInformation("Transitioning lease {LeaseId} to {State}, outcome={Outcome}", Id, RpcLease.State, RpcLease.Outcome);
+			span.SetAttribute("horde.lease.state", RpcLease.State.ToString());
+			span.SetAttribute("horde.lease.outcome", RpcLease.Outcome.ToString());
 
 			return result;
 		}
@@ -156,20 +160,18 @@ namespace HordeAgent.Leases
 		/// <summary>
 		/// Dispatch a lease payload to the appropriate handler
 		/// </summary>
-		internal async Task<LeaseResult> HandleLeasePayloadAsync(ISession session, LeaseLoggerFactory leaseLoggerFactory)
+		internal async Task<LeaseResult> HandleLeasePayloadAsync(ISession session, Tracer tracer, LeaseLoggerFactory leaseLoggerFactory)
 		{
 			using ILoggerFactory loggerFactory = leaseLoggerFactory.CreateLoggerFactory(Id);
 			ILogger leaseLogger = loggerFactory.CreateLogger(GetType());
-
-			GlobalTracer.Instance.ActiveSpan?.SetTag("task", RpcPayload.TypeUrl);
-			return await ExecuteAsync(session, leaseLogger, _cancellationSource.Token);
+			return await ExecuteAsync(session, tracer, leaseLogger, _cancellationSource.Token);
 		}
 
 		/// <summary>
 		/// Executes a lease
 		/// </summary>
 		/// <returns>Result for the lease</returns>
-		protected abstract Task<LeaseResult> ExecuteAsync(ISession session, ILogger logger, CancellationToken cancellationToken);
+		protected abstract Task<LeaseResult> ExecuteAsync(ISession session, Tracer tracer, ILogger logger, CancellationToken cancellationToken);
 
 		/// <summary>
 		/// Runs a child process, piping the output to the given logger
@@ -256,13 +258,13 @@ namespace HordeAgent.Leases
 		}
 
 		/// <inheritdoc/>
-		protected override Task<LeaseResult> ExecuteAsync(ISession session, ILogger logger, CancellationToken cancellationToken)
+		protected override Task<LeaseResult> ExecuteAsync(ISession session, Tracer tracer, ILogger logger, CancellationToken cancellationToken)
 		{
-			return ExecuteAsync(session, Id, RpcLease.Payload.Unpack<T>(), logger, cancellationToken);
+			return ExecuteAsync(session, Id, RpcLease.Payload.Unpack<T>(), tracer, logger, cancellationToken);
 		}
 
-		/// <inheritdoc/>
-		protected abstract Task<LeaseResult> ExecuteAsync(ISession session, LeaseId leaseId, T message, ILogger logger, CancellationToken cancellationToken);
+		/// <inheritdoc cref="LeaseHandler.ExecuteAsync" />
+		protected abstract Task<LeaseResult> ExecuteAsync(ISession session, LeaseId leaseId, T message, Tracer tracer, ILogger logger, CancellationToken cancellationToken);
 	}
 
 	class DefaultLeaseHandler : LeaseHandler
@@ -273,7 +275,7 @@ namespace HordeAgent.Leases
 			=> _result = result;
 
 		/// <inheritdoc/>
-		protected override Task<LeaseResult> ExecuteAsync(ISession session, ILogger logger, CancellationToken cancellationToken)
+		protected override Task<LeaseResult> ExecuteAsync(ISession session, Tracer tracer, ILogger logger, CancellationToken cancellationToken)
 			=> Task.FromResult(_result);
 	}
 }
diff --git a/Engine/Source/Programs/Horde/HordeAgent/Leases/LeaseManager.cs b/Engine/Source/Programs/Horde/HordeAgent/Leases/LeaseManager.cs
--- a/Engine/Source/Programs/Horde/HordeAgent/Leases/LeaseManager.cs
+++ b/Engine/Source/Programs/Horde/HordeAgent/Leases/LeaseManager.cs
@@ -11,6 +11,8 @@ using HordeCommon.Rpc.Messages;
 using Microsoft.Extensions.DependencyInjection;
 using Microsoft.Extensions.Logging;
 using Microsoft.Extensions.Options;
+using OpenTelemetry.Trace;
+using StatusCode = Grpc.Core.StatusCode;
 
 namespace HordeAgent.Leases
 {
@@ -88,11 +90,21 @@ namespace HordeAgent.Leases
 		readonly Dictionary<string, LeaseHandlerFactory> _typeUrlToLeaseHandler;
 		readonly LeaseLoggerFactory _leaseLoggerFactory;
 		readonly IOptions<AgentSettings> _settings;
+		readonly Tracer _tracer;
 		readonly ILogger _logger;
 
 		RpcAgentCapabilities? _capabilities;
 
-		public LeaseManager(ISession session, CapabilitiesService capabilitiesService, StatusService statusService, ISystemMetrics systemMetrics, IEnumerable<LeaseHandlerFactory> leaseHandlerFactories, LeaseLoggerFactory leaseLoggerFactory, IOptions<AgentSettings> settings, ILogger logger)
+		public LeaseManager(
+			ISession session,
+			CapabilitiesService capabilitiesService,
+			StatusService statusService,
+			ISystemMetrics systemMetrics,
+			IEnumerable<LeaseHandlerFactory> leaseHandlerFactories,
+			LeaseLoggerFactory leaseLoggerFactory,
+			IOptions<AgentSettings> settings,
+			Tracer tracer,
+			ILogger logger)
 		{
 			_session = session;
 			_capabilitiesService = capabilitiesService;
@@ -101,6 +113,7 @@ namespace HordeAgent.Leases
 			_typeUrlToLeaseHandler = leaseHandlerFactories.ToDictionary(x => x.LeaseType, x => x);
 			_leaseLoggerFactory = leaseLoggerFactory;
 			_settings = settings;
+			_tracer = tracer;
 			_logger = logger;
 		}
 
@@ -112,6 +125,7 @@ namespace HordeAgent.Leases
 				serviceProvider.GetRequiredService<IEnumerable<LeaseHandlerFactory>>(),
 				serviceProvider.GetRequiredService<LeaseLoggerFactory>(),
 				serviceProvider.GetRequiredService<IOptions<AgentSettings>>(),
+				serviceProvider.GetRequiredService<Tracer>(),
 				serviceProvider.GetRequiredService<ILogger<LeaseManager>>())
 		{
 		}
@@ -360,7 +374,7 @@ namespace HordeAgent.Leases
 									_logger.LogInformation("Adding lease {LeaseId}", serverLease.Id);
 
 									LeaseHandler leaseHandler = CreateLeaseHandler(serverLease);
-									leaseHandler.Start(_session, _logger, _leaseLoggerFactory);
+									leaseHandler.Start(_session, _tracer, _logger, _leaseLoggerFactory);
 									leaseHandler.Result.ContinueWith((Task<LeaseResult> task) =>
 									{
 										LeaseResult result = task.Result;
diff --git a/Engine/Source/Programs/Horde/HordeAgent/Logging.cs b/Engine/Source/Programs/Horde/HordeAgent/Logging.cs
--- a/Engine/Source/Programs/Horde/HordeAgent/Logging.cs
+++ b/Engine/Source/Programs/Horde/HordeAgent/Logging.cs
@@ -2,10 +2,9 @@
 
 using System.Runtime.InteropServices;
 using EpicGames.Core;
+using HordeAgent.Utility;
 using Microsoft.Extensions.Configuration;
 using Microsoft.Extensions.Logging;
-using OpenTracing;
-using OpenTracing.Util;
 using Serilog;
 using Serilog.Core;
 using Serilog.Extensions.Logging;
@@ -16,29 +15,13 @@ namespace HordeAgent
 {
 	static class Logging
 	{
-		static string s_env = "default";
+		public static LoggingLevelSwitch LogLevelSwitch = new ();
 
-		public static void SetEnv(string newEnv)
-		{
-			s_env = newEnv;
-		}
-
-		public static LoggingLevelSwitch LogLevelSwitch = new LoggingLevelSwitch();
-
-		private class DatadogLogEnricher : ILogEventEnricher
+		private class DatadogVersionLogEnricher : ILogEventEnricher
 		{
 			public void Enrich(Serilog.Events.LogEvent logEvent, ILogEventPropertyFactory propertyFactory)
 			{
-				logEvent.AddOrUpdateProperty(propertyFactory.CreateProperty("dd.env", s_env));
-				logEvent.AddOrUpdateProperty(propertyFactory.CreateProperty("dd.service", "hordeagent"));
-				logEvent.AddOrUpdateProperty(propertyFactory.CreateProperty("dd.version", AgentApp.Version));
-
-				ISpan? span = GlobalTracer.Instance?.ActiveSpan;
-				if (span != null)
-				{
-					logEvent.AddPropertyIfAbsent(propertyFactory.CreateProperty("dd.trace_id", span.Context.TraceId));
-					logEvent.AddPropertyIfAbsent(propertyFactory.CreateProperty("dd.span_id", span.Context.SpanId));
-				}
+				logEvent.AddPropertyIfAbsent(propertyFactory.CreateProperty("dd.version", AgentApp.Version));
 			}
 		}
 
@@ -62,23 +45,29 @@ namespace HordeAgent
 			{
 				theme = AnsiConsoleTheme.Code;
 			}
-
-			return new LoggerConfiguration()
+			
+			LoggerConfiguration loggerConfig = new LoggerConfiguration()
 				.WriteTo.Console(outputTemplate: "[{Timestamp:HH:mm:ss} {Level:w3}] {Indent}{Message:l}{NewLine}{Exception}", theme: theme)
 				.WriteTo.File(FileReference.Combine(settings.LogsDir, "Log-.txt").FullName, fileSizeLimitBytes: 50 * 1024 * 1024, rollingInterval: RollingInterval.Day, rollOnFileSizeLimit: true, retainedFileCountLimit: 10)
 				.WriteTo.File(new JsonFormatter(renderMessage: true), FileReference.Combine(settings.LogsDir, "Log-.json").FullName, fileSizeLimitBytes: 50 * 1024 * 1024, rollingInterval: RollingInterval.Day, rollOnFileSizeLimit: true, retainedFileCountLimit: 10)
 				.ReadFrom.Configuration(configuration)
 				.MinimumLevel.ControlledBy(LogLevelSwitch)
-				.Enrich.FromLogContext()
-				.Enrich.With<DatadogLogEnricher>()
-				.CreateLogger();
+				.Enrich.FromLogContext();
+			
+			if (settings.OpenTelemetry.EnableDatadogCompatibility)
+			{
+				loggerConfig = loggerConfig.Enrich.With<DatadogVersionLogEnricher>();
+				loggerConfig = loggerConfig.Enrich.With<OpenTelemetryDatadogLogEnricher>();
+			}
+			
+			return loggerConfig.CreateLogger();
 		}
 
 		public static ILoggerProvider CreateFileLoggerProvider(DirectoryReference baseDir, string name)
 		{
 			DirectoryReference.CreateDirectory(baseDir);
 
-			Serilog.Core.Logger logger = new LoggerConfiguration()
+			Logger logger = new LoggerConfiguration()
 				.WriteTo.File(FileReference.Combine(baseDir, $"{name}.txt").FullName)
 				.WriteTo.File(new JsonFormatter(renderMessage: true), FileReference.Combine(baseDir, $"{name}.json").FullName)
 				.Enrich.FromLogContext()
diff --git a/Engine/Source/Programs/Horde/HordeAgent/Services/CapabilitiesService.cs b/Engine/Source/Programs/Horde/HordeAgent/Services/CapabilitiesService.cs
--- a/Engine/Source/Programs/Horde/HordeAgent/Services/CapabilitiesService.cs
+++ b/Engine/Source/Programs/Horde/HordeAgent/Services/CapabilitiesService.cs
@@ -17,8 +17,7 @@ using HordeCommon.Rpc.Messages;
 using Microsoft.Extensions.Logging;
 using Microsoft.Extensions.Options;
 using Microsoft.Management.Infrastructure;
-using OpenTracing;
-using OpenTracing.Util;
+using OpenTelemetry.Trace;
 using AddressFamily = System.Net.Sockets.AddressFamily;
 
 namespace HordeAgent.Services
@@ -32,14 +31,16 @@ namespace HordeAgent.Services
 		static readonly DateTimeOffset s_bootTime = DateTimeOffset.Now - TimeSpan.FromTicks(Environment.TickCount64 * TimeSpan.TicksPerMillisecond);
 
 		readonly AgentSettings _settings;
+		readonly Tracer _tracer;
 		readonly ILogger _logger;
 
 		/// <summary>
 		/// Constructor
 		/// </summary>
-		public CapabilitiesService(IOptions<AgentSettings> settings, ILogger<CapabilitiesService> logger)
+		public CapabilitiesService(IOptions<AgentSettings> settings, Tracer tracer, ILogger<CapabilitiesService> logger)
 		{
 			_settings = settings.Value;
+			_tracer = tracer;
 			_logger = logger;
 		}
 
@@ -79,7 +80,7 @@ namespace HordeAgent.Services
 
 		async Task<RpcAgentCapabilities> GetCapabilitiesInternalAsync(DirectoryReference? workingDir)
 		{
-			using IScope scope = GlobalTracer.Instance.BuildSpan($"{nameof(CapabilitiesService)}.{nameof(GetCapabilitiesInternalAsync)}").StartActive();
+			using TelemetrySpan span = _tracer.StartActiveSpan($"{nameof(CapabilitiesService)}.{nameof(GetCapabilitiesInternalAsync)}");
 			ILogger logger = _logger;
 
 			// Create the primary device
@@ -87,7 +88,7 @@ namespace HordeAgent.Services
 
 			if (RuntimeInformation.IsOSPlatform(OSPlatform.Windows))
 			{
-				using IScope winScope = GlobalTracer.Instance.BuildSpan("GetWindowsCapabilities").StartActive();
+				using TelemetrySpan winSpan = _tracer.StartActiveSpan("GetWindowsCapabilities");
 				capabilities.Properties.Add($"{KnownPropertyNames.Platform}=Win64");
 				capabilities.Properties.Add($"{KnownPropertyNames.PlatformGroup}=Windows");
 				capabilities.Properties.Add($"{KnownPropertyNames.PlatformGroup}=Microsoft");
@@ -104,7 +105,7 @@ namespace HordeAgent.Services
 					
 					{
 						// Add OS info
-						using IScope _ = GlobalTracer.Instance.BuildSpan("GetOsInfo").StartActive();
+						using TelemetrySpan _ = _tracer.StartActiveSpan("GetOsInfo");
 						foreach (CimInstance instance in session.QueryInstances(QueryNamespace, QueryDialect, "select * from Win32_OperatingSystem"))
 						{
 							foreach (CimProperty property in instance.CimInstanceProperties)
@@ -124,7 +125,7 @@ namespace HordeAgent.Services
 					
 					{
 						// Add CPU info
-						using IScope _ = GlobalTracer.Instance.BuildSpan("GetCpuInfo").StartActive();
+						using TelemetrySpan _ = _tracer.StartActiveSpan("GetCpuInfo");
 						Dictionary<string, int> cpuNameToCount = new Dictionary<string, int>();
 						int totalPhysicalCores = 0;
 						int totalLogicalCores = 0;
@@ -164,7 +165,7 @@ namespace HordeAgent.Services
 					
 					{
 						// Add RAM info
-						using IScope _ = GlobalTracer.Instance.BuildSpan("GetRamInfo").StartActive();
+						using TelemetrySpan _ = _tracer.StartActiveSpan("GetRamInfo");
 						ulong totalCapacity = 0;
 						foreach (CimInstance instance in session.QueryInstances(QueryNamespace, QueryDialect, "select Capacity from Win32_PhysicalMemory"))
 						{
@@ -182,7 +183,7 @@ namespace HordeAgent.Services
 					
 					{
 						// Add GPU info
-						using IScope _ = GlobalTracer.Instance.BuildSpan("GetGpuInfo").StartActive();
+						using TelemetrySpan _ = _tracer.StartActiveSpan("GetGpuInfo");
 						int index = 0;
 						foreach (CimInstance instance in session.QueryInstances(QueryNamespace, QueryDialect, "select Name, DriverVersion, AdapterRAM from Win32_VideoController"))
 						{
@@ -229,7 +230,7 @@ namespace HordeAgent.Services
 			}
 			else if (RuntimeInformation.IsOSPlatform(OSPlatform.Linux))
 			{
-				using IScope _ = GlobalTracer.Instance.BuildSpan("GetLinuxCapabilities").StartActive();
+				using TelemetrySpan _ = _tracer.StartActiveSpan("GetLinuxCapabilities");
 				capabilities.Properties.Add($"{KnownPropertyNames.Platform}=Linux");
 				capabilities.Properties.Add($"{KnownPropertyNames.PlatformGroup}=Linux");
 				capabilities.Properties.Add($"{KnownPropertyNames.PlatformGroup}=Unix");
@@ -307,7 +308,7 @@ namespace HordeAgent.Services
 			}
 			else if (RuntimeInformation.IsOSPlatform(OSPlatform.OSX))
 			{
-				using IScope _ = GlobalTracer.Instance.BuildSpan("GetMacCapabilities").StartActive();
+				using TelemetrySpan _ = _tracer.StartActiveSpan("GetMacCapabilities");
 				capabilities.Properties.Add($"{KnownPropertyNames.Platform}=Mac");
 				capabilities.Properties.Add($"{KnownPropertyNames.PlatformGroup}=Apple");
 				capabilities.Properties.Add($"{KnownPropertyNames.PlatformGroup}=Desktop");
@@ -388,7 +389,7 @@ namespace HordeAgent.Services
 			// Get the IP addresses
 			try
 			{
-				using IScope _ = GlobalTracer.Instance.BuildSpan("ResolveIp").StartActive();
+				using TelemetrySpan _ = _tracer.StartActiveSpan("ResolveIp");
 				using CancellationTokenSource dnsCts = new(3000);
 				IPHostEntry entry = await Dns.GetHostEntryAsync(Dns.GetHostName(), dnsCts.Token);
 				foreach (IPAddress address in entry.AddressList)
@@ -516,9 +517,9 @@ namespace HordeAgent.Services
 		/// <param name="hostname">A hostname to test against</param>
 		/// <param name="timeoutMs">Max time to wait for a connect, in milliseconds</param>
 		/// <returns>Local IP address of this machine</returns>
-		public static async Task<IPAddress?> GetLocalIpAddressAsync(string hostname, int timeoutMs = 2000)
+		public async Task<IPAddress?> GetLocalIpAddressAsync(string hostname, int timeoutMs = 2000)
 		{
-			using IScope _ = GlobalTracer.Instance.BuildSpan(nameof(GetLocalIpAddressAsync)).StartActive();
+			using TelemetrySpan _ = _tracer.StartActiveSpan(nameof(GetLocalIpAddressAsync));
 			try
 			{
 				using Socket socket = new(AddressFamily.InterNetwork, SocketType.Dgram, ProtocolType.IP);
@@ -614,9 +615,9 @@ namespace HordeAgent.Services
 			}
 		}
 
-		static async Task AddAwsPropertiesInternalAsync(IList<string> properties, ILogger logger)
+		async Task AddAwsPropertiesInternalAsync(IList<string> properties, ILogger logger)
 		{
-			using IScope scope = GlobalTracer.Instance.BuildSpan(nameof(AddAwsPropertiesInternalAsync)).StartActive();
+			using TelemetrySpan _ = _tracer.StartActiveSpan(nameof(AddAwsPropertiesInternalAsync));
 			if (EC2InstanceMetadata.IdentityDocument != null)
 			{
 				properties.Add("EC2=1");
diff --git a/Engine/Source/Programs/Horde/HordeAgent/Tracing.cs b/Engine/Source/Programs/Horde/HordeAgent/Tracing.cs
deleted file mode 100644
diff --git a/Engine/Source/Programs/Horde/HordeAgent/Utility/OpenTelemetry.cs b/Engine/Source/Programs/Horde/HordeAgent/Utility/OpenTelemetry.cs
new file mode 100644
--- /dev/null
+++ b/Engine/Source/Programs/Horde/HordeAgent/Utility/OpenTelemetry.cs
@@ -0,0 +1,218 @@
+// Copyright Epic Games, Inc. All Rights Reserved.
+
+using System.Diagnostics;
+using System.Diagnostics.Metrics;
+using EpicGames.Horde.Utilities;
+using Microsoft.Extensions.DependencyInjection;
+using Microsoft.Extensions.Logging;
+using OpenTelemetry.Exporter;
+using OpenTelemetry.Logs;
+using OpenTelemetry.Metrics;
+using OpenTelemetry.Resources;
+using OpenTelemetry.Trace;
+using Serilog.Core;
+
+namespace HordeAgent.Utility;
+
+/// <summary>
+/// Serilog event enricher attaching trace and span ID for Datadog using current System.Diagnostics.Activity
+/// </summary>
+public class OpenTelemetryDatadogLogEnricher : ILogEventEnricher
+{
+	/// <inheritdoc />
+	public void Enrich(Serilog.Events.LogEvent logEvent, ILogEventPropertyFactory propertyFactory)
+	{
+		if (Activity.Current != null)
+		{
+			string stringTraceId = Activity.Current.TraceId.ToString();
+			string stringSpanId = Activity.Current.SpanId.ToString();
+			string ddTraceId = Convert.ToUInt64(stringTraceId.Substring(16), 16).ToString();
+			string ddSpanId = Convert.ToUInt64(stringSpanId, 16).ToString();
+
+			logEvent.AddPropertyIfAbsent(propertyFactory.CreateProperty("dd.trace_id", ddTraceId));
+			logEvent.AddPropertyIfAbsent(propertyFactory.CreateProperty("dd.span_id", ddSpanId));
+		}
+	}
+}
+
+/// <summary>
+/// Helper for configuring OpenTelemetry
+/// </summary>
+public static class OpenTelemetryHelper
+{
+	private static ResourceBuilder? s_resourceBuilder;
+	
+	/// <summary>
+	/// Name of resource name attribute in Datadog
+	/// Some traces use this for prettier display inside their UI
+	/// </summary>
+	public const string DatadogResourceAttribute = "resource.name";
+	
+	/// <summary>
+	/// Name of default Horde tracer (aka activity source)
+	/// </summary>
+	public const string HordeName = "Horde";
+	
+	/// <summary>
+	/// List of all source names configured in this class.
+	/// They are needed at startup when initializing OpenTelemetry
+	/// </summary>
+	public static string[] SourceNames => new[] { HordeName };
+	
+	/// <summary>
+	/// Default tracer used in Horde
+	/// Prefer dependency-injected tracer over this static member.
+	/// </summary>
+	public static readonly Tracer Horde = TracerProvider.Default.GetTracer(HordeName);
+
+	/// <summary>
+	/// Configure OpenTelemetry in Horde and ASP.NET
+	/// </summary>
+	/// <param name="services">Service collection for DI</param>
+	/// <param name="settings">Current server settings</param>
+	public static void Configure(IServiceCollection services, OpenTelemetrySettings settings)
+	{
+		// Always configure tracers/meters as they are used in the codebase even when OpenTelemetry is not configured
+		services.AddSingleton(Horde);
+		services.AddSingleton(OpenTelemetryMeters.Horde);
+
+		if (!settings.Enabled)
+		{
+			return;
+		}
+
+		services.AddOpenTelemetry()
+			.WithTracing(builder => ConfigureTracing(builder, settings))
+			.WithMetrics(builder => ConfigureMetrics(builder, settings));
+	}
+
+	private static void ConfigureTracing(TracerProviderBuilder builder, OpenTelemetrySettings settings)
+	{
+		void DatadogHttpRequestEnricher(Activity activity, HttpRequestMessage message)
+		{
+			activity.SetTag("service.name", settings.ServiceName + "-http-client");
+			activity.SetTag("operation.name", "http.request");
+			string url = $"{message.Method} {message.Headers.Host}{message.RequestUri?.LocalPath}";
+			activity.DisplayName = url;
+			activity.SetTag("resource.name", url);
+		}
+
+		builder
+			.AddSource(SourceNames)
+			.AddHttpClientInstrumentation(options =>
+			{
+				if (settings.EnableDatadogCompatibility)
+				{
+					options.EnrichWithHttpRequestMessage = DatadogHttpRequestEnricher;
+				}
+			})
+			.AddGrpcClientInstrumentation(options =>
+			{
+				if (settings.EnableDatadogCompatibility)
+				{
+					options.EnrichWithHttpRequestMessage = DatadogHttpRequestEnricher;
+				}
+			})
+			.SetResourceBuilder(GetResourceBuilder(settings));
+
+		if (settings.EnableConsoleExporter)
+		{
+			builder.AddConsoleExporter();
+		}
+
+		foreach ((string name, OpenTelemetryProtocolExporterSettings exporterSettings) in settings.ProtocolExporters)
+		{
+			builder.AddOtlpExporter(name, exporter =>
+			{
+				exporter.Endpoint = exporterSettings.Endpoint!;
+				exporter.Protocol = Enum.TryParse(exporterSettings.Protocol, true, out OtlpExportProtocol protocol) ? protocol : OtlpExportProtocol.Grpc;
+			});
+		}
+	}
+
+	private static void ConfigureMetrics(MeterProviderBuilder builder, OpenTelemetrySettings settings)
+	{
+		builder.AddMeter(OpenTelemetryMeters.MeterNames);
+		builder.AddAspNetCoreInstrumentation();
+		builder.AddHttpClientInstrumentation();
+
+		if (settings.EnableConsoleExporter)
+		{
+			builder.AddConsoleExporter();
+		}
+
+		foreach ((string name, OpenTelemetryProtocolExporterSettings exporterSettings) in settings.ProtocolExporters)
+		{
+			builder.AddOtlpExporter(name, exporter =>
+			{
+				exporter.Endpoint = exporterSettings.Endpoint!;
+				exporter.Protocol = Enum.TryParse(exporterSettings.Protocol, true, out OtlpExportProtocol protocol) ? protocol : OtlpExportProtocol.Grpc;
+			});
+		}
+	}
+
+	/// <summary>
+	/// Configure .NET logging with OpenTelemetry
+	/// </summary>
+	/// <param name="builder">Logging builder to modify</param>
+	/// <param name="settings">Current server settings</param>
+	public static void ConfigureLogging(ILoggingBuilder builder, OpenTelemetrySettings settings)
+	{
+		if (!settings.Enabled)
+		{
+			return;
+		}
+
+		builder.AddOpenTelemetry(options =>
+		{
+			options.IncludeScopes = true;
+			options.IncludeFormattedMessage = true;
+			options.ParseStateValues = true;
+			options.SetResourceBuilder(GetResourceBuilder(settings));
+
+			if (settings.EnableConsoleExporter)
+			{
+				options.AddConsoleExporter();
+			}
+		});
+	}
+
+	private static ResourceBuilder GetResourceBuilder(OpenTelemetrySettings settings)
+	{
+		if (s_resourceBuilder != null)
+		{
+			return s_resourceBuilder;
+		}
+
+		List<KeyValuePair<string, object>> attributes = settings.Attributes.Select(x => new KeyValuePair<string, object>(x.Key, x.Value)).ToList();
+		s_resourceBuilder = ResourceBuilder.CreateDefault()
+			.AddService(settings.ServiceName, serviceNamespace: settings.ServiceNamespace, serviceVersion: settings.ServiceVersion)
+			.AddAttributes(attributes)
+			.AddTelemetrySdk()
+			.AddEnvironmentVariableDetector();
+
+		return s_resourceBuilder;
+	}
+}
+
+/// <summary>
+/// Static initialization of all available OpenTelemetry meters
+/// </summary>
+public static class OpenTelemetryMeters
+{
+	/// <summary>
+	/// Name of default Horde meter
+	/// </summary>
+	public const string HordeName = "Horde";
+
+	/// <summary>
+	/// List of all source names configured in this class.
+	/// They are needed at startup when initializing OpenTelemetry
+	/// </summary>
+	public static string[] MeterNames => new[] { HordeName };
+
+	/// <summary>
+	/// Default meter used in Horde
+	/// </summary>
+	public static readonly Meter Horde = new(HordeName);
+}
\ No newline at end of file
