Horde: Fix issue with changelist link when swarm isn't configured
#rnx
#skipci

diff --git a/Engine/Source/Programs/Horde/HordeDashboard/src/components/IssueViewV2.tsx b/Engine/Source/Programs/Horde/HordeDashboard/src/components/IssueViewV2.tsx
--- a/Engine/Source/Programs/Horde/HordeDashboard/src/components/IssueViewV2.tsx
+++ b/Engine/Source/Programs/Horde/HordeDashboard/src/components/IssueViewV2.tsx
@@ -1332,11 +1332,11 @@ const IssueSummaryPanel: React.FC = () => {
 
    if (issue.fingerprintDescription) {
       const components = issue.fingerprintDescription.split(" / ");
-      if (components.length) {         
+      if (components.length) {
          items.push({
             title: "Fingerprint",
             text: components[0].replace("Type:", "").trim()
-         })   
+         })
       }
    }
 
@@ -1568,9 +1568,10 @@ const StepPanel: React.FC<{ streamId: string, hstep: GetIssueStepResponse }> = o
                         <Stack style={{ height: "100%" }}>
                            <Stack horizontal verticalAlign="center" tokens={{ childrenGap: 8 }}>
                               <Icon styles={{ root: { margin: '0px', padding: '0px', paddingTop: "2px", userSelect: "none" } }} iconName="Commit" className={hordeClasses.iconBlue} />
-                              <a href={changeUrl} target="blank">
+                              {!!dashboard.swarmUrl && <a href={changeUrl} target="blank">
                                  <Text variant={textSize} styles={{ root: { margin: '0px', padding: '0px' } }} >{`CL ${hstep.change}`}</Text>
-                              </a>
+                              </a>}
+                              {!dashboard.swarmUrl && <Text variant={textSize} styles={{ root: { margin: '0px', padding: '0px' } }} >{`CL ${hstep.change}`}</Text>}
                            </Stack>
                         </Stack>
                         <Stack style={{ paddingRight: 18 }}>
@@ -1937,7 +1938,7 @@ export const IssueModalV2: React.FC<{ popHistoryOnClose: boolean, issueId?: stri
    }
 
    // subscribe
-   if (details.update) { }   
+   if (details.update) { }
 
    if (details.issueError) {
       return <Dialog hidden={false} onDismiss={() => { details.clear(); if (onCloseExternal) { onCloseExternal() } else { onClose() } }} dialogContentProps={{
@@ -1947,10 +1948,10 @@ export const IssueModalV2: React.FC<{ popHistoryOnClose: boolean, issueId?: stri
       }}
          modalProps={{ styles: { main: { width: "640px !important", minWidth: "640px !important", maxWidth: "640px !important" } } }}>
          <DialogFooter>
-            <PrimaryButton onClick={() => { details.clear();  if (onCloseExternal) { onCloseExternal() } else { onClose() } }} text="Ok" />
+            <PrimaryButton onClick={() => { details.clear(); if (onCloseExternal) { onCloseExternal() } else { onClose() } }} text="Ok" />
          </DialogFooter>
       </Dialog>
-   }   
+   }
 
    details.set(parseInt(issueId));
 
