Horde: Address white bar at bottom of screen
#rnx
#skipci

diff --git a/Engine/Source/Programs/Horde/HordeDashboard/src/App.tsx b/Engine/Source/Programs/Horde/HordeDashboard/src/App.tsx
--- a/Engine/Source/Programs/Horde/HordeDashboard/src/App.tsx
+++ b/Engine/Source/Programs/Horde/HordeDashboard/src/App.tsx
@@ -85,7 +85,7 @@ const Main: React.FC = () => {
          ErrorHandler.set({ title: "Error initializing site", reason: reason }, true);
       });
 
-      return (<ThemeProvider applyTo='body' theme={dashboard.darktheme ? darkTheme : lightTheme}>
+      return (<ThemeProvider applyTo='element' theme={dashboard.darktheme ? darkTheme : lightTheme}>
          <div style={{ position: 'absolute', left: '50%', top: '50%', transform: 'translate(-50%, -50%)' }}>
             <Stack horizontalAlign="center" styles={{ root: { padding: 20, minWidth: 200, minHeight: 100 } }}>
                <Stack horizontal>
@@ -183,8 +183,13 @@ const Main: React.FC = () => {
       router = createBrowserRouter(routes);
    }
 
+   // avoid white bar on body
+   
+   document.body.setAttribute('style', `background: ${dashboard.darktheme ? "#0F0F0F" : "#FAF9F9"}`)   
+   
+   
    return (
-      <ThemeProvider applyTo='body' theme={dashboard.darktheme ? darkTheme : lightTheme}>
+      <ThemeProvider applyTo='element' theme={dashboard.darktheme ? darkTheme : lightTheme}>
          <RouterProvider router={router} />
       </ThemeProvider>
    );
diff --git a/Engine/Source/Programs/Horde/HordeDashboard/src/components/UserHome.tsx b/Engine/Source/Programs/Horde/HordeDashboard/src/components/UserHome.tsx
--- a/Engine/Source/Programs/Horde/HordeDashboard/src/components/UserHome.tsx
+++ b/Engine/Source/Programs/Horde/HordeDashboard/src/components/UserHome.tsx
@@ -1095,7 +1095,7 @@ const UserHomeViewInner: React.FC = () => {
       {(!!change || !!search.has("preflightconfig")) && <PreflightConfigModal onClose={() => { navigate("/index", { replace: true }) }} />}
       <Stack style={{ padding: 0 }} className={detailClasses.detailsRow}>
          <FocusZone direction={FocusZoneDirection.vertical} style={{ padding: 0 }}>
-            <div className={detailClasses.container} style={{ width: "100%", height: 'calc(100vh - 208px)', position: 'relative' }} data-is-scrollable={true}>
+            <div className={detailClasses.container} style={{ width: "100%", height: 'calc(100vh - 192px)', position: 'relative' }} data-is-scrollable={true}>
                <ScrollablePane scrollbarVisibility={ScrollbarVisibility.auto} onScroll={() => { }}>
                   <Stack tokens={{ childrenGap: 18 }} style={{ padding: 0 }}>
                      <ProjectsPanel />
diff --git a/Engine/Source/Programs/Horde/HordeDashboard/src/components/accounts/HordeLoginView.tsx b/Engine/Source/Programs/Horde/HordeDashboard/src/components/accounts/HordeLoginView.tsx
--- a/Engine/Source/Programs/Horde/HordeDashboard/src/components/accounts/HordeLoginView.tsx
+++ b/Engine/Source/Programs/Horde/HordeDashboard/src/components/accounts/HordeLoginView.tsx
@@ -61,7 +61,7 @@ export const HordeLoginView: React.FC = () => {
       }
    }
 
-   return (<ThemeProvider applyTo='body' theme={dashboard.darktheme ? darkTheme : lightTheme}>
+   return (<ThemeProvider applyTo='element' theme={dashboard.darktheme ? darkTheme : lightTheme}>
       <div style={{ position: 'absolute', left: '50%', top: '50%', transform: 'translate(-50%, -50%)' }}>
          <Stack>
             <Stack horizontalAlign="center">
diff --git a/Engine/Source/Programs/Horde/HordeDashboard/src/components/setup/HordeSetupView.tsx b/Engine/Source/Programs/Horde/HordeDashboard/src/components/setup/HordeSetupView.tsx
--- a/Engine/Source/Programs/Horde/HordeDashboard/src/components/setup/HordeSetupView.tsx
+++ b/Engine/Source/Programs/Horde/HordeDashboard/src/components/setup/HordeSetupView.tsx
@@ -96,7 +96,7 @@ export const HordeSetupView: React.FC = () => {
    }
 
 
-   return (<ThemeProvider applyTo='body' theme={dashboard.darktheme ? darkTheme : lightTheme}>
+   return (<ThemeProvider applyTo='element' theme={dashboard.darktheme ? darkTheme : lightTheme}>
       <div style={{ position: 'absolute', left: '50%', top: '50%', transform: 'translate(-50%, -50%)' }}>
          <Stack>
             <Stack horizontalAlign="center">
