Horde: Add missing file.

diff --git a/Engine/Source/Programs/Horde/HordeTrayApp/HordeClientProvider.cs b/Engine/Source/Programs/Horde/HordeTrayApp/HordeClientProvider.cs
new file mode 100644
--- /dev/null
+++ b/Engine/Source/Programs/Horde/HordeTrayApp/HordeClientProvider.cs
@@ -0,0 +1,204 @@
+// Copyright Epic Games, Inc. All Rights Reserved.
+
+using System.Diagnostics;
+using EpicGames.Horde;
+using Microsoft.Extensions.DependencyInjection;
+using Microsoft.Extensions.Logging;
+
+namespace HordeTrayApp
+{
+	/// <summary>
+	/// Implements a mechanism for creating scoped HordeClient instances, and modifying configuration settings
+	/// </summary>
+	class HordeClientProvider : IHordeClientProvider, IAsyncDisposable
+	{
+		class HordeClientLifetime : IAsyncDisposable
+		{
+			readonly TaskCompletionSource _refZeroTcs = new TaskCompletionSource();
+
+			readonly Task<IHordeClient?> _clientTask;
+			int _refCount;
+			bool _disposed;
+
+			public Task<IHordeClient?> ClientTask => _clientTask;
+
+			public HordeClientLifetime(Task<IHordeClient?> clientTask)
+			{
+				_clientTask = clientTask;
+				_refCount = 1;
+			}
+
+			public async ValueTask DisposeAsync()
+			{
+				if (!_disposed)
+				{
+					Release();
+					await _refZeroTcs.Task;
+					_disposed = true;
+				}
+			}
+
+			public void AddRef()
+			{
+				int refCount = Interlocked.Increment(ref _refCount);
+				Debug.Assert(refCount > 1);
+			}
+
+			public void Release()
+			{
+				if (Interlocked.Decrement(ref _refCount) == 0)
+				{
+					_refZeroTcs.SetResult();
+				}
+			}
+		}
+
+		class HordeClientRef : IHordeClientRef
+		{
+			IHordeClient _client;
+			HordeClientLifetime? _lifetime;
+
+			public IHordeClient Client
+				=> _client ?? throw new ObjectDisposedException(null);
+
+			public HordeClientRef(IHordeClient client, HordeClientLifetime lifetime)
+			{
+				_client = client;
+				_lifetime = lifetime;
+			}
+
+			public void Dispose()
+			{
+				_lifetime?.Release();
+				_lifetime = null;
+				_client = null!;
+			}
+		}
+
+		readonly object _lockObject = new object();
+		readonly ILoggerFactory _loggerFactory;
+		ServiceProvider? _serviceProvider;
+		IHordeClient? _hordeClient;
+		HordeClientLifetime _lifetime;
+		readonly ILogger _logger;
+
+		/// <summary>
+		/// Event signalled whenever the connection state changes
+		/// </summary>
+		public event Action? OnStateChanged;
+
+		/// <summary>
+		/// Event signalled whenever the access token state changes
+		/// </summary>
+		public event Action? OnAccessTokenStateChanged;
+
+		/// <summary>
+		/// Constructor
+		/// </summary>
+		public HordeClientProvider(ILoggerFactory loggerFactory, ILogger<HordeClientProvider> logger)
+		{
+			_loggerFactory = loggerFactory;
+			_logger = logger;
+
+			CreateServices();
+
+			_lifetime = new HordeClientLifetime(Task.FromResult(_hordeClient));
+		}
+
+		/// <inheritdoc/>
+		public async ValueTask DisposeAsync()
+		{
+			await _lifetime.DisposeAsync();
+			await DestroyServicesAsync();
+		}
+
+		/// <inheritdoc/>
+		public async Task<IHordeClientRef?> GetClientRefAsync()
+		{
+			HordeClientLifetime? lifetime = null;
+			try
+			{
+				lock (_lockObject)
+				{
+					lifetime = _lifetime;
+					lifetime.AddRef();
+				}
+
+				IHordeClient? client = await lifetime.ClientTask;
+				if (client == null)
+				{
+					lifetime.Release();
+					return null;
+				}
+
+				return new HordeClientRef(client, lifetime);
+			}
+			catch
+			{
+				lifetime?.Release();
+				throw;
+			}
+		}
+
+		/// <inheritdoc/>
+		public async Task RecreateAsync()
+		{
+			lock (_lockObject)
+			{
+				HordeClientLifetime prevLifetime = _lifetime;
+				Task<IHordeClient?> clientTask = Task.Run(async () =>
+				{
+					await prevLifetime.DisposeAsync();
+					await DestroyServicesAsync();
+					CreateServices();
+					return _hordeClient;
+				});
+				_lifetime = new HordeClientLifetime(clientTask);
+			}
+			await _lifetime.ClientTask;
+		}
+
+		void OnAccessTokenStateChangedInternal()
+		{
+			OnAccessTokenStateChanged?.Invoke();
+		}
+
+		void CreateServices()
+		{
+			Debug.Assert(_serviceProvider == null);
+			Debug.Assert(_hordeClient == null);
+
+			ServiceCollection serviceCollection = new ServiceCollection();
+			serviceCollection.AddHorde(options => options.AllowAuthPrompt = false);
+			serviceCollection.AddSingleton<ILoggerFactory>(_loggerFactory);
+			serviceCollection.AddSingleton(typeof(ILogger<>), typeof(Logger<>));
+			_serviceProvider = serviceCollection.BuildServiceProvider();
+
+			try
+			{
+				_hordeClient = _serviceProvider.GetRequiredService<IHordeClient>();
+				_hordeClient.OnAccessTokenStateChanged += OnAccessTokenStateChangedInternal;
+			}
+			catch (Exception ex)
+			{
+				_logger.LogWarning(ex, "Horde client could not be created: {Message}", ex.Message);
+			}
+
+			OnStateChanged?.Invoke();
+		}
+
+		async ValueTask DestroyServicesAsync()
+		{
+			if (_hordeClient != null)
+			{
+				_hordeClient.OnAccessTokenStateChanged -= OnAccessTokenStateChangedInternal;
+				_hordeClient = null;
+			}
+			if (_serviceProvider != null)
+			{
+				await _serviceProvider.DisposeAsync();
+				_serviceProvider = null;
+			}
+		}
+	}
+}
