Horde: Fix changes not being enumerated in descending order by PerforceService, causing the minimum changelist number for ICommitCollection.SubscribeAsync() not being updated. This prevents issue fixed tags from being processed.

diff --git a/Engine/Source/Programs/Horde/Plugins/Build/HordeServer.Build/BuildPlugin.cs b/Engine/Source/Programs/Horde/Plugins/Build/HordeServer.Build/BuildPlugin.cs
--- a/Engine/Source/Programs/Horde/Plugins/Build/HordeServer.Build/BuildPlugin.cs
+++ b/Engine/Source/Programs/Horde/Plugins/Build/HordeServer.Build/BuildPlugin.cs
@@ -119,7 +119,6 @@ namespace HordeServer
 			{
 				services.AddHostedService<AgentReportService>();
 				services.AddHostedService<IssueReportService>();
-				services.AddHostedService<IssueTagService>();
 				services.AddHostedService<JobExpirationService>();
 				services.AddHostedService(provider => provider.GetRequiredService<PerforceLoadBalancer>());
 				services.AddHostedService<PoolUpdateService>();
diff --git a/Engine/Source/Programs/Horde/Plugins/Build/HordeServer.Build/Issues/IssueTagService.cs b/Engine/Source/Programs/Horde/Plugins/Build/HordeServer.Build/Issues/IssueTagService.cs
--- a/Engine/Source/Programs/Horde/Plugins/Build/HordeServer.Build/Issues/IssueTagService.cs
+++ b/Engine/Source/Programs/Horde/Plugins/Build/HordeServer.Build/Issues/IssueTagService.cs
@@ -62,9 +62,17 @@ namespace HordeServer.Issues
 			using CancellationTokenSource cancellationSource = CancellationTokenSource.CreateLinkedTokenSource(cancellationToken);
 			using IDisposable? listener = _buildConfig.OnChange((_, _) => cancellationSource.Cancel());
 
+			BuildConfig buildConfig = _buildConfig.CurrentValue;
+
 			State initialState = await _state.GetAsync(cancellationSource.Token);
+			foreach (StreamId streamId in initialState.Streams.Keys)
+			{
+				if (!buildConfig.TryGetStream(streamId, out _))
+				{
+					await _state.UpdateAsync(x => x.Streams.Remove(streamId), cancellationToken);
+				}
+			}
 
-			BuildConfig buildConfig = _buildConfig.CurrentValue;
 			if (buildConfig.Streams.Count > 0)
 			{
 				List<Task> tasks = new List<Task>();
diff --git a/Engine/Source/Programs/Horde/Plugins/Build/HordeServer.Build/Perforce/PerforceService.cs b/Engine/Source/Programs/Horde/Plugins/Build/HordeServer.Build/Perforce/PerforceService.cs
--- a/Engine/Source/Programs/Horde/Plugins/Build/HordeServer.Build/Perforce/PerforceService.cs
+++ b/Engine/Source/Programs/Horde/Plugins/Build/HordeServer.Build/Perforce/PerforceService.cs
@@ -1237,6 +1237,8 @@ namespace HordeServer.Perforce
 						string filter = GetFilter($"//{perforce.Settings.ClientName}/...", minChange, maxChange);
 
 						List<ChangesRecord> changes = await perforce.GetChangesAsync(ChangesOptions.IncludeTimes | ChangesOptions.LongOutput, maxResults ?? -1, ChangeStatus.Submitted, filter, cancellationToken);
+						changes.SortBy(x => x.Number);
+
 						foreach (ChangesRecord change in changes)
 						{
 							ICommit commit = await PerforceService.CreateCommitAsync(StreamConfig, change, info, cancellationToken);
@@ -1262,6 +1264,8 @@ namespace HordeServer.Perforce
 							const int MaxFiles = 1000;
 
 							List<DescribeRecord> describeRecords = await perforce.DescribeAsync(DescribeOptions.None, MaxFiles, changesRecords.Select(x => x.Number).ToArray(), cancellationToken);
+							describeRecords.SortBy(x => x.Number);
+
 							foreach (DescribeRecord describeRecord in describeRecords)
 							{
 								bool hasAllFiles = describeRecord.Files.Count < MaxFiles;
