Horde: Add an exponential pause between iterations over the GC tick queue. This should reduce the time that we spend idle during blob deletion.

diff --git a/Engine/Source/Programs/Horde/Plugins/Storage/HordeServer.Storage/Storage/StorageService.cs b/Engine/Source/Programs/Horde/Plugins/Storage/HordeServer.Storage/Storage/StorageService.cs
--- a/Engine/Source/Programs/Horde/Plugins/Storage/HordeServer.Storage/Storage/StorageService.cs
+++ b/Engine/Source/Programs/Horde/Plugins/Storage/HordeServer.Storage/Storage/StorageService.cs
@@ -700,9 +700,13 @@ namespace HordeServer.Storage
 				{
 					break;
 				}
-				if (await ShouldPauseBlobTickAsync(current, cancellationToken))
+
+				// Wait until there's some space in the queue
+				int delaySecs = 1;
+				while (await ShouldPauseBlobTickAsync(current, cancellationToken))
 				{
-					break;
+					await Task.Delay(TimeSpan.FromSeconds(delaySecs), cancellationToken);
+					delaySecs = Math.Min(delaySecs * 2, 128);
 				}
 
 				// Add a check record for each new blob
