Horde: Updating schema documentation

#rnx

diff --git a/Engine/Source/Programs/Horde/Docs/Deployment/AgentSettings.md b/Engine/Source/Programs/Horde/Docs/Deployment/AgentSettings.md
--- a/Engine/Source/Programs/Horde/Docs/Deployment/AgentSettings.md
+++ b/Engine/Source/Programs/Horde/Docs/Deployment/AgentSettings.md
@@ -22,6 +22,7 @@ Name | Description
 `useLocalStorageClient` | `boolean`<br>Option to use a local storage client rather than connecting through the server. Primarily for convenience when debugging / iterating locally.
 `computeIp` | `string`<br>Incoming IP for listening for compute work. If not set, it will be automatically resolved.
 `computePort` | `integer`<br>Incoming port for listening for compute work. Needs to be tied with a lease. Set port to 0 to disable incoming compute requests.
+`openTelemetry` | [OpenTelemetrySettings](#opentelemetrysettings)<br>Options for OpenTelemetry
 `enableTelemetry` | `boolean`<br>Whether to send telemetry back to Horde server
 `telemetryReportInterval` | `integer`<br>How often to report telemetry events to server in milliseconds
 `bundleCacheSize` | `integer`<br>Maximum size of the bundle cache, in megabytes.
@@ -59,3 +60,27 @@ Name | Description
 ---- | -----------
 `mountPoint` | `string`<br>Where the share should be mounted on the local machine. Must be a drive letter for Windows.
 `remotePath` | `string`<br>Path to the remote resource
+
+## OpenTelemetrySettings
+
+OpenTelemetry configuration for collection and sending of traces and metrics.
+
+Name | Description
+---- | -----------
+`enabled` | `boolean`<br>Whether OpenTelemetry exporting is enabled
+`serviceName` | `string`<br>Service name
+`serviceNamespace` | `string`<br>Service namespace
+`serviceVersion` | `string`<br>Service version
+`enableDatadogCompatibility` | `boolean`<br>Whether to enrich and format telemetry to fit presentation in Datadog
+`attributes` | `string` `->` `string`<br>Extra attributes to set
+`enableConsoleExporter` | `boolean`<br>Whether to enable the console exporter (for debugging purposes)
+`protocolExporters` | `string` `->` [OpenTelemetryProtocolExporterSettings](#opentelemetryprotocolexportersettings)<br>Protocol exporters (key is a unique and arbitrary name)
+
+## OpenTelemetryProtocolExporterSettings
+
+Configuration for an OpenTelemetry exporter
+
+Name | Description
+---- | -----------
+`endpoint` | `string`<br>Endpoint URL. Usually differs depending on protocol used.
+`protocol` | `string`<br>Protocol for the exporter ('grpc' or 'httpprotobuf')
