Horde: Upgrade Horde Agent to NET 8.

diff --git a/Engine/Source/Programs/Horde/HordeAgent/HordeAgent.csproj b/Engine/Source/Programs/Horde/HordeAgent/HordeAgent.csproj
--- a/Engine/Source/Programs/Horde/HordeAgent/HordeAgent.csproj
+++ b/Engine/Source/Programs/Horde/HordeAgent/HordeAgent.csproj
@@ -2,7 +2,7 @@
 
   <PropertyGroup>
     <OutputType>Exe</OutputType>
-    <TargetFramework>net6.0</TargetFramework>
+    <TargetFramework>net8.0</TargetFramework>
     <Nullable>enable</Nullable>
     <ImplicitUsings>enable</ImplicitUsings>
 	<ApplicationIcon>HordeAgent.ico</ApplicationIcon>
diff --git a/Engine/Source/Programs/Horde/HordeAgent/Services/ComputeListenerService.cs b/Engine/Source/Programs/Horde/HordeAgent/Services/ComputeListenerService.cs
--- a/Engine/Source/Programs/Horde/HordeAgent/Services/ComputeListenerService.cs
+++ b/Engine/Source/Programs/Horde/HordeAgent/Services/ComputeListenerService.cs
@@ -133,7 +133,7 @@ namespace HordeAgent.Services
 		/// <inheritdoc/>
 		public async Task StopAsync(CancellationToken cancellationToken)
 		{
-			_cancellationSource.Cancel();
+			await _cancellationSource.CancelAsync();
 			await _serverTask.IgnoreCanceledExceptionsAsync();
 		}
 
diff --git a/Engine/Source/Programs/Horde/HordeAgent/Services/TelemetryService.cs b/Engine/Source/Programs/Horde/HordeAgent/Services/TelemetryService.cs
--- a/Engine/Source/Programs/Horde/HordeAgent/Services/TelemetryService.cs
+++ b/Engine/Source/Programs/Horde/HordeAgent/Services/TelemetryService.cs
@@ -284,7 +284,10 @@ class TelemetryService : IHostedService, IDisposable
 
 	public async Task StopAsync(CancellationToken cancellationToken)
 	{
-		_eventLoopHeartbeatCts?.Cancel();
+		if (_eventLoopHeartbeatCts != null)
+		{
+			await _eventLoopHeartbeatCts.CancelAsync();
+		}
 
 		if (_eventLoopTask != null)
 		{
