Horde: Add a dataflow graph with multiple deletion workers to improve GC performance.

diff --git a/Engine/Source/Programs/Horde/Plugins/Storage/HordeServer.Storage/Storage/StorageService.cs b/Engine/Source/Programs/Horde/Plugins/Storage/HordeServer.Storage/Storage/StorageService.cs
--- a/Engine/Source/Programs/Horde/Plugins/Storage/HordeServer.Storage/Storage/StorageService.cs
+++ b/Engine/Source/Programs/Horde/Plugins/Storage/HordeServer.Storage/Storage/StorageService.cs
@@ -4,6 +4,7 @@ using System.Buffers;
 using System.ComponentModel;
 using System.Diagnostics;
 using System.Linq.Expressions;
+using System.Threading.Tasks.Dataflow;
 using EpicGames.Core;
 using EpicGames.Horde.Storage;
 using EpicGames.Horde.Storage.Bundles;
@@ -1073,73 +1074,124 @@ namespace HordeServer.Storage
 			}
 		}
 
-		async Task TickGcForNamespaceAsync(NamespaceInfo namespaceInfo, ObjectId lastImportBlobInfoId, DateTime utcNow, bool deleteObjects, CancellationToken cancellationToken)
+		class GcSweepState
 		{
-			IStorageNamespace client = this.GetNamespace(namespaceInfo.Id);
+			long _score;
+			int _numItemsRemoved;
 
-			Stopwatch timer = Stopwatch.StartNew();
-			_logger.LogInformation("Running garbage collection for namespace {NamespaceId}...", namespaceInfo.Id);
-			int numItemsRemoved = 0;
+			public int NumRemovedItems
+				=> _numItemsRemoved;
 
-			double score = GetGcTimestamp(utcNow);
+			public GcSweepState(double score)
+				=> _score = BitConverter.DoubleToInt64Bits(score);
 
-			RedisSortedSetKey<RedisValue> checkSet = GetGcCheckSet(namespaceInfo.Id);
-			while (_storageConfig.CurrentValue.EnableGc || _storageConfig.CurrentValue.EnableGcVerification)
+			public double GetNextScore()
+				=> BitConverter.Int64BitsToDouble(Interlocked.Increment(ref _score));
+
+			public void OnRemovedItem()
+				=> Interlocked.Increment(ref _numItemsRemoved);
+		}
+
+		async Task TickGcForNamespaceAsync(NamespaceInfo namespaceInfo, ObjectId lastImportBlobInfoId, DateTime utcNow, bool deleteObjects, CancellationToken cancellationToken)
+		{
+			bool modified = false;
+			using IDisposable? notification = _storageConfig.OnChange((_, _) => modified = true);
+
+			StorageConfig storageConfig = _storageConfig.CurrentValue;
+			if (storageConfig.EnableGc || storageConfig.EnableGcVerification)
 			{
-				long length = await _redisService.GetDatabase().SortedSetLengthAsync(checkSet);
-				int batchSize = (int)Math.Min(length, 1024);
-				_logger.LogInformation("Garbage collection queue for namespace {NamespaceId} ({QueueName}) has {Length} entries; taking {Count}", namespaceInfo.Id, checkSet.Inner, length, batchSize);
+				IStorageNamespace client = this.GetNamespace(namespaceInfo.Id);
 
-				if (length == 0)
-				{
-					await _gcState.UpdateAsync(state => state.FindOrAddNamespace(namespaceInfo.Id).LastTime = utcNow, cancellationToken);
-					_logger.LogInformation("Finished garbage collection for namespace {NamespaceId} in {TimeSecs}s ({NumItems} removed)", namespaceInfo.Id, timer.Elapsed.TotalSeconds, numItemsRemoved);
-					break;
-				}
+				Stopwatch timer = Stopwatch.StartNew();
+				_logger.LogInformation("Running garbage collection for namespace {NamespaceId}...", namespaceInfo.Id);
+
+				GcSweepState sweepState = new GcSweepState(GetGcTimestamp(utcNow));
+				RedisSortedSetKey<RedisValue> checkSet = GetGcCheckSet(namespaceInfo.Id);
 
-				RedisValue[] values = await _redisService.GetDatabase().SortedSetRangeByRankAsync(checkSet, 0, batchSize);
-				foreach (RedisValue value in values)
+				// Create a dataflow pipeline for deleting blobs
+				BufferBlock<ObjectId> sourceBlock = new BufferBlock<ObjectId>(
+					new DataflowBlockOptions { BoundedCapacity = 1000, CancellationToken = cancellationToken });
+
+				ActionBlock<ObjectId> deleteBlock = new ActionBlock<ObjectId>(blobInfoId => RemoveBlobIfUnreferencedAsync(namespaceInfo, blobInfoId, sweepState, storageConfig, cancellationToken),
+					new ExecutionDataflowBlockOptions { MaxDegreeOfParallelism = 4, CancellationToken = cancellationToken });
+
+				sourceBlock.LinkTo(deleteBlock,
+					new DataflowLinkOptions { PropagateCompletion = true });
+
+				// Push items into the source block
+				double minScore = Double.NegativeInfinity;
+				while (!modified)
 				{
-					ObjectId blobInfoId = new ObjectId(((byte[]?)value)!);
+					long length = await _redisService.GetDatabase().SortedSetLengthAsync(checkSet);
+					int batchSize = (int)Math.Min(length, 1024);
+					_logger.LogInformation("Garbage collection queue for namespace {NamespaceId} ({QueueName}) has {Length} entries; taking {Count}", namespaceInfo.Id, checkSet.Inner, length, batchSize);
 
-					using TelemetrySpan span = _tracer.StartActiveSpan($"{nameof(StorageService)}.{nameof(TickGcForNamespaceAsync)}");
-					span.SetAttribute("BlobId", blobInfoId.ToString());
+					if (length == 0)
+					{
+						break;
+					}
 
-					if (blobInfoId <= lastImportBlobInfoId && !await IsBlobReferencedAsync(blobInfoId, cancellationToken))
+					SortedSetEntry<RedisValue>[] entries = await _redisService.GetDatabase().SortedSetRangeByScoreWithScoresAsync(checkSet, start: minScore, exclude: Exclude.Start, take: batchSize);
+					foreach (SortedSetEntry<RedisValue> entry in entries)
 					{
-						BlobInfo? info;
-						if (deleteObjects)
+						ObjectId blobInfoId = new ObjectId(((byte[]?)entry.ElementValue)!);
+						if (blobInfoId <= lastImportBlobInfoId)
 						{
-							info = await _blobCollection.FindOneAndDeleteAsync(x => x.Id == blobInfoId, cancellationToken: cancellationToken);
-						}
-						else
-						{
-							info = await _blobCollection.FindOneAndUpdateAsync(x => x.Id == blobInfoId, Builders<BlobInfo>.Update.Set(x => x.GcVersion, CurrentGcVersion), cancellationToken: cancellationToken);
+							await sourceBlock.SendAsync(blobInfoId, cancellationToken);
 						}
+						minScore = entry.Score;
+					}
+				}
 
-						if (info != null)
-						{
-							if (info.Imports != null)
-							{
-								SortedSetEntry<RedisValue>[] entries = info.Imports.Select(x => new SortedSetEntry<RedisValue>(x.ToByteArray(), score)).ToArray();
-								_ = _redisService.GetDatabase().SortedSetAddAsync(checkSet, entries, flags: CommandFlags.FireAndForget);
-								score = Math.BitIncrement(score);
-							}
+				// Mark the dataflow as complete and wait for it to flush
+				sourceBlock.Complete();
+				await deleteBlock.Completion;
 
-							ObjectKey objectKey = GetObjectKey(new BlobLocator(info.Path));
-							_logger.LogDebug("Deleting {NamespaceId} blob {BlobId}, key: {ObjectKey} ({ImportCount} imports)", namespaceInfo.Id, blobInfoId, objectKey, info.Imports?.Count ?? 0);
+				// Update the 
+				await _gcState.UpdateAsync(state => state.FindOrAddNamespace(namespaceInfo.Id).LastTime = utcNow, cancellationToken);
+				_logger.LogInformation("Finished garbage collection for namespace {NamespaceId} in {TimeSecs}s ({NumItems} removed)", namespaceInfo.Id, timer.Elapsed.TotalSeconds, sweepState.NumRemovedItems);
+			}
+		}
 
-							if (deleteObjects)
-							{
-								await namespaceInfo.Store.DeleteAsync(objectKey, cancellationToken);
-							}
+		async Task RemoveBlobIfUnreferencedAsync(NamespaceInfo namespaceInfo, ObjectId blobInfoId, GcSweepState state, StorageConfig storageConfig, CancellationToken cancellationToken)
+		{
+			using TelemetrySpan span = _tracer.StartActiveSpan($"{nameof(StorageService)}.{nameof(TickGcForNamespaceAsync)}");
+			span.SetAttribute("BlobId", blobInfoId.ToString());
 
-							numItemsRemoved++;
-						}
+			RedisSortedSetKey<RedisValue> checkSet = GetGcCheckSet(namespaceInfo.Id);
+			if (!await IsBlobReferencedAsync(blobInfoId, cancellationToken))
+			{
+				BlobInfo? info;
+				if (storageConfig.EnableGc)
+				{
+					info = await _blobCollection.FindOneAndDeleteAsync(x => x.Id == blobInfoId, cancellationToken: cancellationToken);
+				}
+				else
+				{
+					info = await _blobCollection.FindOneAndUpdateAsync(x => x.Id == blobInfoId, Builders<BlobInfo>.Update.Set(x => x.GcVersion, CurrentGcVersion), cancellationToken: cancellationToken);
+				}
+
+				if (info != null)
+				{
+					if (info.Imports != null)
+					{
+						SortedSetEntry<RedisValue>[] entries = info.Imports.Select(x => new SortedSetEntry<RedisValue>(x.ToByteArray(), state.GetNextScore())).ToArray();
+						_ = _redisService.GetDatabase().SortedSetAddAsync(checkSet, entries, flags: CommandFlags.FireAndForget);
 					}
-					_ = _redisService.GetDatabase().SortedSetRemoveAsync(checkSet, value, CommandFlags.FireAndForget);
+
+					ObjectKey objectKey = GetObjectKey(new BlobLocator(info.Path));
+					_logger.LogDebug("Deleting {NamespaceId} blob {BlobId}, key: {ObjectKey} ({ImportCount} imports)", namespaceInfo.Id, blobInfoId, objectKey, info.Imports?.Count ?? 0);
+
+					if (storageConfig.EnableGc)
+					{
+						await namespaceInfo.Store.DeleteAsync(objectKey, cancellationToken);
+					}
+
+					state.OnRemovedItem();
 				}
 			}
+
+			_ = _redisService.GetDatabase().SortedSetRemoveAsync(checkSet, blobInfoId.ToByteArray(), CommandFlags.FireAndForget);
 		}
 
 		static void SyncNamespaceList(GcState state, List<NamespaceConfig> namespaces)
