Horde: Check licenses for NuGet packages used by HordeTrayApp, and add build process to the public build script.

#rnx

diff --git a/Engine/Source/Programs/Horde/BuildHorde.xml b/Engine/Source/Programs/Horde/BuildHorde.xml
--- a/Engine/Source/Programs/Horde/BuildHorde.xml
+++ b/Engine/Source/Programs/Horde/BuildHorde.xml
@@ -101,6 +101,22 @@
 		</Node>
 	</Agent>
 	
+	<Agent Name="HordeTrayApp" Type="Win64">
+		<Node Name="Analyze HordeTrayApp">
+			<DotNet Arguments="build &quot;$(EngineDir)/Source/Programs/Horde/HordeTrayApp/HordeTrayApp.csproj&quot; -p:Configuration=Analyze --nologo" DotNetPath="$(DotNetPath)"/>
+		</Node>	
+		
+		<Node Name="Check HordeTrayApp Licenses" Requires="Analyze HordeTrayApp">
+			<Property Name="IgnorePackages">
+				[REDACTED_EMAIL]
+				[REDACTED_EMAIL]
+				[REDACTED_EMAIL]
+				[REDACTED_EMAIL]
+			</Property>
+			<NuGet-LicenseCheck BaseDir="Engine/Source/Programs/Horde/HordeTrayApp" IgnorePackages="$(IgnorePackages)" LicenseDir="Engine/Source/Programs/Horde/ThirdParty/Licenses" DotNetPath="$(DotNetPath)"/>
+		</Node>
+	</Agent>
+	
 	<!-- INSTALLER -->
 	
 	<Property Name="InstallerDir" Value="$(RootDir)/Engine/Source/Programs/Horde/Installer"/>
