Horde: Fix static analysis warning.

#rnx

diff --git a/Engine/Source/Programs/Horde/Plugins/Storage/HordeServer.Storage/Storage/StorageService.cs b/Engine/Source/Programs/Horde/Plugins/Storage/HordeServer.Storage/Storage/StorageService.cs
--- a/Engine/Source/Programs/Horde/Plugins/Storage/HordeServer.Storage/Storage/StorageService.cs
+++ b/Engine/Source/Programs/Horde/Plugins/Storage/HordeServer.Storage/Storage/StorageService.cs
@@ -1059,7 +1059,7 @@ namespace HordeServer.Storage
 							{
 								try
 								{
-									await TickGcForNamespaceAsync(state.Namespaces[namespaceId], gcState.LastImportBlobInfoId, utcNow, storageConfig.EnableGc, cancellationToken);
+									await TickGcForNamespaceAsync(state.Namespaces[namespaceId], gcState.LastImportBlobInfoId, utcNow, cancellationToken);
 								}
 								catch (Exception ex)
 								{
@@ -1092,7 +1092,7 @@ namespace HordeServer.Storage
 				=> Interlocked.Increment(ref _numItemsRemoved);
 		}
 
-		async Task TickGcForNamespaceAsync(NamespaceInfo namespaceInfo, ObjectId lastImportBlobInfoId, DateTime utcNow, bool deleteObjects, CancellationToken cancellationToken)
+		async Task TickGcForNamespaceAsync(NamespaceInfo namespaceInfo, ObjectId lastImportBlobInfoId, DateTime utcNow, CancellationToken cancellationToken)
 		{
 			bool modified = false;
 			using IDisposable? notification = _storageConfig.OnChange((_, _) => modified = true);
