Horde: Tweak trends view graph to ignore failures and better fit results
#skipci
#rnx

diff --git a/Engine/Source/Programs/Horde/HordeDashboard/src/components/jobDetailsV2/JobDetailStepTrendsV2.tsx b/Engine/Source/Programs/Horde/HordeDashboard/src/components/jobDetailsV2/JobDetailStepTrendsV2.tsx
--- a/Engine/Source/Programs/Horde/HordeDashboard/src/components/jobDetailsV2/JobDetailStepTrendsV2.tsx
+++ b/Engine/Source/Programs/Horde/HordeDashboard/src/components/jobDetailsV2/JobDetailStepTrendsV2.tsx
@@ -83,6 +83,10 @@ class StepTrendsDataView extends JobDataView {
             return false;
          }
 
+         if (r.outcome !== JobStepOutcome.Success && r.outcome !== JobStepOutcome.Warnings) {
+            return false;
+         }
+
          return true;
 
       });
@@ -146,7 +150,7 @@ class StepTrendsDataView extends JobDataView {
       this.median = median(alltimes);
       if (this.median) {
 
-         this.median *= 2;
+         this.median *= 2.5;
 
          if (this.median < 1) {
             this.median = 1;
