Horde: Add an option for automatically updating tools in the tray app.

diff --git a/Engine/Source/Programs/Horde/HordeTrayApp/IToolCatalog.cs b/Engine/Source/Programs/Horde/HordeTrayApp/IToolCatalog.cs
--- a/Engine/Source/Programs/Horde/HordeTrayApp/IToolCatalog.cs
+++ b/Engine/Source/Programs/Horde/HordeTrayApp/IToolCatalog.cs
@@ -7,6 +7,11 @@ namespace HordeTrayApp
 	/// </summary>
 	public interface IToolCatalog
 	{
+		/// <summary>
+		/// Whether to auto-update tools to latest
+		/// </summary>
+		bool AutoUpdate { get; set; }
+
 		/// <summary>
 		/// List of available tools
 		/// </summary>
diff --git a/Engine/Source/Programs/Horde/HordeTrayApp/ToolCatalog.cs b/Engine/Source/Programs/Horde/HordeTrayApp/ToolCatalog.cs
--- a/Engine/Source/Programs/Horde/HordeTrayApp/ToolCatalog.cs
+++ b/Engine/Source/Programs/Horde/HordeTrayApp/ToolCatalog.cs
@@ -141,6 +141,7 @@ namespace HordeTrayApp
 
 		class JsonState
 		{
+			public bool AutoUpdate { get; set; }
 			public List<JsonItemState> Items { get; set; } = new List<JsonItemState>();
 		}
 
@@ -207,6 +208,7 @@ namespace HordeTrayApp
 		IReadOnlyDictionary<ToolId, Item> _items;
 
 		Uri? _serverUri;
+		bool _autoUpdate = true;
 
 		public DirectoryReference BaseDir => _baseDir;
 
@@ -214,6 +216,19 @@ namespace HordeTrayApp
 
 		static readonly JsonSerializerOptions s_jsonSerializerOptions = new JsonSerializerOptions { AllowTrailingCommas = true, PropertyNameCaseInsensitive = true, PropertyNamingPolicy = JsonNamingPolicy.CamelCase, WriteIndented = true };
 
+		public bool AutoUpdate
+		{
+			get => _autoUpdate;
+			set
+			{
+				if (_autoUpdate != value)
+				{
+					_autoUpdate = value;
+					_ = SaveStateAsync(CancellationToken.None);
+				}
+			}
+		}
+
 		public ToolCatalog(IHordeClientProvider hordeClientProvider, ILogger<ToolCatalog> logger)
 		{
 			DirectoryReference? toolsDir = DirectoryReference.GetSpecialFolder(Environment.SpecialFolder.LocalApplicationData);
@@ -221,8 +236,8 @@ namespace HordeTrayApp
 
 			_hordeClientProvider = hordeClientProvider;
 			_logger = logger;
-			_baseDir = DirectoryReference.Combine(toolsDir, "Horde", "Launcher");
-			_configFile = FileReference.Combine(_baseDir, "state.json");
+			_baseDir = DirectoryReference.Combine(toolsDir, "Epic Games", "Horde", "TrayApp");
+			_configFile = FileReference.Combine(_baseDir, "Tools.json");
 			_updateEvent = new AsyncEvent();
 			_updateTask = new BackgroundTask(UpdateAsync);
 			_cleanupEvent = new AsyncEvent();
@@ -363,6 +378,7 @@ namespace HordeTrayApp
 				JsonState? state = JsonSerializer.Deserialize<JsonState>(data, s_jsonSerializerOptions);
 				if (state != null)
 				{
+					_autoUpdate = state.AutoUpdate;
 					foreach (JsonItemState jsonItem in state.Items)
 					{
 						CurrentToolDeploymentInfo? current = null;
@@ -388,6 +404,7 @@ namespace HordeTrayApp
 			lock (_workerThreadLockObject)
 			{
 				state = new JsonState();
+				state.AutoUpdate = _autoUpdate;
 				foreach (Item item in _items.Values)
 				{
 					if (item._state.Current != null)
@@ -398,6 +415,8 @@ namespace HordeTrayApp
 			}
 
 			byte[] data = JsonSerializer.SerializeToUtf8Bytes(state, s_jsonSerializerOptions);
+
+			DirectoryReference.CreateDirectory(_configFile.Directory);
 			await FileTransaction.WriteAllBytesAsync(_configFile, data, cancellationToken);
 		}
 
@@ -771,6 +790,17 @@ namespace HordeTrayApp
 					}
 				}
 				_items = newItems;
+
+				if (_autoUpdate)
+				{
+					foreach (Item item in _items.Values)
+					{
+						if (item._state.Current != null && item._state.Latest != null && item._state.Latest.Id != item._state.Current.Id)
+						{
+							UpdateDeployment(item, item._state.Latest);
+						}
+					}
+				}
 			}
 			PostMainThreadUpdate();
 		}
diff --git a/Engine/Source/Programs/Horde/HordeTrayApp/ToolCatalogViewModel.cs b/Engine/Source/Programs/Horde/HordeTrayApp/ToolCatalogViewModel.cs
--- a/Engine/Source/Programs/Horde/HordeTrayApp/ToolCatalogViewModel.cs
+++ b/Engine/Source/Programs/Horde/HordeTrayApp/ToolCatalogViewModel.cs
@@ -1,6 +1,7 @@
 // Copyright Epic Games, Inc. All Rights Reserved.
 
 using System.Collections.ObjectModel;
+using System.ComponentModel;
 using CommunityToolkit.Mvvm.ComponentModel;
 
 namespace HordeTrayApp
@@ -14,6 +15,9 @@ namespace HordeTrayApp
 
 		public bool ShowTools => Items.Count > 0;
 
+		[ObservableProperty]
+		bool _autoUpdate = true;
+
 		public ToolCatalogViewModel()
 			: this(null)
 		{ }
@@ -31,6 +35,16 @@ namespace HordeTrayApp
 			}
 		}
 
+		protected override void OnPropertyChanged(PropertyChangedEventArgs e)
+		{
+			base.OnPropertyChanged(e);
+
+			if (_catalog != null)
+			{
+				_catalog.AutoUpdate = AutoUpdate;
+			}
+		}
+
 		public void Dispose()
 		{
 			foreach (ToolCatalogItemViewModel item in _items.Values)
