Horde: Remove OpenTelemetry.Instrumentation.AspNetCore from agent

diff --git a/Engine/Source/Programs/Horde/HordeAgent/HordeAgent.csproj b/Engine/Source/Programs/Horde/HordeAgent/HordeAgent.csproj
--- a/Engine/Source/Programs/Horde/HordeAgent/HordeAgent.csproj
+++ b/Engine/Source/Programs/Horde/HordeAgent/HordeAgent.csproj
@@ -63,7 +63,6 @@
     <PackageReference Include="OpenTelemetry.Exporter.Console" Version="1.8.0" />
     <PackageReference Include="OpenTelemetry.Exporter.OpenTelemetryProtocol" Version="1.8.0" />
     <PackageReference Include="OpenTelemetry.Extensions.Hosting" Version="1.8.0" />
-    <PackageReference Include="OpenTelemetry.Instrumentation.AspNetCore" Version="1.8.1" />
     <PackageReference Include="OpenTelemetry.Instrumentation.GrpcNetClient" Version="1.8.0-beta.1" />
     <PackageReference Include="OpenTelemetry.Instrumentation.Http" Version="1.8.1" />
     <PackageReference Include="Serilog" Version="3.1.1" />
diff --git a/Engine/Source/Programs/Horde/HordeAgent/Utility/OpenTelemetry.cs b/Engine/Source/Programs/Horde/HordeAgent/Utility/OpenTelemetry.cs
--- a/Engine/Source/Programs/Horde/HordeAgent/Utility/OpenTelemetry.cs
+++ b/Engine/Source/Programs/Horde/HordeAgent/Utility/OpenTelemetry.cs
@@ -133,7 +133,6 @@ public static class OpenTelemetryHelper
 	private static void ConfigureMetrics(MeterProviderBuilder builder, OpenTelemetrySettings settings)
 	{
 		builder.AddMeter(OpenTelemetryMeters.MeterNames);
-		builder.AddAspNetCoreInstrumentation();
 		builder.AddHttpClientInstrumentation();
 
 		if (settings.EnableConsoleExporter)
