Horde: Ensure all streams register default artifact types.

We recently introduced the requirement that artifact types are explicitly configured (CL 35531333). The job driver expects certain artifact types to be present, and jobs will fail if these types are not registered.

This issue affects new installations and older installations upgrading to 5.5.

diff --git a/Engine/Source/Programs/Horde/Plugins/Build/HordeServer.Build/Streams/StreamConfig.cs b/Engine/Source/Programs/Horde/Plugins/Build/HordeServer.Build/Streams/StreamConfig.cs
--- a/Engine/Source/Programs/Horde/Plugins/Build/HordeServer.Build/Streams/StreamConfig.cs
+++ b/Engine/Source/Programs/Horde/Plugins/Build/HordeServer.Build/Streams/StreamConfig.cs
@@ -352,6 +352,12 @@ namespace HordeServer.Streams
 				_artifactTypeLookup[artifactTypeConfig.Type] = artifactTypeConfig;
 			}
 
+			// Add the default artifact types
+			EnsureDefaultArtifactType(ArtifactType.StepOutput);
+			EnsureDefaultArtifactType(ArtifactType.StepSaved);
+			EnsureDefaultArtifactType(ArtifactType.StepTrace);
+			EnsureDefaultArtifactType(ArtifactType.StepTestData);
+
 			// Compute a hash of this stream revision to make it easier to detect changes
 			byte[] streamData = JsonSerializer.SerializeToUtf8Bytes(this, JsonUtils.DefaultSerializerOptions);
 			Revision = IoHash.Compute(streamData).ToString();
@@ -371,6 +377,16 @@ namespace HordeServer.Streams
 		public bool TryGetArtifactType(ArtifactType type, [NotNullWhen(true)] out ArtifactTypeConfig? artifactTypeConfig)
 			=> _artifactTypeLookup.TryGetValue(type, out artifactTypeConfig);
 
+		void EnsureDefaultArtifactType(ArtifactType artifactType)
+		{
+			if (!_artifactTypeLookup.ContainsKey(artifactType))
+			{
+				ArtifactTypeConfig artifactTypeConfig = new ArtifactTypeConfig { Type = artifactType };
+				ArtifactTypes.Add(artifactTypeConfig);
+				_artifactTypeLookup.Add(artifactType, artifactTypeConfig);
+			}
+		}
+
 		/// <summary>
 		/// Enumerates all commit tags, including the default tags for code and content.
 		/// </summary>
