Horde: Document the location of the server config file.

diff --git a/Engine/Source/Programs/Horde/Docs/Deployment/Agent.md b/Engine/Source/Programs/Horde/Docs/Deployment/Agent.md
--- a/Engine/Source/Programs/Horde/Docs/Deployment/Agent.md
+++ b/Engine/Source/Programs/Horde/Docs/Deployment/Agent.md
@@ -78,7 +78,11 @@ The first time an agent connects to the server, it will generate a unique connec
 
 On Windows, connection tokens are stored in:
 
-    [REDACTED_USER_HOME]\AppData\Local\Horde.Agent\servers.json
+	[REDACTED_USER_HOME]\AppData\Local\Epic Games\Horde\Agent\servers.json
+    (if running under a normal user account)
+
+    C:\Windows\system32\config\systemprofile\AppData\Local\Epic Games\Horde\Agent\servers.json
+	(if running as a service)
 
 On Mac/Linux, connection tokens are stored in:
 
