Horde: Replace OpenTracing with OpenTelemetry in JobDriver

diff --git a/Engine/Source/Programs/Horde/Drivers/JobDriver.Tests/Execution/WorkspaceExecutorTest.cs b/Engine/Source/Programs/Horde/Drivers/JobDriver.Tests/Execution/WorkspaceExecutorTest.cs
--- a/Engine/Source/Programs/Horde/Drivers/JobDriver.Tests/Execution/WorkspaceExecutorTest.cs
+++ b/Engine/Source/Programs/Horde/Drivers/JobDriver.Tests/Execution/WorkspaceExecutorTest.cs
@@ -164,7 +164,7 @@ namespace JobDriver.Tests.Execution
 			_workspace.SetFile(1, "foo/bar/baz.h", "baz");
 
 			JobExecutorOptions executorOptions = new JobExecutorOptions(_hordeClient, _workingDir, null, _jobId, _batchId, default, new RpcJobOptions());
-			_executor = new(executorOptions, _workspace, null, NullLogger.Instance);
+			_executor = new(executorOptions, _workspace, null, SimpleTestExecutor.NoOpTracer, NullLogger.Instance);
 		}
 
 		public void Dispose()
@@ -188,7 +188,7 @@ namespace JobDriver.Tests.Execution
 		public async Task RegularAndAutoSdkWorkspaceAsync()
 		{
 			JobExecutorOptions executorOptions = new JobExecutorOptions(_hordeClient, _workingDir, Array.Empty<ProcessToTerminate>(), _jobId, _batchId, default, new RpcJobOptions());
-			using WorkspaceExecutor executor = new(executorOptions, _workspace, _autoSdkWorkspace, NullLogger.Instance);
+			using WorkspaceExecutor executor = new(executorOptions, _workspace, _autoSdkWorkspace, SimpleTestExecutor.NoOpTracer, NullLogger.Instance);
 
 			RpcBeginBatchResponse batch = new RpcBeginBatchResponse { Change = 1 };
 			await executor.InitializeAsync(batch, _logger, CancellationToken.None);
@@ -202,7 +202,7 @@ namespace JobDriver.Tests.Execution
 		public async Task EnvVarsAsync()
 		{
 			JobExecutorOptions executorOptions = new JobExecutorOptions(_hordeClient, _workingDir, Array.Empty<ProcessToTerminate>(), _jobId, _batchId, default, new RpcJobOptions());
-			using WorkspaceExecutor executor = new(executorOptions, _workspace, _autoSdkWorkspace, NullLogger.Instance);
+			using WorkspaceExecutor executor = new(executorOptions, _workspace, _autoSdkWorkspace, SimpleTestExecutor.NoOpTracer, NullLogger.Instance);
 
 			RpcBeginBatchResponse batch = new RpcBeginBatchResponse { Change = 1, StreamName = "//UE5/Main" };
 			await executor.InitializeAsync(batch, _logger, CancellationToken.None);
@@ -228,7 +228,7 @@ namespace JobDriver.Tests.Execution
 			_workspace.SetFile(1000, "New/Feature/Foo.cs", "foo");
 
 			JobExecutorOptions executorOptions = new JobExecutorOptions(_hordeClient, _workingDir, Array.Empty<ProcessToTerminate>(), preflightJobId, _batchId, default, new RpcJobOptions());
-			using WorkspaceExecutor executor = new(executorOptions, _workspace, null, NullLogger.Instance);
+			using WorkspaceExecutor executor = new(executorOptions, _workspace, null, SimpleTestExecutor.NoOpTracer, NullLogger.Instance);
 
 			RpcBeginBatchResponse batch = new RpcBeginBatchResponse { Change = 1, PreflightChange = 1000 };
 			await executor.InitializeAsync(batch, _logger, CancellationToken.None);
@@ -246,7 +246,7 @@ namespace JobDriver.Tests.Execution
 			_hordeClient.AddJob(noChangeJobId, _streamId, 0, 0);
 
 			JobExecutorOptions executorOptions = new JobExecutorOptions(_hordeClient, _workingDir, Array.Empty<ProcessToTerminate>(), noChangeJobId, _batchId, default, new RpcJobOptions());
-			using WorkspaceExecutor executor = new(executorOptions, _workspace, null, NullLogger.Instance);
+			using WorkspaceExecutor executor = new(executorOptions, _workspace, null, SimpleTestExecutor.NoOpTracer, NullLogger.Instance);
 
 			RpcBeginBatchResponse batch = new RpcBeginBatchResponse { };
 			await Assert.ThrowsExceptionAsync<WorkspaceMaterializationException>(() => executor.InitializeAsync(batch, _logger, CancellationToken.None));
diff --git a/Engine/Source/Programs/Horde/Drivers/JobDriver.Tests/SimpleTestExecutor.cs b/Engine/Source/Programs/Horde/Drivers/JobDriver.Tests/SimpleTestExecutor.cs
--- a/Engine/Source/Programs/Horde/Drivers/JobDriver.Tests/SimpleTestExecutor.cs
+++ b/Engine/Source/Programs/Horde/Drivers/JobDriver.Tests/SimpleTestExecutor.cs
@@ -9,6 +9,7 @@ using Horde.Common.Rpc;
 using HordeCommon.Rpc.Messages;
 using Microsoft.Extensions.Logging;
 using Microsoft.Extensions.Logging.Abstractions;
+using OpenTelemetry.Trace;
 
 namespace JobDriver.Tests
 {
@@ -17,9 +18,10 @@ namespace JobDriver.Tests
 		public const string Name = "Simple";
 
 		private readonly Func<JobStepInfo, ILogger, CancellationToken, Task<JobStepOutcome>> _func;
+		public static Tracer NoOpTracer { get;  }= TracerProvider.Default.GetTracer("NoOp");
 
 		public SimpleTestExecutor(Func<JobStepInfo, ILogger, CancellationToken, Task<JobStepOutcome>> func)
-			: base(new JobExecutorOptions(null!, null!, null, default, default, default, null!), NullLogger.Instance)
+			: base(new JobExecutorOptions(null!, null!, null, default, default, default, null!), NoOpTracer, NullLogger.Instance)
 		{
 			_func = func;
 		}
diff --git a/Engine/Source/Programs/Horde/Drivers/JobDriver/Commands/Execution/ExecuteConformCommand.cs b/Engine/Source/Programs/Horde/Drivers/JobDriver/Commands/Execution/ExecuteConformCommand.cs
--- a/Engine/Source/Programs/Horde/Drivers/JobDriver/Commands/Execution/ExecuteConformCommand.cs
+++ b/Engine/Source/Programs/Horde/Drivers/JobDriver/Commands/Execution/ExecuteConformCommand.cs
@@ -9,6 +9,7 @@ using JobDriver.Execution;
 using HordeCommon.Rpc.Tasks;
 using Microsoft.Extensions.Logging;
 using Microsoft.Extensions.Options;
+using OpenTelemetry.Trace;
 
 namespace JobDriver.Commands.Execution
 {
@@ -32,11 +33,13 @@ namespace JobDriver.Commands.Execution
 
 		readonly HordeClientFactory _hordeClientFactory;
 		readonly IOptions<DriverSettings> _driverSettings;
+		readonly Tracer _tracer;
 
-		public ExecuteConformCommand(HordeClientFactory hordeClientFactory, IOptions<DriverSettings> driverSettings)
+		public ExecuteConformCommand(HordeClientFactory hordeClientFactory, IOptions<DriverSettings> driverSettings, Tracer tracer)
 		{
 			_hordeClientFactory = hordeClientFactory;
 			_driverSettings = driverSettings;
+			_tracer = tracer;
 		}
 
 		public override async Task<int> ExecuteAsync(ILogger logger)
@@ -47,7 +50,7 @@ namespace JobDriver.Commands.Execution
 
 			await using HordeClient hordeClient = _hordeClientFactory.Create();
 
-			ConformExecutor conformExecutor = new ConformExecutor(hordeClient, WorkingDir, AgentId, LeaseId, conformTask, _driverSettings.Value, logger);
+			ConformExecutor conformExecutor = new ConformExecutor(hordeClient, WorkingDir, AgentId, LeaseId, conformTask, _driverSettings.Value, _tracer, logger);
 			await conformExecutor.ExecuteAsync(CancellationToken.None);
 
 			return 0;
diff --git a/Engine/Source/Programs/Horde/Drivers/JobDriver/DriverApp.cs b/Engine/Source/Programs/Horde/Drivers/JobDriver/DriverApp.cs
--- a/Engine/Source/Programs/Horde/Drivers/JobDriver/DriverApp.cs
+++ b/Engine/Source/Programs/Horde/Drivers/JobDriver/DriverApp.cs
@@ -3,9 +3,12 @@
 using System.Reflection;
 using EpicGames.Core;
 using EpicGames.Horde;
+using EpicGames.Horde.Utilities;
 using JobDriver.Execution;
+using JobDriver.Utility;
 using Microsoft.Extensions.Configuration;
 using Microsoft.Extensions.DependencyInjection;
+using OpenTelemetry.Trace;
 
 namespace JobDriver
 {
@@ -21,13 +24,19 @@ namespace JobDriver
 		{
 			CommandLineArguments arguments = new CommandLineArguments(args);
 
-			// Create the services 
+			// Create the services
 			IServiceCollection services = new ServiceCollection();
 			RegisterServices(services);
 
 			// Run the host
 			await using ServiceProvider serviceProvider = services.BuildServiceProvider();
-			return await CommandHost.RunAsync(arguments, serviceProvider, null);
+			int exitCode = await CommandHost.RunAsync(arguments, serviceProvider, null);
+			
+			// JobDriver execution is often short-lived so ensure any outstanding traces are flushed
+			TracerProvider? tracerProvider = serviceProvider.GetService<TracerProvider>();
+			tracerProvider?.ForceFlush(5000);
+			
+			return exitCode;
 		}
 
 		/// <summary>
@@ -46,6 +55,24 @@ namespace JobDriver
 			// Register the services
 			services.AddOptions<DriverSettings>().Configure(options => configuration.GetSection("Driver").Bind(options)).ValidateDataAnnotations();
 			services.AddLogging(builder => builder.AddEpicDefault());
+			
+			// Read any OpenTelemetry settings as an env var, so we can bootstrap it during dependency injection setup
+			// Passing it as an argument is too late in startup
+			OpenTelemetrySettings openTelemetrySettings = new();
+			string? otelSettingsJson = Environment.GetEnvironmentVariable("UE_HORDE_OTEL_SETTINGS");
+			if (otelSettingsJson != null)
+			{
+				try
+				{
+					openTelemetrySettings = OpenTelemetrySettingsExtensions.Deserialize(otelSettingsJson, true);
+				}
+				catch (Exception e)
+				{
+					Console.WriteLine("Unable to enable OpenTelemetry: " + e.Message);
+				}
+			}
+			OpenTelemetryHelper.Configure(services, openTelemetrySettings);
+
 			services.AddHorde(options => options.AllowAuthPrompt = false);
 
 			services.AddSingleton<IJobExecutorFactory, PerforceExecutorFactory>();
diff --git a/Engine/Source/Programs/Horde/Drivers/JobDriver/Execution/ConformExecutor.cs b/Engine/Source/Programs/Horde/Drivers/JobDriver/Execution/ConformExecutor.cs
--- a/Engine/Source/Programs/Horde/Drivers/JobDriver/Execution/ConformExecutor.cs
+++ b/Engine/Source/Programs/Horde/Drivers/JobDriver/Execution/ConformExecutor.cs
@@ -9,6 +9,7 @@ using Horde.Common.Rpc;
 using HordeCommon.Rpc.Messages;
 using HordeCommon.Rpc.Tasks;
 using Microsoft.Extensions.Logging;
+using OpenTelemetry.Trace;
 
 namespace JobDriver.Execution
 {
@@ -20,9 +21,10 @@ namespace JobDriver.Execution
 		readonly LeaseId _leaseId;
 		readonly ConformTask _conformTask;
 		readonly DriverSettings _driverSettings;
+		readonly Tracer _tracer;
 		readonly ILogger _logger;
 
-		public ConformExecutor(IHordeClient hordeClient, DirectoryReference workingDir, AgentId agentId, LeaseId leaseId, ConformTask conformTask, DriverSettings driverSettings, ILogger logger)
+		public ConformExecutor(IHordeClient hordeClient, DirectoryReference workingDir, AgentId agentId, LeaseId leaseId, ConformTask conformTask, DriverSettings driverSettings, Tracer tracer, ILogger logger)
 		{
 			_hordeClient = hordeClient;
 			_workingDir = workingDir;
@@ -30,6 +32,7 @@ namespace JobDriver.Execution
 			_leaseId = leaseId;
 			_conformTask = conformTask;
 			_driverSettings = driverSettings;
+			_tracer = tracer;
 			_logger = logger;
 		}
 
@@ -67,7 +70,7 @@ namespace JobDriver.Execution
 				// Run the conform task
 				if (isExecutorConformCompatible && _driverSettings.PerforceExecutor.RunConform)
 				{
-					await PerforceExecutor.ConformAsync(_workingDir, pendingWorkspaces, removeUntrackedFiles, _logger, cancellationToken);
+					await PerforceExecutor.ConformAsync(_workingDir, pendingWorkspaces, removeUntrackedFiles, _tracer, _logger, cancellationToken);
 				}
 				else
 				{
diff --git a/Engine/Source/Programs/Horde/Drivers/JobDriver/Execution/JobExecutor.cs b/Engine/Source/Programs/Horde/Drivers/JobDriver/Execution/JobExecutor.cs
--- a/Engine/Source/Programs/Horde/Drivers/JobDriver/Execution/JobExecutor.cs
+++ b/Engine/Source/Programs/Horde/Drivers/JobDriver/Execution/JobExecutor.cs
@@ -23,8 +23,8 @@ using HordeCommon.Rpc.Messages;
 using HordeCommon.Rpc.Tasks;
 using Microsoft.Extensions.Logging;
 using Microsoft.Extensions.Logging.Abstractions;
-using OpenTracing;
-using OpenTracing.Util;
+using OpenTelemetry.Trace;
+using StatusCode = Grpc.Core.StatusCode;
 
 namespace JobDriver.Execution
 {
@@ -206,6 +206,8 @@ namespace JobDriver.Execution
 		protected string? _scriptFileName;
 		protected bool _preprocessScript;
 		protected bool _savePreprocessedScript;
+		
+		protected Tracer Tracer { get; }
 
 		/// <summary>
 		/// Logger for the local agent process (as opposed to job logger)
@@ -244,7 +246,7 @@ namespace JobDriver.Execution
 		/// </summary>
 		internal TimeSpan _stepAbortPollRetryDelay = TimeSpan.FromSeconds(30);
 
-		protected JobExecutor(JobExecutorOptions options, ILogger logger)
+		protected JobExecutor(JobExecutorOptions options, Tracer tracer, ILogger logger)
 		{
 			HordeClient = options.HordeClient;
 			WorkingDir = options.WorkingDir;
@@ -258,6 +260,7 @@ namespace JobDriver.Execution
 			JobOptions = options.JobOptions;
 			JobRpc = null!; // Set in InitializeAsync()
 
+			Tracer = tracer;
 			Logger = logger;
 		}
 
@@ -274,6 +277,13 @@ namespace JobDriver.Execution
 
 		public async Task ExecuteAsync(ILogger logger, CancellationToken cancellationToken)
 		{
+			using TelemetrySpan span = Tracer.StartActiveSpan($"{nameof(JobExecutor)}.{nameof(ExecuteAsync)}");
+			span.SetAttribute("horde.job.id", JobId.ToString());
+			span.SetAttribute("horde.job.batch_id", BatchId.ToString());
+			span.SetAttribute("horde.job.executor", JobOptions.Executor);
+			span.SetAttribute("horde.lease.id", LeaseId.ToString());
+			JobOptions.DecorateSpan(span);
+			
 			JobRpc.JobRpcClient jobRpc = await HordeClient.CreateGrpcClientAsync<JobRpc.JobRpcClient>(cancellationToken);
 
 			// Create a storage client for this session
@@ -326,7 +336,6 @@ namespace JobDriver.Execution
 			logger.LogInformation("Initializing executor...");
 			using (logger.BeginIndentScope("  "))
 			{
-				using IScope scope = GlobalTracer.Instance.BuildSpan("Initialize").StartActive();
 				await InitializeAsync(batch, logger, cancellationToken);
 			}
 
@@ -388,11 +397,11 @@ namespace JobDriver.Execution
 					logger.LogInformation("Starting job {JobId}, batch {BatchId}, step {StepId} (Drive Space Left: {DriveSpaceRemaining} GB)", JobId, BatchId, step.StepId, availableFreeSpace.ToString("F1"));
 
 					// Create a trace span
-					using IScope scope = GlobalTracer.Instance.BuildSpan("Execute").WithTag("resource.name", step.Name).StartActive();
-					scope.Span.SetTag("stepId", step.StepId.ToString());
-					scope.Span.SetTag("logId", step.LogId.ToString());
-					//				using IDisposable TraceProperty = LogContext.PushProperty("dd.trace_id", CorrelationIdentifier.TraceId.ToString());
-					//				using IDisposable SpanProperty = LogContext.PushProperty("dd.span_id", CorrelationIdentifier.SpanId.ToString());
+					using TelemetrySpan span = Tracer.StartActiveSpan("Execute");
+					span.SetAttribute("resource.name", step.Name);
+					span.SetAttribute("horde.job.step.id", step.StepId.ToString());
+					span.SetAttribute("horde.job.step.name", step.Name);
+					span.SetAttribute("horde.job.step.log_id", step.LogId.ToString());
 
 					// Update the context to include information about this step
 					JobStepOutcome stepOutcome;
@@ -474,7 +483,6 @@ namespace JobDriver.Execution
 			logger.LogInformation("Finalizing...");
 			using (logger.BeginIndentScope("  "))
 			{
-				using IScope scope = GlobalTracer.Instance.BuildSpan("Finalize").StartActive();
 				await FinalizeAsync(logger, CancellationToken.None);
 			}
 		}
@@ -564,6 +572,8 @@ namespace JobDriver.Execution
 
 		public virtual async Task InitializeAsync(RpcBeginBatchResponse batch, ILogger logger, CancellationToken cancellationToken)
 		{
+			using TelemetrySpan span = Tracer.StartActiveSpan($"{nameof(JobExecutor)}.{nameof(InitializeAsync)}");
+			
 			Batch = batch;
 			JobRpc = await HordeClient.CreateGrpcClientAsync<JobRpc.JobRpcClient>(cancellationToken);
 
@@ -689,6 +699,7 @@ namespace JobDriver.Execution
 
 		public virtual Task FinalizeAsync(ILogger jobLogger, CancellationToken cancellationToken)
 		{
+			using TelemetrySpan span = Tracer.StartActiveSpan($"{nameof(JobExecutor)}.{nameof(FinalizeAsync)}");
 			return Task.CompletedTask;
 		}
 
@@ -806,8 +817,11 @@ namespace JobDriver.Execution
 				buildGraphFiles.Add(preprocessedSchemaFile);
 			}
 
-			using (GlobalTracer.Instance.BuildSpan("TempStorage").WithTag("resource", "Write").StartActive())
 			{
+				using TelemetrySpan span = Tracer.StartActiveSpan("TempStorage");
+				const string Prefix = "horde.temp_storage.";
+				span.SetAttribute(Prefix + "action", "upload");
+				
 				// Create the artifact
 				ArtifactName artifactName = TempStorage.GetArtifactNameForNode(SetupStepName);
 				ArtifactType artifactType = ArtifactType.StepOutput;
@@ -821,6 +835,7 @@ namespace JobDriver.Execution
 				RpcCreateJobArtifactResponseV2 artifact = await JobRpc.CreateArtifactV2Async(artifactRequest, cancellationToken: cancellationToken);
 				ArtifactId artifactId = ArtifactId.Parse(artifact.Id);
 				logger.LogInformation("Creating output artifact {ArtifactId} '{ArtifactName}' ({ArtifactType}) with ref {RefName} in namespace {NamespaceId}", artifactId, artifactName, artifactType, artifact.RefName, artifact.NamespaceId);
+				DecorateSpanWithArtifact(span, Prefix, artifactRequest, artifact);
 
 				// Write the data
 				IStorageNamespace storage = GetStorageNamespace(new NamespaceId(artifact.NamespaceId), artifact.Token);
@@ -1196,16 +1211,20 @@ namespace JobDriver.Execution
 
 			// Read the manifests for all the input storage blocks
 			Dictionary<TempStorageBlockRef, TempStorageBlockManifest> inputManifests = new Dictionary<TempStorageBlockRef, TempStorageBlockManifest>();
-			using (IScope scope = GlobalTracer.Instance.BuildSpan("TempStorage").WithTag("resource", "read").StartActive())
+
 			{
+				using TelemetrySpan span = Tracer.StartActiveSpan("TempStorage");
+				const string Prefix = "horde.temp_storage.";
+				span.SetAttribute(Prefix + "action", "download");
+				
 				Stopwatch timer = Stopwatch.StartNew();
-				scope.Span.SetTag("blocks", inputStorageBlocks.Count);
+				span.SetAttribute(Prefix + "blocks", inputStorageBlocks.Count);
 				foreach (TempStorageBlockRef inputStorageBlock in inputStorageBlocks)
 				{
 					TempStorageBlockManifest manifest = await TempStorage.RetrieveBlockAsync(HordeClient, JobId, step.StepId, inputStorageBlock.NodeName, inputStorageBlock.OutputName, workspaceDir, manifestDir, logger, cancellationToken);
 					inputManifests[inputStorageBlock] = manifest;
 				}
-				scope.Span.SetTag("size", inputManifests.Sum(x => x.Value.GetTotalSize()));
+				span.SetAttribute(Prefix + "size", inputManifests.Sum(x => x.Value.GetTotalSize()));
 				logger.LogInformation("Download took {Time:n1}s", timer.Elapsed.TotalSeconds);
 			}
 
@@ -1330,8 +1349,11 @@ namespace JobDriver.Execution
 			}
 
 			// Write all the storage blocks, and update the mapping from file to storage block
-			using (GlobalTracer.Instance.BuildSpan("TempStorage").WithTag("resource", "Write").StartActive())
 			{
+				using TelemetrySpan span = Tracer.StartActiveSpan("TempStorage");
+				const string Prefix = "horde.temp_storage.";
+				span.SetAttribute(Prefix + "action", "upload");
+				
 				// Create the artifact
 				RpcCreateJobArtifactRequestV2 artifactRequest = new RpcCreateJobArtifactRequestV2();
 				artifactRequest.JobId = JobId.ToString();
@@ -1342,6 +1364,7 @@ namespace JobDriver.Execution
 				RpcCreateJobArtifactResponseV2 artifact = await JobRpc.CreateArtifactV2Async(artifactRequest, cancellationToken: cancellationToken);
 				ArtifactId artifactId = ArtifactId.Parse(artifact.Id);
 				logger.LogInformation("Created artifact {ArtifactId} '{ArtifactName}' ({ArtifactType}) with ref {RefName} ({RefUrl})", artifactId, artifactRequest.Name, ArtifactType.StepOutput, artifact.RefName, $"{HordeClient.ServerUrl.ToString().TrimEnd('/')}/api/v1/storage/{artifact.NamespaceId}/refs/{artifact.RefName}");
+				DecorateSpanWithArtifact(span, Prefix, artifactRequest, artifact);
 
 				IStorageNamespace storage = GetStorageNamespace(new NamespaceId(artifact.NamespaceId), artifact.Token);
 
@@ -1407,8 +1430,13 @@ namespace JobDriver.Execution
 					continue;
 				}
 
-				using (GlobalTracer.Instance.BuildSpan("Artifact").WithTag("name", graphArtifact.Name).WithTag("resource", "Write").StartActive())
 				{
+					using TelemetrySpan span = Tracer.StartActiveSpan("Artifact");
+					const string Prefix = "horde.artifact.";
+					span.SetAttribute(Prefix + "action", "upload");
+					span.SetAttribute(Prefix + "name", graphArtifact.Name);
+					span.SetAttribute(Prefix + "type", graphArtifact.Type);
+					
 					// Create the artifact
 					RpcCreateJobArtifactRequestV2 artifactRequest = new RpcCreateJobArtifactRequestV2();
 					artifactRequest.JobId = JobId.ToString();
@@ -1422,6 +1450,7 @@ namespace JobDriver.Execution
 					RpcCreateJobArtifactResponseV2 artifact = await JobRpc.CreateArtifactV2Async(artifactRequest, cancellationToken: cancellationToken);
 					ArtifactId artifactId = ArtifactId.Parse(artifact.Id);
 					logger.LogInformation("Created artifact {ArtifactId} '{ArtifactName}' ({ArtifactType}) with ref {RefName} ({RefUrl})", artifactId, artifactRequest.Name, ArtifactType.StepOutput, artifact.RefName, $"{HordeClient.ServerUrl.ToString().TrimEnd('/')}/api/v1/storage/{artifact.NamespaceId}/refs/{artifact.RefName}");
+					DecorateSpanWithArtifact(span, Prefix, artifactRequest, artifact);
 
 					IStorageNamespace storage = GetStorageNamespace(new NamespaceId(artifact.NamespaceId), artifact.Token);
 
@@ -1458,7 +1487,9 @@ namespace JobDriver.Execution
 		protected async Task<int> ExecuteAutomationToolAsync(JobStepInfo step, DirectoryReference workspaceDir, string? arguments, bool? useP4, ILogger logger, CancellationToken cancellationToken)
 		{
 			int result;
-			using IScope scope = GlobalTracer.Instance.BuildSpan("BuildGraph").StartActive();
+			using TelemetrySpan span = Tracer.StartActiveSpan("BuildGraph");
+			span.SetAttribute("horde.job.step.id", step.StepId.ToString());
+			span.SetAttribute("horde.job.step.name", step.Name);
 
 			if (!_compileAutomationTool)
 			{
@@ -1469,10 +1500,12 @@ namespace JobDriver.Execution
 			{
 				arguments += " -NoP4";
 			}
+			
+			span.SetAttribute("horde.args", arguments);
 
 			if (RuntimeInformation.IsOSPlatform(OSPlatform.Windows))
 			{
-				result = await ExecuteCommandAsync(step, workspaceDir, Environment.GetEnvironmentVariable("COMSPEC") ?? "cmd.exe", $"/C \"\"{workspaceDir}\\Engine\\Build\\BatchFiles\\RunUAT.bat\" {arguments}\"", logger, cancellationToken);
+				result = await ExecuteCommandAsync(step, workspaceDir, Environment.GetEnvironmentVariable("COMSPEC") ?? "cmd.exe", $"/C \"\"{workspaceDir}\\Engine\\Build\\BatchFiles\\RunUAT.bat\" {arguments}\"", span, logger, cancellationToken);
 			}
 			else if (RuntimeInformation.IsOSPlatform(OSPlatform.Linux))
 			{
@@ -1483,11 +1516,11 @@ namespace JobDriver.Execution
 					args = $"\"{workspaceDir}/Engine/Build/BatchFiles/RunWineUAT.sh\" {arguments}";
 				}
 
-				result = await ExecuteCommandAsync(step, workspaceDir, "/bin/bash", args, logger, cancellationToken);
+				result = await ExecuteCommandAsync(step, workspaceDir, "/bin/bash", args, span, logger, cancellationToken);
 			}
 			else if (RuntimeInformation.IsOSPlatform(OSPlatform.OSX))
 			{
-				result = await ExecuteCommandAsync(step, workspaceDir, "/bin/sh", $"\"{workspaceDir}/Engine/Build/BatchFiles/RunUAT.sh\" {arguments}", logger, cancellationToken);
+				result = await ExecuteCommandAsync(step, workspaceDir, "/bin/sh", $"\"{workspaceDir}/Engine/Build/BatchFiles/RunUAT.sh\" {arguments}", span, logger, cancellationToken);
 			}
 			else
 			{
@@ -1752,7 +1785,7 @@ namespace JobDriver.Execution
 			return name;
 		}
 
-		async Task<int> ExecuteCommandAsync(JobStepInfo step, DirectoryReference workspaceDir, string fileName, string arguments, ILogger jobLogger, CancellationToken cancellationToken)
+		async Task<int> ExecuteCommandAsync(JobStepInfo step, DirectoryReference workspaceDir, string fileName, string arguments, TelemetrySpan buildGraphSpan, ILogger jobLogger, CancellationToken cancellationToken)
 		{
 			// Method for expanding environment variable properties related to this step
 			Dictionary<string, string> properties = new Dictionary<string, string>(StringComparer.OrdinalIgnoreCase);
@@ -1889,7 +1922,7 @@ namespace JobDriver.Execution
 				List<(string, FileReference)> telemetryFiles = new List<(string, FileReference)>();
 				try
 				{
-					await CreateTelemetryFilesAsync(step.Name, telemetryDir, telemetryFiles, jobLogger, cancellationToken);
+					await CreateTelemetryFilesAsync(step.Name, telemetryDir, telemetryFiles, buildGraphSpan, jobLogger, cancellationToken);
 				}
 				catch (Exception ex)
 				{
@@ -1998,7 +2031,7 @@ namespace JobDriver.Execution
 			return exitCode;
 		}
 
-		async Task CreateTelemetryFilesAsync(string stepName, DirectoryReference telemetryDir, List<(string, FileReference)> telemetryFiles, ILogger jobLogger, CancellationToken cancellationToken)
+		async Task CreateTelemetryFilesAsync(string stepName, DirectoryReference telemetryDir, List<(string, FileReference)> telemetryFiles, TelemetrySpan buildGraphSpan, ILogger jobLogger, CancellationToken cancellationToken)
 		{
 			List<TraceEventList> telemetryList = new List<TraceEventList>();
 			foreach (FileReference telemetryFile in DirectoryReference.EnumerateFiles(telemetryDir, "*.json"))
@@ -2113,7 +2146,7 @@ namespace JobDriver.Execution
 				}
 				telemetryFiles.Add(("Trace.json", traceFile));
 
-				CreateTracingData(GlobalTracer.Instance.ActiveSpan, rootSpan);
+				CreateTracingData(buildGraphSpan, rootSpan);
 			}
 		}
 
@@ -2154,20 +2187,17 @@ namespace JobDriver.Execution
 			await JobRpc.CreateReportAsync(request);
 		}
 
-		private static ISpan CreateTracingData(ISpan parent, TraceSpan span)
+		private TelemetrySpan CreateTracingData(TelemetrySpan parent, TraceSpan span)
 		{
-			ISpan newSpan = GlobalTracer.Instance.BuildSpan(span.Name)
-				.AsChildOf(parent)
-				.WithTag("service.name", span.Service)
-				.WithTag("resource.name", span.Resource)
-				.WithStartTimestamp(new DateTime(span.Start, DateTimeKind.Utc))
-				.Start();
+			TelemetrySpan newSpan = Tracer.StartSpan(span.Name ?? "TraceJsonUnknown", SpanKind.Internal, parent, startTime: new DateTime(span.Start, DateTimeKind.Utc));
+			newSpan.SetAttribute("service.name", span.Service);
+			newSpan.SetAttribute("resource.name", span.Resource);
 
 			if (span.Properties != null)
 			{
 				foreach (KeyValuePair<string, string> pair in span.Properties)
 				{
-					newSpan.SetTag(pair.Key, pair.Value);
+					newSpan.SetAttribute(pair.Key, pair.Value);
 				}
 			}
 			if (span.Children != null)
@@ -2178,9 +2208,21 @@ namespace JobDriver.Execution
 				}
 			}
 
-			newSpan.Finish(new DateTime(span.Finish, DateTimeKind.Utc));
+			newSpan.End(new DateTime(span.Finish, DateTimeKind.Utc));
 			return newSpan;
 		}
+		
+		private static void DecorateSpanWithArtifact(TelemetrySpan span, string prefix, RpcCreateJobArtifactRequestV2 request, RpcCreateJobArtifactResponseV2 response)
+		{
+			span.SetAttribute(prefix + "artifact.name", request.Name);
+			span.SetAttribute(prefix + "artifact.type", request.Type);
+			span.SetAttribute(prefix + "artifact.job_id", request.JobId);
+			span.SetAttribute(prefix + "artifact.step_id", request.StepId);
+			
+			span.SetAttribute(prefix + "artifact.id", response.Id);
+			span.SetAttribute(prefix + "artifact.ref_name", response.RefName);
+			span.SetAttribute(prefix + "artifact.namespace_id", response.NamespaceId);
+		}
 
 		public async Task UploadTestDataAsync(JobStepId jobStepId, IEnumerable<KeyValuePair<string, object>> pairs)
 		{
diff --git a/Engine/Source/Programs/Horde/Drivers/JobDriver/Execution/LocalExecutor.cs b/Engine/Source/Programs/Horde/Drivers/JobDriver/Execution/LocalExecutor.cs
--- a/Engine/Source/Programs/Horde/Drivers/JobDriver/Execution/LocalExecutor.cs
+++ b/Engine/Source/Programs/Horde/Drivers/JobDriver/Execution/LocalExecutor.cs
@@ -4,6 +4,7 @@ using EpicGames.Core;
 using HordeCommon.Rpc.Messages;
 using Microsoft.Extensions.Logging;
 using Microsoft.Extensions.Options;
+using OpenTelemetry.Trace;
 
 namespace JobDriver.Execution
 {
@@ -31,8 +32,8 @@ namespace JobDriver.Execution
 		private readonly LocalExecutorSettings _settings;
 		private readonly DirectoryReference _localWorkspaceDir;
 
-		public LocalExecutor(JobExecutorOptions options, LocalExecutorSettings settings, ILogger logger)
-			: base(options, logger)
+		public LocalExecutor(JobExecutorOptions options, LocalExecutorSettings settings, Tracer tracer, ILogger logger)
+			: base(options, tracer, logger)
 		{
 			_settings = settings;
 			if (settings.WorkspaceDir == null)
@@ -84,19 +85,21 @@ namespace JobDriver.Execution
 	class LocalExecutorFactory : IJobExecutorFactory
 	{
 		readonly LocalExecutorSettings _settings;
+		readonly Tracer _tracer;
 		readonly ILogger<LocalExecutor> _logger;
 
 		public string Name => "Local";
 
-		public LocalExecutorFactory(IOptions<LocalExecutorSettings> settings, ILogger<LocalExecutor> logger)
+		public LocalExecutorFactory(IOptions<LocalExecutorSettings> settings, Tracer tracer, ILogger<LocalExecutor> logger)
 		{
 			_settings = settings.Value;
+			_tracer = tracer;
 			_logger = logger;
 		}
 
 		public Task<JobExecutor> CreateExecutorAsync(RpcAgentWorkspace workspaceInfo, RpcAgentWorkspace? autoSdkWorkspaceInfo, JobExecutorOptions options, CancellationToken cancellationToken)
 		{
-			return Task.FromResult<JobExecutor>(new LocalExecutor(options, _settings, _logger));
+			return Task.FromResult<JobExecutor>(new LocalExecutor(options, _settings, _tracer, _logger));
 		}
 	}
 }
diff --git a/Engine/Source/Programs/Horde/Drivers/JobDriver/Execution/ManagedWorkspaceMaterializer.cs b/Engine/Source/Programs/Horde/Drivers/JobDriver/Execution/ManagedWorkspaceMaterializer.cs
--- a/Engine/Source/Programs/Horde/Drivers/JobDriver/Execution/ManagedWorkspaceMaterializer.cs
+++ b/Engine/Source/Programs/Horde/Drivers/JobDriver/Execution/ManagedWorkspaceMaterializer.cs
@@ -5,10 +5,9 @@ using EpicGames.Perforce.Managed;
 using JobDriver.Utility;
 using HordeCommon.Rpc.Messages;
 using Microsoft.Extensions.Logging;
-using OpenTracing;
-using OpenTracing.Util;
 using EpicGames.Perforce;
 using Microsoft.Extensions.DependencyInjection;
+using OpenTelemetry.Trace;
 
 namespace JobDriver.Execution;
 
@@ -26,6 +25,7 @@ public sealed class ManagedWorkspaceMaterializer : IWorkspaceMaterializer
 	private readonly bool _useCacheFile;
 	private readonly bool _cleanDuringFinalize;
 	private readonly WorkspaceInfo _workspace;
+	private readonly Tracer _tracer;
 	private readonly ILogger _logger;
 
 	/// <inheritdoc/>
@@ -47,18 +47,21 @@ public sealed class ManagedWorkspaceMaterializer : IWorkspaceMaterializer
 	/// <param name="useCacheFile">Whether to use a cache file during syncs</param>
 	/// <param name="cleanDuringFinalize">Whether to clean and revert files during finalize</param>
 	/// <param name="workspace"></param>
+	/// <param name="tracer"></param>
 	/// <param name="logger"></param>
 	private ManagedWorkspaceMaterializer(
 		RpcAgentWorkspace agentWorkspace,
 		bool useCacheFile,
 		bool cleanDuringFinalize,
 		WorkspaceInfo workspace,
+		Tracer tracer,
 		ILogger logger)
 	{
 		_agentWorkspace = agentWorkspace;
 		_useCacheFile = useCacheFile;
 		_cleanDuringFinalize = cleanDuringFinalize;
 		_workspace = workspace;
+		_tracer = tracer;
 		_logger = logger;
 
 		// Variables expected to be set for UAT/BuildGraph when Perforce is enabled (-P4 flag is set) 
@@ -79,19 +82,19 @@ public sealed class ManagedWorkspaceMaterializer : IWorkspaceMaterializer
 	}
 
 	/// <inheritdoc/>
-	public static async Task<ManagedWorkspaceMaterializer> CreateAsync(RpcAgentWorkspace agentWorkspace, DirectoryReference workingDir, bool useCacheFile, bool cleanDuringFinalize, ILogger logger, CancellationToken cancellationToken)
+	public static async Task<ManagedWorkspaceMaterializer> CreateAsync(RpcAgentWorkspace agentWorkspace, DirectoryReference workingDir, bool useCacheFile, bool cleanDuringFinalize, Tracer tracer, ILogger logger, CancellationToken cancellationToken)
 	{
 		ManagedWorkspaceOptions options = WorkspaceInfo.GetMwOptions(agentWorkspace);
 		WorkspaceInfo workspace = await WorkspaceInfo.CreateWorkspaceInfoAsync(agentWorkspace, workingDir, options, logger, cancellationToken);
-		return new ManagedWorkspaceMaterializer(agentWorkspace, useCacheFile, cleanDuringFinalize, workspace, logger);
+		return new ManagedWorkspaceMaterializer(agentWorkspace, useCacheFile, cleanDuringFinalize, workspace, tracer, logger);
 	}
 
 	/// <inheritdoc/>
 	public async Task SyncAsync(int changeNum, int preflightChangeNum, SyncOptions options, CancellationToken cancellationToken)
 	{
-		using IScope scope = CreateTraceSpan("ManagedWorkspaceMaterializer.SyncAsync");
-		scope.Span.SetTag("ChangeNum", changeNum);
-		scope.Span.SetTag("RemoveUntracked", options.RemoveUntracked);
+		using TelemetrySpan span = CreateTraceSpan($"{nameof(ManagedWorkspaceMaterializer)}.{nameof(SyncAsync)}");
+		span.SetAttribute("horde.change_num", changeNum);
+		span.SetAttribute("horde.remove_untracked", options.RemoveUntracked);
 
 		using IPerforceConnection perforce = await PerforceConnection.CreateAsync(_workspace.PerforceSettings, _logger);
 		await _workspace.SetupWorkspaceAsync(perforce, cancellationToken);
@@ -99,7 +102,7 @@ public sealed class ManagedWorkspaceMaterializer : IWorkspaceMaterializer
 		if (changeNum == IWorkspaceMaterializer.LatestChangeNumber)
 		{
 			int latestChangeNum = await _workspace.GetLatestChangeAsync(perforce, cancellationToken);
-			scope.Span.SetTag("LatestChangeNum", latestChangeNum);
+			span.SetAttribute("horde.change_num_latest", latestChangeNum);
 			changeNum = latestChangeNum;
 		}
 
@@ -107,7 +110,7 @@ public sealed class ManagedWorkspaceMaterializer : IWorkspaceMaterializer
 		if (_useCacheFile)
 		{
 			bool isSyncedDataDirty = await _workspace.UpdateLocalCacheMarkerAsync(cacheFile, changeNum, preflightChangeNum);
-			scope.Span.SetTag("IsSyncedDataDirty", isSyncedDataDirty);
+			span.SetAttribute("horde.is_dirty", isSyncedDataDirty);
 			if (!isSyncedDataDirty)
 			{
 				return;
@@ -124,7 +127,7 @@ public sealed class ManagedWorkspaceMaterializer : IWorkspaceMaterializer
 	/// <inheritdoc/>
 	public async Task FinalizeAsync(CancellationToken cancellationToken)
 	{
-		using IScope scope = CreateTraceSpan("ManagedWorkspaceMaterializer.FinalizeAsync");
+		using TelemetrySpan span = CreateTraceSpan($"{nameof(ManagedWorkspaceMaterializer)}.{nameof(FinalizeAsync)}");
 
 		if (_workspace != null && _cleanDuringFinalize)
 		{
@@ -142,18 +145,19 @@ public sealed class ManagedWorkspaceMaterializer : IWorkspaceMaterializer
 		return _workspace;
 	}
 
-	private IScope CreateTraceSpan(string operationName)
+	private TelemetrySpan CreateTraceSpan(string operationName)
 	{
-		IScope scope = GlobalTracer.Instance.BuildSpan(operationName).WithTag("resource.name", _agentWorkspace.Identifier).StartActive();
-		scope.Span.SetTag("UseHaveTable", WorkspaceInfo.ShouldUseHaveTable(_agentWorkspace.Method));
-		scope.Span.SetTag("Cluster", _agentWorkspace.Cluster);
-		scope.Span.SetTag("Incremental", _agentWorkspace.Incremental);
-		scope.Span.SetTag("Method", _agentWorkspace.Method);
-		scope.Span.SetTag("Stream", _agentWorkspace.Stream);
-		scope.Span.SetTag("Partitioned", _agentWorkspace.Partitioned);
-		scope.Span.SetTag("UseCacheFile", _useCacheFile);
-		scope.Span.SetTag("CleanDuringFinalize", _cleanDuringFinalize);
-		return scope;
+		TelemetrySpan span = _tracer.StartActiveSpan(operationName);
+		span.SetAttribute("resource.name", _agentWorkspace.Identifier);
+		span.SetAttribute("horde.use_have_table", WorkspaceInfo.ShouldUseHaveTable(_agentWorkspace.Method));
+		span.SetAttribute("horde.cluster", _agentWorkspace.Cluster);
+		span.SetAttribute("horde.incremental", _agentWorkspace.Incremental);
+		span.SetAttribute("horde.method", _agentWorkspace.Method);
+		span.SetAttribute("horde.stream", _agentWorkspace.Stream);
+		span.SetAttribute("horde.partitioned", _agentWorkspace.Partitioned);
+		span.SetAttribute("horde.use_cache_file", _useCacheFile);
+		span.SetAttribute("horde.clean_during_finalize", _cleanDuringFinalize);
+		return span;
 	}
 }
 
@@ -161,21 +165,23 @@ class ManagedWorkspaceMaterializerFactory : IWorkspaceMaterializerFactory
 {
 	readonly IServiceProvider _serviceProvider;
 
-	public ManagedWorkspaceMaterializerFactory(IServiceProvider serviceProvider)
-		=> _serviceProvider = serviceProvider;
+	public ManagedWorkspaceMaterializerFactory(IServiceProvider serviceProvider) => _serviceProvider = serviceProvider;
 
 	/// <inheritdoc/>
 	public async Task<IWorkspaceMaterializer?> CreateMaterializerAsync(string name, RpcAgentWorkspace workspaceInfo, DirectoryReference workspaceDir, bool forAutoSdk, CancellationToken cancellationToken)
 	{
 		if (name.Equals(ManagedWorkspaceMaterializer.Name, StringComparison.OrdinalIgnoreCase))
 		{
+			Tracer tracer = _serviceProvider.GetRequiredService<Tracer>();
+			ILogger<ManagedWorkspaceMaterializer> logger = _serviceProvider.GetRequiredService<ILogger<ManagedWorkspaceMaterializer>>();
+			
 			if (forAutoSdk)
 			{
-				return await ManagedWorkspaceMaterializer.CreateAsync(workspaceInfo, workspaceDir, true, false, _serviceProvider.GetRequiredService<ILogger<ManagedWorkspaceMaterializer>>(), cancellationToken);
+				return await ManagedWorkspaceMaterializer.CreateAsync(workspaceInfo, workspaceDir, true, false, tracer, logger, cancellationToken);
 			}
 			else
 			{
-				return await ManagedWorkspaceMaterializer.CreateAsync(workspaceInfo, workspaceDir, false, true, _serviceProvider.GetRequiredService<ILogger<ManagedWorkspaceMaterializer>>(), cancellationToken);
+				return await ManagedWorkspaceMaterializer.CreateAsync(workspaceInfo, workspaceDir, false, true, tracer, logger, cancellationToken);
 			}
 		}
 		return null;
diff --git a/Engine/Source/Programs/Horde/Drivers/JobDriver/Execution/PerforceExecutor.cs b/Engine/Source/Programs/Horde/Drivers/JobDriver/Execution/PerforceExecutor.cs
--- a/Engine/Source/Programs/Horde/Drivers/JobDriver/Execution/PerforceExecutor.cs
+++ b/Engine/Source/Programs/Horde/Drivers/JobDriver/Execution/PerforceExecutor.cs
@@ -9,8 +9,7 @@ using JobDriver.Utility;
 using Horde.Common.Rpc;
 using HordeCommon.Rpc.Messages;
 using Microsoft.Extensions.Logging;
-using OpenTracing;
-using OpenTracing.Util;
+using OpenTelemetry.Trace;
 
 namespace JobDriver.Execution
 {
@@ -39,24 +38,25 @@ namespace JobDriver.Execution
 		protected WorkspaceInfo? _autoSdkWorkspace;
 		protected WorkspaceInfo _workspace;
 
-		public PerforceExecutor(RpcAgentWorkspace workspaceInfo, RpcAgentWorkspace? autoSdkWorkspaceInfo, JobExecutorOptions options, ILogger logger)
-			: base(options, logger)
+		public PerforceExecutor(RpcAgentWorkspace workspaceInfo, RpcAgentWorkspace? autoSdkWorkspaceInfo, JobExecutorOptions options, Tracer tracer, ILogger logger)
+			: base(options, tracer, logger)
 		{
 			_workspaceInfo = workspaceInfo;
 			_autoSdkWorkspaceInfo = autoSdkWorkspaceInfo;
 			_rootDir = options.WorkingDir;
-
 			_workspace = null!;
 		}
 
 		public override async Task InitializeAsync(RpcBeginBatchResponse batch, ILogger logger, CancellationToken cancellationToken)
 		{
+			using TelemetrySpan span = Tracer.StartActiveSpan($"{nameof(PerforceExecutor)}.{nameof(InitializeAsync)}");
 			await base.InitializeAsync(batch, logger, cancellationToken);
 
 			// Setup and sync the AutoSDK workspace
 			if (_autoSdkWorkspaceInfo != null)
 			{
-				using IScope _ = GlobalTracer.Instance.BuildSpan("Workspace").WithTag("resource.name", "AutoSDK").StartActive();
+				using TelemetrySpan autoSdkSpan = Tracer.StartActiveSpan("Workspace");
+				autoSdkSpan.SetAttribute("resource.name", "AutoSDK");
 
 				ManagedWorkspaceOptions options = WorkspaceInfo.GetMwOptions(_autoSdkWorkspaceInfo);
 				_autoSdkWorkspace = await WorkspaceInfo.CreateWorkspaceInfoAsync(_autoSdkWorkspaceInfo, _rootDir, options, logger, cancellationToken);
@@ -78,6 +78,7 @@ namespace JobDriver.Execution
 				}
 
 				int autoSdkChangeNumber = await _autoSdkWorkspace.GetLatestChangeAsync(autoSdkPerforce, cancellationToken);
+				autoSdkSpan.SetAttribute("horde.change_number", autoSdkChangeNumber);
 
 				string syncText = $"Synced to CL {autoSdkChangeNumber}";
 				if (_autoSdkWorkspaceInfo.View.Count > 0)
@@ -103,8 +104,10 @@ namespace JobDriver.Execution
 				}
 			}
 
-			using (IScope scope = GlobalTracer.Instance.BuildSpan("Workspace").WithTag("resource.name", _workspaceInfo.Identifier).StartActive())
 			{
+				using TelemetrySpan autoSdkSpan = Tracer.StartActiveSpan("Workspace");
+				autoSdkSpan.SetAttribute("resource.name", _workspaceInfo.Identifier);
+				
 				// Sync the regular workspace
 				ManagedWorkspaceOptions options = WorkspaceInfo.GetMwOptions(_workspaceInfo);
 				_workspace = await WorkspaceInfo.CreateWorkspaceInfoAsync(_workspaceInfo, _rootDir, options, logger, cancellationToken);
@@ -243,11 +246,11 @@ namespace JobDriver.Execution
 			await base.FinalizeAsync(logger, cancellationToken);
 		}
 
-		public static async Task ConformAsync(DirectoryReference rootDir, IList<RpcAgentWorkspace> pendingWorkspaces, bool removeUntrackedFiles, ILogger logger, CancellationToken cancellationToken)
+		public static async Task ConformAsync(DirectoryReference rootDir, IList<RpcAgentWorkspace> pendingWorkspaces, bool removeUntrackedFiles, Tracer tracer, ILogger logger, CancellationToken cancellationToken)
 		{
-			using IScope scope = GlobalTracer.Instance.BuildSpan("Conform").StartActive();
-			scope.Span.SetTag("workspaces", String.Join(',', pendingWorkspaces.Select(x => x.Identifier)));
-			scope.Span.SetTag("removeUntrackedFiles", removeUntrackedFiles);
+			using TelemetrySpan span = tracer.StartActiveSpan("Conform");
+			span.SetAttribute("horde.perforce.workspaces", String.Join(',', pendingWorkspaces.Select(x => x.Identifier)));
+			span.SetAttribute("horde.perforce.remove_untracked", removeUntrackedFiles);
 
 			// Print out all the workspaces we're going to sync
 			logger.LogInformation("Workspaces:");
@@ -429,17 +432,19 @@ namespace JobDriver.Execution
 	class PerforceExecutorFactory : IJobExecutorFactory
 	{
 		readonly ILogger<PerforceExecutor> _logger;
+		readonly Tracer _tracer;
 
 		public string Name => PerforceExecutor.Name;
 
-		public PerforceExecutorFactory(ILogger<PerforceExecutor> logger)
+		public PerforceExecutorFactory(Tracer tracer, ILogger<PerforceExecutor> logger)
 		{
+			_tracer = tracer;
 			_logger = logger;
 		}
 
 		public Task<JobExecutor> CreateExecutorAsync(RpcAgentWorkspace workspaceInfo, RpcAgentWorkspace? autoSdkWorkspaceInfo, JobExecutorOptions options, CancellationToken cancellationToken)
 		{
-			return Task.FromResult<JobExecutor>(new PerforceExecutor(workspaceInfo, autoSdkWorkspaceInfo, options, _logger));
+			return Task.FromResult<JobExecutor>(new PerforceExecutor(workspaceInfo, autoSdkWorkspaceInfo, options, _tracer, _logger));
 		}
 	}
 }
diff --git a/Engine/Source/Programs/Horde/Drivers/JobDriver/Execution/TestExecutor.cs b/Engine/Source/Programs/Horde/Drivers/JobDriver/Execution/TestExecutor.cs
--- a/Engine/Source/Programs/Horde/Drivers/JobDriver/Execution/TestExecutor.cs
+++ b/Engine/Source/Programs/Horde/Drivers/JobDriver/Execution/TestExecutor.cs
@@ -8,6 +8,7 @@ using JobDriver.Parser;
 using Horde.Common.Rpc;
 using HordeCommon.Rpc.Messages;
 using Microsoft.Extensions.Logging;
+using OpenTelemetry.Trace;
 
 namespace JobDriver.Execution
 {
@@ -17,8 +18,8 @@ namespace JobDriver.Execution
 
 		readonly Dictionary<string, string> _arguments = new Dictionary<string, string>(StringComparer.OrdinalIgnoreCase);
 
-		public TestExecutor(JobExecutorOptions options, ILogger logger)
-			: base(options, logger)
+		public TestExecutor(JobExecutorOptions options, Tracer tracer, ILogger logger)
+			: base(options, tracer, logger)
 		{
 		}
 
@@ -223,18 +224,20 @@ namespace JobDriver.Execution
 
 	class TestExecutorFactory : IJobExecutorFactory
 	{
+		readonly Tracer _tracer;
 		readonly ILogger<TestExecutor> _logger;
 
 		public string Name => TestExecutor.Name;
 
-		public TestExecutorFactory(ILogger<TestExecutor> logger)
+		public TestExecutorFactory(Tracer tracer, ILogger<TestExecutor> logger)
 		{
+			_tracer = tracer;
 			_logger = logger;
 		}
 
 		public Task<JobExecutor> CreateExecutorAsync(RpcAgentWorkspace workspaceInfo, RpcAgentWorkspace? autoSdkWorkspaceInfo, JobExecutorOptions options, CancellationToken cancellationToken)
 		{
-			return Task.FromResult<JobExecutor>(new TestExecutor(options, _logger));
+			return Task.FromResult<JobExecutor>(new TestExecutor(options, _tracer, _logger));
 		}
 	}
 }
diff --git a/Engine/Source/Programs/Horde/Drivers/JobDriver/Execution/WorkspaceExecutor.cs b/Engine/Source/Programs/Horde/Drivers/JobDriver/Execution/WorkspaceExecutor.cs
--- a/Engine/Source/Programs/Horde/Drivers/JobDriver/Execution/WorkspaceExecutor.cs
+++ b/Engine/Source/Programs/Horde/Drivers/JobDriver/Execution/WorkspaceExecutor.cs
@@ -6,8 +6,7 @@ using JobDriver.Utility;
 using Horde.Common.Rpc;
 using HordeCommon.Rpc.Messages;
 using Microsoft.Extensions.Logging;
-using OpenTracing;
-using OpenTracing.Util;
+using OpenTelemetry.Trace;
 
 namespace JobDriver.Execution
 {
@@ -19,8 +18,8 @@ namespace JobDriver.Execution
 		private readonly IWorkspaceMaterializer _workspace;
 		private readonly IWorkspaceMaterializer? _autoSdkWorkspace;
 
-		public WorkspaceExecutor(JobExecutorOptions options, IWorkspaceMaterializer workspace, IWorkspaceMaterializer? autoSdkWorkspace, ILogger logger)
-			: base(options, logger)
+		public WorkspaceExecutor(JobExecutorOptions options, IWorkspaceMaterializer workspace, IWorkspaceMaterializer? autoSdkWorkspace, Tracer tracer, ILogger logger)
+			: base(options, tracer, logger)
 		{
 			_workspace = workspace;
 			_autoSdkWorkspace = autoSdkWorkspace;
@@ -39,6 +38,7 @@ namespace JobDriver.Execution
 
 		public override async Task InitializeAsync(RpcBeginBatchResponse batch, ILogger logger, CancellationToken cancellationToken)
 		{
+			using TelemetrySpan span = Tracer.StartActiveSpan($"{nameof(WorkspaceExecutor)}.{nameof(InitializeAsync)}");
 			await base.InitializeAsync(batch, logger, cancellationToken);
 
 			if (Batch.Change == 0)
@@ -49,7 +49,8 @@ namespace JobDriver.Execution
 			// Setup and sync the AutoSDK workspace
 			if (_autoSdkWorkspace != null)
 			{
-				using IScope _ = GlobalTracer.Instance.BuildSpan("Workspace").WithTag("resource.name", "AutoSDK").StartActive();
+				using TelemetrySpan autoSdkSpan = Tracer.StartActiveSpan("Workspace");
+				autoSdkSpan.SetAttribute("resource.name", "AutoSDK");
 				// TODO: Set type of workspace materializer as scope tag.
 
 				SyncOptions syncOptions = new();
@@ -57,9 +58,9 @@ namespace JobDriver.Execution
 			}
 
 			// Sync the regular workspace
-			using (IScope scope = GlobalTracer.Instance.BuildSpan("Workspace").StartActive())
 			{
-				scope.Span.SetTag(Datadog.Trace.OpenTracing.DatadogTags.ResourceName, _workspace.Identifier);
+				using TelemetrySpan autoSdkSpan = Tracer.StartActiveSpan("Workspace");
+				autoSdkSpan.SetAttribute("resource.name", _workspace.Identifier);
 
 				await _workspace.SyncAsync(Batch.Change, Batch.PreflightChange, new SyncOptions(), cancellationToken);
 
@@ -152,13 +153,15 @@ namespace JobDriver.Execution
 	class WorkspaceExecutorFactory : IJobExecutorFactory
 	{
 		private readonly IEnumerable<IWorkspaceMaterializerFactory> _materializerFactories;
+		private readonly Tracer _tracer;
 		private readonly ILoggerFactory _loggerFactory;
 
 		public string Name => WorkspaceExecutor.Name;
 
-		public WorkspaceExecutorFactory(IEnumerable<IWorkspaceMaterializerFactory> materializerFactories, ILoggerFactory loggerFactory)
+		public WorkspaceExecutorFactory(IEnumerable<IWorkspaceMaterializerFactory> materializerFactories, Tracer tracer, ILoggerFactory loggerFactory)
 		{
 			_materializerFactories = materializerFactories;
+			_tracer = tracer;
 			_loggerFactory = loggerFactory;
 		}
 
@@ -176,7 +179,7 @@ namespace JobDriver.Execution
 					autoSdkMaterializer = await CreateMaterializerAsync(name, autoSdkWorkspaceInfo, options.WorkingDir, forAutoSdk: true, cancellationToken);
 				}
 
-				return new WorkspaceExecutor(options, workspaceMaterializer, autoSdkMaterializer, _loggerFactory.CreateLogger<WorkspaceExecutor>());
+				return new WorkspaceExecutor(options, workspaceMaterializer, autoSdkMaterializer, _tracer, _loggerFactory.CreateLogger<WorkspaceExecutor>());
 			}
 			catch
 			{
diff --git a/Engine/Source/Programs/Horde/Drivers/JobDriver/JobDriver.csproj b/Engine/Source/Programs/Horde/Drivers/JobDriver/JobDriver.csproj
--- a/Engine/Source/Programs/Horde/Drivers/JobDriver/JobDriver.csproj
+++ b/Engine/Source/Programs/Horde/Drivers/JobDriver/JobDriver.csproj
@@ -55,7 +55,6 @@
     <PackageReference Include="AWSSDK.Core" Version="3.7.300.12" />
     <PackageReference Include="AWSSDK.EC2" Version="3.7.306.6" />
     <PackageReference Include="AWSSDK.SimpleSystemsManagement" Version="3.7.301.8" />
-    <PackageReference Include="Datadog.Trace.OpenTracing" Version="2.14.0" />
     <PackageReference Include="Google.Protobuf" Version="3.25.1" />
     <PackageReference Include="Grpc.Net.Client" Version="2.59.0" />
     <PackageReference Include="Grpc.Tools" Version="2.59.0">
@@ -79,7 +78,12 @@
       <PrivateAssets>all</PrivateAssets>
       <IncludeAssets>runtime; build; native; contentfiles; analyzers; buildtransitive</IncludeAssets>
     </PackageReference>
-    <PackageReference Include="OpenTracing" Version="0.12.1" />
+    <PackageReference Include="OpenTelemetry" Version="1.8.0" />
+    <PackageReference Include="OpenTelemetry.Exporter.Console" Version="1.8.0" />
+    <PackageReference Include="OpenTelemetry.Exporter.OpenTelemetryProtocol" Version="1.8.0" />
+    <PackageReference Include="OpenTelemetry.Extensions.Hosting" Version="1.8.0" />
+    <PackageReference Include="OpenTelemetry.Instrumentation.GrpcNetClient" Version="1.8.0-beta.1" />
+    <PackageReference Include="OpenTelemetry.Instrumentation.Http" Version="1.8.1" />
     <PackageReference Include="Serilog" Version="3.1.1" />
     <PackageReference Include="Serilog.Extensions.Hosting" Version="8.0.0" />
     <PackageReference Include="Serilog.Extensions.Logging" Version="8.0.0" />
diff --git a/Engine/Source/Programs/Horde/Drivers/JobDriver/Utilities/HordeRpc.cs b/Engine/Source/Programs/Horde/Drivers/JobDriver/Utilities/HordeRpc.cs
--- a/Engine/Source/Programs/Horde/Drivers/JobDriver/Utilities/HordeRpc.cs
+++ b/Engine/Source/Programs/Horde/Drivers/JobDriver/Utilities/HordeRpc.cs
@@ -2,6 +2,7 @@
 
 using EpicGames.Horde.Agents.Leases;
 using EpicGames.Horde.Jobs;
+using OpenTelemetry.Trace;
 
 // These partial classes/extensions operate on generated gRPC and Protobuf code.
 // Warnings below are disabled to avoid documenting every class touched.
@@ -80,5 +81,37 @@ namespace Horde.Common.Rpc
 	}
 }
 
+namespace HordeCommon.Rpc.Tasks
+{
+	partial class RpcJobOptions
+	{
+		public void DecorateSpan(TelemetrySpan span)
+		{
+			const string Prefix = "horde.job.options.";
+			span.SetAttribute(Prefix + "executor", Executor);
+			span.SetAttribute(Prefix + "materializer", WorkspaceMaterializer);
+			span.SetAttribute(Prefix + "driver", Driver);
+			
+			if (UseWine != null)
+			{
+				span.SetAttribute(Prefix + "use_wine", UseWine.Value);	
+			}
+			
+			if (RunInSeparateProcess != null)
+			{
+				span.SetAttribute(Prefix + "separate_process", RunInSeparateProcess.Value);
+			}
+			
+			if (Container.Enabled != null)
+			{
+				span.SetAttribute(Prefix + "container.enabled", Container.Enabled.Value);
+				span.SetAttribute(Prefix + "container.image_url", Container.ImageUrl);
+				span.SetAttribute(Prefix + "container.engine", Container.ContainerEngineExecutable);
+				span.SetAttribute(Prefix + "container.extra_args", Container.ExtraArguments);
+			}
+		}
+	}
+}
+
 #pragma warning restore CA1716
 #pragma warning restore CS1591
\ No newline at end of file
diff --git a/Engine/Source/Programs/Horde/Drivers/JobDriver/Utility/OpenTelemetry.cs b/Engine/Source/Programs/Horde/Drivers/JobDriver/Utility/OpenTelemetry.cs
new file mode 100644
--- /dev/null
+++ b/Engine/Source/Programs/Horde/Drivers/JobDriver/Utility/OpenTelemetry.cs
@@ -0,0 +1,178 @@
+// Copyright Epic Games, Inc. All Rights Reserved.
+
+using System.Diagnostics;
+using System.Diagnostics.Metrics;
+using EpicGames.Horde.Utilities;
+using Microsoft.Extensions.DependencyInjection;
+using Microsoft.Extensions.Logging;
+using OpenTelemetry.Exporter;
+using OpenTelemetry.Logs;
+using OpenTelemetry.Metrics;
+using OpenTelemetry.Resources;
+using OpenTelemetry.Trace;
+using Serilog.Core;
+
+namespace JobDriver.Utility;
+
+/// <summary>
+/// Serilog event enricher attaching trace and span ID for Datadog using current System.Diagnostics.Activity
+/// </summary>
+public class OpenTelemetryDatadogLogEnricher : ILogEventEnricher
+{
+	/// <inheritdoc />
+	public void Enrich(Serilog.Events.LogEvent logEvent, ILogEventPropertyFactory propertyFactory)
+	{
+		if (Activity.Current != null)
+		{
+			string stringTraceId = Activity.Current.TraceId.ToString();
+			string stringSpanId = Activity.Current.SpanId.ToString();
+			string ddTraceId = Convert.ToUInt64(stringTraceId.Substring(16), 16).ToString();
+			string ddSpanId = Convert.ToUInt64(stringSpanId, 16).ToString();
+
+			logEvent.AddPropertyIfAbsent(propertyFactory.CreateProperty("dd.trace_id", ddTraceId));
+			logEvent.AddPropertyIfAbsent(propertyFactory.CreateProperty("dd.span_id", ddSpanId));
+		}
+	}
+}
+
+/// <summary>
+/// Helper for configuring OpenTelemetry
+/// </summary>
+public static class OpenTelemetryHelper
+{
+	private static ResourceBuilder? s_resourceBuilder;
+
+	private static readonly Tracer s_noOpTracer = TracerProvider.Default.GetTracer("NoOp");
+	private static readonly Meter s_noOpMeter = new("NoOp");
+
+	/// <summary>
+	/// Configure OpenTelemetry in Horde and ASP.NET
+	/// </summary>
+	/// <param name="services">Service collection for DI</param>
+	/// <param name="settings">Current server settings</param>
+	public static void Configure(IServiceCollection services, OpenTelemetrySettings settings)
+	{
+		if (settings.Enabled)
+		{
+			services.AddOpenTelemetry()
+				.WithTracing(builder => ConfigureTracing(builder, settings))
+				.WithMetrics(builder => ConfigureMetrics(builder, settings));
+			
+			services.AddSingleton<Tracer>(sp => sp.GetRequiredService<TracerProvider>().GetTracer(settings.ServiceName));
+			services.AddSingleton<Meter>(sp => sp.GetRequiredService<Meter>());
+		}
+		else
+		{
+			// Configure a no-op tracer and meters when OpenTelemetry is disabled as they are still used in the codebase
+			services.AddSingleton(s_noOpTracer);
+			services.AddSingleton(s_noOpMeter);
+		}
+	}
+
+	private static void ConfigureTracing(TracerProviderBuilder builder, OpenTelemetrySettings settings)
+	{
+		void DatadogHttpRequestEnricher(Activity activity, HttpRequestMessage message)
+		{
+			activity.SetTag("service.name", settings.ServiceName + "-http-client");
+			activity.SetTag("operation.name", "http.request");
+			string url = $"{message.Method} {message.Headers.Host}{message.RequestUri?.LocalPath}";
+			activity.DisplayName = url;
+			activity.SetTag("resource.name", url);
+		}
+
+		builder
+			.AddSource(settings.ServiceName)
+			.AddHttpClientInstrumentation(options =>
+			{
+				if (settings.EnableDatadogCompatibility)
+				{
+					options.EnrichWithHttpRequestMessage = DatadogHttpRequestEnricher;
+				}
+			})
+			.AddGrpcClientInstrumentation(options =>
+			{
+				if (settings.EnableDatadogCompatibility)
+				{
+					options.EnrichWithHttpRequestMessage = DatadogHttpRequestEnricher;
+				}
+			})
+			.SetResourceBuilder(GetResourceBuilder(settings));
+
+		if (settings.EnableConsoleExporter)
+		{
+			builder.AddConsoleExporter();
+		}
+
+		foreach ((string name, OpenTelemetryProtocolExporterSettings exporterSettings) in settings.ProtocolExporters)
+		{
+			builder.AddOtlpExporter(name, exporter =>
+			{
+				exporter.Endpoint = exporterSettings.Endpoint!;
+				exporter.Protocol = Enum.TryParse(exporterSettings.Protocol, true, out OtlpExportProtocol protocol) ? protocol : OtlpExportProtocol.Grpc;
+			});
+		}
+	}
+
+	private static void ConfigureMetrics(MeterProviderBuilder builder, OpenTelemetrySettings settings)
+	{
+		builder.AddMeter(settings.ServiceName);
+		builder.AddHttpClientInstrumentation();
+
+		if (settings.EnableConsoleExporter)
+		{
+			builder.AddConsoleExporter();
+		}
+
+		foreach ((string name, OpenTelemetryProtocolExporterSettings exporterSettings) in settings.ProtocolExporters)
+		{
+			builder.AddOtlpExporter(name, exporter =>
+			{
+				exporter.Endpoint = exporterSettings.Endpoint!;
+				exporter.Protocol = Enum.TryParse(exporterSettings.Protocol, true, out OtlpExportProtocol protocol) ? protocol : OtlpExportProtocol.Grpc;
+			});
+		}
+	}
+
+	/// <summary>
+	/// Configure .NET logging with OpenTelemetry
+	/// </summary>
+	/// <param name="builder">Logging builder to modify</param>
+	/// <param name="settings">Current server settings</param>
+	public static void ConfigureLogging(ILoggingBuilder builder, OpenTelemetrySettings settings)
+	{
+		if (!settings.Enabled)
+		{
+			return;
+		}
+
+		builder.AddOpenTelemetry(options =>
+		{
+			options.IncludeScopes = true;
+			options.IncludeFormattedMessage = true;
+			options.ParseStateValues = true;
+			options.SetResourceBuilder(GetResourceBuilder(settings));
+
+			if (settings.EnableConsoleExporter)
+			{
+				options.AddConsoleExporter();
+			}
+		});
+	}
+
+	private static ResourceBuilder GetResourceBuilder(OpenTelemetrySettings settings)
+	{
+		if (s_resourceBuilder != null)
+		{
+			return s_resourceBuilder;
+		}
+
+		List<KeyValuePair<string, object>> attributes = settings.Attributes.Select(x => new KeyValuePair<string, object>(x.Key, x.Value)).ToList();
+		s_resourceBuilder = ResourceBuilder.CreateDefault()
+			.AddService(settings.ServiceName, serviceNamespace: settings.ServiceNamespace, serviceVersion: settings.ServiceVersion)
+			.AddAttributes(attributes)
+			.AddTelemetrySdk()
+			.AddEnvironmentVariableDetector();
+
+		return s_resourceBuilder;
+	}
+}
\ No newline at end of file
diff --git a/Engine/Source/Programs/Shared/EpicGames.Horde/Utilities/OpenTelemetry.cs b/Engine/Source/Programs/Shared/EpicGames.Horde/Utilities/OpenTelemetry.cs
--- a/Engine/Source/Programs/Shared/EpicGames.Horde/Utilities/OpenTelemetry.cs
+++ b/Engine/Source/Programs/Shared/EpicGames.Horde/Utilities/OpenTelemetry.cs
@@ -2,6 +2,8 @@
 
 using System;
 using System.Collections.Generic;
+using System.Text;
+using System.Text.Json;
 
 namespace EpicGames.Horde.Utilities;
 
@@ -67,4 +69,40 @@ public class OpenTelemetryProtocolExporterSettings
 	/// Protocol for the exporter ('grpc' or 'httpprotobuf')
 	/// </summary>
 	public string Protocol { get; set; } = "grpc";
+}
+
+/// <summary>
+/// Provides extension methods for serializing and deserializing OpenTelemetrySettings
+/// </summary>
+public static class OpenTelemetrySettingsExtensions
+{
+	/// <summary>
+	/// Serializes OpenTelemetrySettings to a JSON string, with an option to encode as base64
+	/// </summary>
+	/// <param name="settings">OpenTelemetrySettings to serialize.</param>
+	/// <param name="asBase64">If true, the resulting JSON string is encoded as base64</param>
+	public static string Serialize(OpenTelemetrySettings settings, bool asBase64 = false)
+	{
+		string data = JsonSerializer.Serialize(settings);
+		return asBase64 ? Convert.ToBase64String(Encoding.UTF8.GetBytes(data)) : data;
+	}
+	
+	/// <summary>
+	/// Deserializes a JSON string of OpenTelemetrySettings
+	/// </summary>
+	/// <param name="data">The string to deserialize, which can be either a JSON string or a Base64 encoded JSON string.</param>
+	/// <param name="asBase64">If true, the input string is treated as base64 encoded</param>
+	/// <returns>The deserialized OpenTelemetrySettings</returns>
+	/// <exception cref="JsonException">Thrown when deserialization fails or results in a null object.</exception>
+	public static OpenTelemetrySettings Deserialize(string data, bool asBase64 = false)
+	{
+		byte[] temp = asBase64 ? Convert.FromBase64String(data) : Encoding.UTF8.GetBytes(data);
+		OpenTelemetrySettings? settings = JsonSerializer.Deserialize<OpenTelemetrySettings>(temp);
+		if (settings == null)
+		{
+			throw new JsonException($"Unable to deserialize {nameof(OpenTelemetrySettings)} from JSON");
+		}
+		
+		return settings;
+	}
 }
\ No newline at end of file
