Horde: Fix static analysis warnings in build plugin and UGS

#rnx

diff --git a/Engine/Source/Programs/Horde/Plugins/Build/HordeServer.Build.Tests/Issues/IssueServiceTests.cs b/Engine/Source/Programs/Horde/Plugins/Build/HordeServer.Build.Tests/Issues/IssueServiceTests.cs
--- a/Engine/Source/Programs/Horde/Plugins/Build/HordeServer.Build.Tests/Issues/IssueServiceTests.cs
+++ b/Engine/Source/Programs/Horde/Plugins/Build/HordeServer.Build.Tests/Issues/IssueServiceTests.cs
@@ -1991,7 +1991,7 @@ namespace HordeServer.Tests.Issues
 			Assert.AreEqual("Errors in Update Version Files", issue.Summary);
 		}
 
-		static private IEnumerable<JsonLogEvent> MultilineLogEvent(LogLevel level, EventId eventId, string format, Dictionary<string, object?> properties)
+		private static IEnumerable<JsonLogEvent> MultilineLogEvent(LogLevel level, EventId eventId, string format, Dictionary<string, object?> properties)
 		{
 			DateTime time = new DateTime(2024, 1, 1, 0, 0, 0, 0, DateTimeKind.Utc);
 			string message = MessageTemplate.Render(format, properties!);
diff --git a/Engine/Source/Programs/Horde/Plugins/Build/HordeServer.Build/Jobs/Templates/TemplateConfig.cs b/Engine/Source/Programs/Horde/Plugins/Build/HordeServer.Build/Jobs/Templates/TemplateConfig.cs
--- a/Engine/Source/Programs/Horde/Plugins/Build/HordeServer.Build/Jobs/Templates/TemplateConfig.cs
+++ b/Engine/Source/Programs/Horde/Plugins/Build/HordeServer.Build/Jobs/Templates/TemplateConfig.cs
@@ -118,7 +118,7 @@ namespace HordeServer.Jobs.Templates
 		/// <param name="name"></param>
 		/// <param name="parameterIds"></param>
 		/// <returns></returns>
-		static protected ParameterId AddParameterId(ParameterId explicitParameterId, string name, HashSet<ParameterId> parameterIds)
+		protected static ParameterId AddParameterId(ParameterId explicitParameterId, string name, HashSet<ParameterId> parameterIds)
 		{
 			if (explicitParameterId.Id.IsEmpty)
 			{
@@ -139,7 +139,7 @@ namespace HordeServer.Jobs.Templates
 		/// </summary>
 		/// <param name="name">Default name for the parameter</param>
 		/// <param name="parameterIds">The existing parameter ids</param>
-		static protected ParameterId CreateUniqueParameterId(string name, HashSet<ParameterId> parameterIds)
+		protected static ParameterId CreateUniqueParameterId(string name, HashSet<ParameterId> parameterIds)
 		{
 			string baseName = name;
 			if (String.IsNullOrEmpty(baseName))
diff --git a/Engine/Source/Programs/UnrealGameSync/UnrealGameSync/Controls/WorkspaceControl.cs b/Engine/Source/Programs/UnrealGameSync/UnrealGameSync/Controls/WorkspaceControl.cs
--- a/Engine/Source/Programs/UnrealGameSync/UnrealGameSync/Controls/WorkspaceControl.cs
+++ b/Engine/Source/Programs/UnrealGameSync/UnrealGameSync/Controls/WorkspaceControl.cs
@@ -4271,12 +4271,14 @@ namespace UnrealGameSync
 					}
 					if (openedFilesTask.Result)
 					{
+#pragma warning disable CA2000 // Dispose objects before losing scope
 						SwitchStreamWarningWindow switchStreamWarningWindow = new SwitchStreamWarningWindow(openedFiles);
 						switchStreamWarningWindow.FormBorderStyle = FormBorderStyle.FixedDialog;
 						if (switchStreamWarningWindow.ShowDialog() == DialogResult.Cancel)
 						{
 							return;
 						}
+#pragma warning restore CA2000 // Dispose objects before losing scope
 					}
 
 					// switch to the new stream
diff --git a/Engine/Source/Programs/UnrealGameSync/UnrealGameSync/Forms/SwitchStreamWarningWindow.cs b/Engine/Source/Programs/UnrealGameSync/UnrealGameSync/Forms/SwitchStreamWarningWindow.cs
--- a/Engine/Source/Programs/UnrealGameSync/UnrealGameSync/Forms/SwitchStreamWarningWindow.cs
+++ b/Engine/Source/Programs/UnrealGameSync/UnrealGameSync/Forms/SwitchStreamWarningWindow.cs
@@ -2,12 +2,6 @@
 
 using System;
 using System.Collections.Generic;
-using System.ComponentModel;
-using System.Data;
-using System.Drawing;
-using System.Linq;
-using System.Text;
-using System.Threading.Tasks;
 using System.Windows.Forms;
 
 namespace UnrealGameSync.Forms
@@ -22,7 +16,6 @@ namespace UnrealGameSync.Forms
 			{
 				OpenedFiles.Items.Add(openedFile);
 			}
-			
 		}
 
 		private void OkButton_Click(object sender, EventArgs e)
