Horde: Fix bug where physical cores are counted twice in capability detection

diff --git a/Engine/Source/Programs/Horde/HordeAgent/Services/CapabilitiesService.cs b/Engine/Source/Programs/Horde/HordeAgent/Services/CapabilitiesService.cs
--- a/Engine/Source/Programs/Horde/HordeAgent/Services/CapabilitiesService.cs
+++ b/Engine/Source/Programs/Horde/HordeAgent/Services/CapabilitiesService.cs
@@ -126,41 +126,31 @@ namespace HordeAgent.Services
 					{
 						// Add CPU info
 						using TelemetrySpan _ = _tracer.StartActiveSpan("GetCpuInfo");
-						Dictionary<string, int> cpuNameToCount = new Dictionary<string, int>();
+						Dictionary<string, int> cpuNameToCount = new ();
 						int totalPhysicalCores = 0;
+						int totalEnabledPhysicalCores = 0;
 						int totalLogicalCores = 0;
 						
 						foreach (CimInstance instance in session.QueryInstances(QueryNamespace, QueryDialect, "select * from Win32_Processor"))
 						{
 							foreach (CimProperty property in instance.CimInstanceProperties)
 							{
-								string name = property.Name;
-								if (name.Equals("Name", StringComparison.OrdinalIgnoreCase))
-								{
-									string cpuName = property.Value.ToString() ?? String.Empty;
-									int count;
-									cpuNameToCount.TryGetValue(cpuName, out count);
-									cpuNameToCount[cpuName] = count + 1;
-								}
-								else if (name.Equals("NumberOfEnabledCore", StringComparison.OrdinalIgnoreCase) ||
-								         name.Equals("NumberOfCores", StringComparison.OrdinalIgnoreCase))
-								{
-									if (property.Value is uint numCores)
-									{
-										totalPhysicalCores += (int)numCores;
-									}
-								}
-								else if (name.Equals("NumberOfLogicalProcessors", StringComparison.OrdinalIgnoreCase))
+								switch (property.Name.ToUpperInvariant())
 								{
-									if (property.Value is uint numCores)
-									{
-										totalLogicalCores += (int)numCores;
-									}
+									case "NAME":
+										{
+											string cpuName = property.Value.ToString() ?? String.Empty;
+											cpuNameToCount.TryGetValue(cpuName, out int count);
+											cpuNameToCount[cpuName] = count + 1;
+											break;
+										}
+									case "NUMBEROFCORES" when property.Value is uint c: totalPhysicalCores += (int)c; break;
+									case "NUMBEROFENABLEDCORE" when property.Value is uint c: totalEnabledPhysicalCores += (int)c; break;
 								}
 							}
 						}
 						
-						AddCpuInfo(capabilities, cpuNameToCount, totalLogicalCores, totalPhysicalCores);
+						AddCpuInfo(capabilities, cpuNameToCount, totalLogicalCores, Math.Min(totalEnabledPhysicalCores, totalPhysicalCores));
 					}
 					
 					{
