Horde: Fix infinite loop issue when querying for jobs

diff --git a/Engine/Source/Programs/Horde/Plugins/Build/HordeServer.Build.Tests/Fixture.cs b/Engine/Source/Programs/Horde/Plugins/Build/HordeServer.Build.Tests/Fixture.cs
--- a/Engine/Source/Programs/Horde/Plugins/Build/HordeServer.Build.Tests/Fixture.cs
+++ b/Engine/Source/Programs/Horde/Plugins/Build/HordeServer.Build.Tests/Fixture.cs
@@ -48,13 +48,7 @@ namespace HordeServer.Tests
 
 		private async Task PopulateAsync(ConfigService configService, IGraphCollection graphCollection, ITemplateCollection templateCollection, JobService jobService, AgentService agentService, IPluginCollection pluginCollection, ServerSettings serverSettings)
 		{
-			FixtureGraph fg = new FixtureGraph();
-			fg.Id = ContentHash.Empty;
-			fg.Schema = 1122;
-			fg.Groups = new List<INodeGroup>();
-			fg.Aggregates = new List<IAggregate>();
-			fg.Labels = new List<ILabel>();
-
+			GraphStub fg = new (ContentHash.Empty, 1122);
 			Template = await templateCollection.GetOrAddAsync(new TemplateConfig { Name = "Test template" });
 			Graph = await graphCollection.AddAsync(Template, null);
 
@@ -119,15 +113,31 @@ namespace HordeServer.Tests
 			Agent1Name = "testAgent1";
 			Agent1 = await agentService.CreateAgentAsync(new CreateAgentOptions(new AgentId(Agent1Name), false, ""));
 		}
-
-		private class FixtureGraph : IGraph
-		{
-			public ContentHash Id { get; set; } = ContentHash.Empty;
-			public int Schema { get; set; }
-			public IReadOnlyList<INodeGroup> Groups { get; set; } = null!;
-			public IReadOnlyList<IAggregate> Aggregates { get; set; } = null!;
-			public IReadOnlyList<ILabel> Labels { get; set; } = null!;
-			public IReadOnlyList<IGraphArtifact> Artifacts { get; set; } = null!;
-		}
+	}
+	
+	internal class GraphStub(ContentHash id, int schema = 1) : IGraph
+	{
+		public ContentHash Id { get; set; } = id;
+		public int Schema { get; set; } = schema;
+		public IReadOnlyList<INodeGroup> Groups { get; set; } = [];
+		public IReadOnlyList<IAggregate> Aggregates { get; set; } = [];
+		public IReadOnlyList<ILabel> Labels { get; set; } = [];
+		public IReadOnlyList<IGraphArtifact> Artifacts { get; set; } = [];
+	}
+	
+	internal class TemplateStub(ContentHash hash, string name) : ITemplate
+	{
+		public ContentHash Hash { get; } = hash;
+		public string Name { get; } = name;
+		public string? Description { get; init; } = null;
+		public Priority? Priority { get; init; } = null;
+		public bool AllowPreflights { get; init; } = false;
+		public bool UpdateIssues { get; init; } = false;
+		public bool PromoteIssuesByDefault { get; init; } = false;
+		public string? InitialAgentType { get; init; } = null;
+		public string? SubmitNewChange { get; init; } = null;
+		public string? SubmitDescription { get; init; } = null;
+		public IReadOnlyList<string> Arguments { get; init; } = [];
+		public IReadOnlyList<Parameter> Parameters { get; init; } = [];
 	}
 }
\ No newline at end of file
diff --git a/Engine/Source/Programs/Horde/Plugins/Build/HordeServer.Build.Tests/Jobs/JobCollectionTests.cs b/Engine/Source/Programs/Horde/Plugins/Build/HordeServer.Build.Tests/Jobs/JobCollectionTests.cs
--- a/Engine/Source/Programs/Horde/Plugins/Build/HordeServer.Build.Tests/Jobs/JobCollectionTests.cs
+++ b/Engine/Source/Programs/Horde/Plugins/Build/HordeServer.Build.Tests/Jobs/JobCollectionTests.cs
@@ -28,14 +28,16 @@ namespace HordeServer.Tests.Jobs
 	[TestClass]
 	public class JobCollectionTests : BuildTestSetup
 	{
+		private readonly StreamId _streamId = new ("ue4-main");
+		private readonly TemplateId _templateId1 = new ("template1");
+		private readonly ITemplate _template1 = new TemplateStub(ContentHash.MD5("graphHash"), "template1");
+		
 		public JobCollectionTests()
 		{
 			UpdateConfig(globalConfig =>
 			{
-				ProjectConfig projectConfig = new ProjectConfig();
-				projectConfig.Id = new ProjectId("ue4");
-				projectConfig.Streams.Add(new StreamConfig { Id = new StreamId("ue4-main") });
-
+				ProjectConfig projectConfig = new () { Id = new ProjectId("ue4") };
+				projectConfig.Streams.Add(new StreamConfig { Id = _streamId });
 				globalConfig.Plugins.GetBuildConfig().Projects.Add(projectConfig);
 			});
 		}
@@ -842,5 +844,35 @@ namespace HordeServer.Tests.Jobs
 			Assert.AreEqual("Engine/Source", graph.Artifacts[0].BasePath);
 			Assert.AreEqual("fileset", graph.Artifacts[0].OutputName);
 		}
+		
+		[TestMethod]
+		public async Task FindJobsAsync()
+		{
+			IGraph graph = await GraphCollection.AddAsync(_template1, null);
+			
+			List<IJob> jobs = [];
+			for (int i = 0; i < 10; i++)
+			{
+				jobs.Add(await AddJobAsync(graph, $"{i}"));
+			}
+			
+			jobs[2] = await RunBatchAsync(jobs[2], 0);
+			await RunStepAsync(jobs[2], 0, 0, JobStepOutcome.Success); // Setup Build
+			
+			{
+				IReadOnlyList<IJob> jobsFound = await JobCollection.FindAsync(new FindJobOptions());
+				Assert.AreEqual(jobs.Count, jobsFound.Count);
+			}
+			
+			{
+				IReadOnlyList<IJob> jobsFound = await JobCollection.FindAsync(new FindJobOptions() { Outcome = [JobStepOutcome.Success] });
+				Assert.AreEqual(1, jobsFound.Count);
+			}
+		}
+		
+		private async Task<IJob> AddJobAsync(IGraph graph, string commitId)
+		{
+			return await JobCollection.AddAsync(JobIdUtils.GenerateNewId(), _streamId, _templateId1, _template1.Hash, graph, "jobName", new CommitId(commitId), null, new CreateJobOptions());
+		}
 	}
 }
diff --git a/Engine/Source/Programs/Horde/Plugins/Build/HordeServer.Build/Jobs/JobCollection.cs b/Engine/Source/Programs/Horde/Plugins/Build/HordeServer.Build/Jobs/JobCollection.cs
--- a/Engine/Source/Programs/Horde/Plugins/Build/HordeServer.Build/Jobs/JobCollection.cs
+++ b/Engine/Source/Programs/Horde/Plugins/Build/HordeServer.Build/Jobs/JobCollection.cs
@@ -925,12 +925,13 @@ namespace HordeServer.Jobs
 				List<IJob> results = new List<IJob>();
 				_logger.LogInformation("Performing scan for job with ");
 
+				const int Count = 5;
 				bool excludeMaxCommitId = false;
-
+				int scanIndex = 0;
 				int maxCount = (count ?? 1);
 				while (results.Count < maxCount)
 				{
-					IReadOnlyList<IJob> scanJobs = await FindInternalAsync(options, 0, 5, cancellationToken: cancellationToken);
+					IReadOnlyList<IJob> scanJobs = await FindInternalAsync(options, scanIndex, Count, cancellationToken: cancellationToken);
 					if (scanJobs.Count == 0)
 					{
 						break;
@@ -972,7 +973,7 @@ namespace HordeServer.Jobs
 						}
 					}
 
-					options = options with { MaxCommitId = scanJobs.Min(x => x.CommitId) };
+					scanIndex += Count;
 					excludeMaxCommitId = true;
 				}
 
