Horde: Add a dedicated installer for the tray app, and include it in the server installer.

diff --git a/Engine/Source/Programs/Horde/BuildHorde.xml b/Engine/Source/Programs/Horde/BuildHorde.xml
--- a/Engine/Source/Programs/Horde/BuildHorde.xml
+++ b/Engine/Source/Programs/Horde/BuildHorde.xml
@@ -129,6 +129,8 @@
 	<Property Name="ServerToolsStagingDir_HordeAgent" Value="$(ServerToolsStagingDir)/HordeAgent"/>
 	<Property Name="ServerToolsStagingDir_HordeAgentWin64" Value="$(ServerToolsStagingDir)/HordeAgentWin64"/>
 	<Property Name="ServerToolsStagingDir_HordeAgentMsi" Value="$(ServerToolsStagingDir)/HordeAgentMsi"/>
+	<Property Name="ServerToolsStagingDir_HordeTrayApp" Value="$(ServerToolsStagingDir)/HordeTrayApp"/>
+	<Property Name="ServerToolsStagingDir_HordeTrayAppMsi" Value="$(ServerToolsStagingDir)/HordeTrayAppMsi"/>
 	<Property Name="ServerToolsStagingDir_P4VUtils" Value="$(ServerToolsStagingDir)/P4VUtils"/>
 	<Property Name="ServerToolsStagingDir_UnrealGameSync" Value="$(ServerToolsStagingDir)/UnrealGameSync"/>
 	<Property Name="ServerToolsStagingDir_UnrealGameSyncCmd" Value="$(ServerToolsStagingDir)/UnrealGameSyncCmd"/>
@@ -171,6 +173,34 @@
 			<Delete Files="$(ServerToolsStagingDir_HordeAgentWin64)/appsettings.Local.json"/>
 			<Tag Files="$(ServerToolsStagingDir_HordeAgentWin64)/..." With="#HordeAgentWin64Files"/>
 		</Node>
+		<Node Name="Build HordeTrayApp" Produces="#HordeTrayAppFiles">
+			<Delete Files="$(ServerToolsStagingDir_HordeTrayApp)/..."/>
+			<DotNet	Arguments="publish &quot;$(RootDir)/Engine/Source/Programs/Horde/HordeTrayApp/HordeTrayApp.csproj&quot; --output &quot;$(ServerToolsStagingDir_HordeTrayApp)&quot; --runtime win-x64 --self-contained $(VersionArguments)" DotNetPath="$(DotNetPath)" />
+			<Sign Files="$(ServerToolsStagingDir_HordeTrayApp)/Horde.Agent.TrayApp.exe;$(ServerToolsStagingDir_HordeTrayApp)/Horde.dll;$(ServerToolsStagingDir_HordeTrayApp)/EpicGames.*.dll" If="'$(SignOutputs)' == 'true'"/>
+			<Tag Files="$(ServerToolsStagingDir_HordeTrayApp)/..." With="#HordeTrayAppFiles"/>
+		</Node>
+		<Node Name="Build HordeTrayApp Msi" Requires="Build HordeTrayApp" Produces="#HordeTrayAppMsiFiles">
+			<Property Name="CommonArgs" Value="-dPlatform=x64 -arch x64"/>
+
+			<!-- Create the tray app file list -->
+			<Delete Files="$(StagingDir)/TrayApp-Bulk/..."/>
+			<Tag Files="#HordeTrayAppFiles" Except="Horde.Agent.TrayApp.exe" With="#TrayAppBulkFiles"/>
+			<Copy Files="#TrayAppBulkFiles" From="$(ServerToolsStagingDir_HordeTrayApp)" To="$(StagingDir)/TrayApp-Bulk"/>
+
+			<Spawn Exe="$(WixDir)/heat.exe" Arguments="dir &quot;$(StagingDir)/TrayApp-Bulk&quot; -cg HordeTrayAppFiles -scom -sreg -gg -dr TrayAppInstallDir -srd -var var.SourceDir -out &quot;$(InstallerDir)/HordeTrayAppFiles.wxs&quot;"/>
+			<Spawn Exe="$(WixDir)/candle.exe" Arguments="$(CommonArgs) HordeTrayAppFiles.wxs -dSourceDir=&quot;$(StagingDir)/TrayApp-Bulk&quot; -out &quot;obj/HordeTrayAppFiles.wixobj&quot;" WorkingDir="$(InstallerDir)"/>
+
+			<!-- Main component -->
+			<Spawn Exe="$(WixDir)/candle.exe" Arguments="$(CommonArgs) TrayAppInstaller.wxs -dStagingDir=&quot;$(ServerToolsStagingDir_HordeTrayApp)&quot; -out &quot;obj/HordeTrayApp.wixobj&quot; -ext WixUtilExtension" WorkingDir="$(InstallerDir)"/>
+			<Spawn Exe="$(WixDir)/light.exe" Arguments="-out &quot;bin/HordeTrayApp.msi&quot; -sw1076 -pdbout &quot;bin/HordeTrayApp.pdb&quot; &quot;obj/HordeTrayApp.wixobj&quot; &quot;obj/HordeTrayAppFiles.wixobj&quot; -ext WixUtilExtension -ext WixUIExtension" WorkingDir="$(InstallerDir)"/>
+
+			<!-- We need to sign here as will be included in the server msi -->
+			<Sign Files="$(InstallerDir)/bin/HordeTrayApp.msi" Description="Horde Tray App" If="'$(SignOutputs)' == 'true'"/>
+			
+			<!-- Copy files to the output directory -->
+			<Copy Files="HordeTrayApp.msi" From="$(InstallerDir)/bin" To="$(ServerToolsStagingDir_HordeTrayAppMsi)"/>
+			<Tag Files="$(ServerToolsStagingDir_HordeTrayAppMsi)/..." With="#HordeTrayAppMsiFiles"/>
+		</Node>
 		<Node Name="Build P4VUtils" Produces="#P4VUtilsFiles">
 			<Delete Files="$(ServerToolsStagingDir_P4VUtils)/..."/>
 			<DotNet Arguments="publish &quot;$(RootDir)/Engine/Extras/P4VUtils/P4VUtils.csproj&quot; --configuration Release --runtime win-x64 --self-contained --output &quot;$(ServerToolsStagingDir_P4VUtils)&quot; /p:IsWindows=true /p:IsOSX=false /p:IsLinux=false /p:WithRestricted=false $(VersionArguments)"/>
@@ -263,7 +293,7 @@
 		</Node>
 
 		<ForEach Name="ServerArch" Values="$(ServerArchList)">
-			<Node Name="Stage Horde Server Installer ($(ServerArch))" Requires="#DashboardOutputFiles;#HordeCmdFiles;#HordeAgentFiles;#HordeAgentMsiFiles;#P4VUtilsFiles;#UnrealGameSyncFiles;#UnrealGameSyncCmdFiles;#UnrealGameSyncMsiFiles;" Produces="#ServerWithToolsFiles_$(ServerArch)">
+			<Node Name="Stage Horde Server Installer ($(ServerArch))" Requires="#DashboardOutputFiles;#HordeCmdFiles;#HordeAgentFiles;#HordeAgentMsiFiles;#HordeTrayAppFiles;#HordeTrayAppMsiFiles;#P4VUtilsFiles;#UnrealGameSyncFiles;#UnrealGameSyncCmdFiles;#UnrealGameSyncMsiFiles;" Produces="#ServerWithToolsFiles_$(ServerArch)">
 		
 				<!-- Stage the server -->
 				<Delete Files="$(StagingDir)/Server/$(ServerArch)/..."/>
@@ -273,15 +303,17 @@
 				<Delete Files="$(StagingDir)/Server/$(ServerArch)/appsettings.Local.json"/>
 
 				<!-- Add the bundled tools into the installer -->
-				<DotNet Arguments="run --project &quot;$(RootDir)/Engine/Source/Programs/Horde/HordeServer/HordeServer.csproj&quot; -- generate tooldata -id=horde-cmd -name=&quot;Horde (Command-Line)&quot; -description=&quot;Update package for the Horde command-line tool.&quot; -version=&quot;$(InformationalVersion)&quot; -inputdir=&quot;$(ServerToolsStagingDir_HordeCmd)&quot; -serverdir=&quot;$(StagingDir)/Server/$(ServerArch)&quot; -public -showinugs -Category=Internal" DotNetPath="$(DotNetPath)" />
+				<DotNet Arguments="run --project &quot;$(RootDir)/Engine/Source/Programs/Horde/HordeServer/HordeServer.csproj&quot; -- generate tooldata -id=horde-cmd -name=&quot;Horde (Command-Line)&quot; -description=&quot;Update package for the Horde command-line tool.&quot; -version=&quot;$(InformationalVersion)&quot; -inputdir=&quot;$(ServerToolsStagingDir_HordeCmd)&quot; -serverdir=&quot;$(StagingDir)/Server/$(ServerArch)&quot; -public -showinugs -showinlauncher -Category=Internal" DotNetPath="$(DotNetPath)" />
 				<DotNet Arguments="run --project &quot;$(RootDir)/Engine/Source/Programs/Horde/HordeServer/HordeServer.csproj&quot; -- generate tooldata -id=horde-agent -name=&quot;Horde Agent (Cross-Platform)&quot; -description=&quot;Cross-platform Horde Agent distribution. Allows hosts to take on work from the Horde Server.&quot; -version=&quot;$(InformationalVersion)&quot; -inputdir=&quot;$(ServerToolsStagingDir_HordeAgent)&quot; -serverdir=&quot;$(StagingDir)/Server/$(ServerArch)&quot; -public" DotNetPath="$(DotNetPath)" />
 				<DotNet Arguments="run --project &quot;$(RootDir)/Engine/Source/Programs/Horde/HordeServer/HordeServer.csproj&quot; -- generate tooldata -id=horde-agent-win64 -name=&quot;Horde Agent (Windows)&quot; -description=&quot;Update package for older installations of Horde Agent through the MSI installer.&quot; -version=&quot;$(InformationalVersion)&quot; -inputdir=&quot;$(ServerToolsStagingDir_HordeAgentWin64)&quot; -serverdir=&quot;$(StagingDir)/Server/$(ServerArch)&quot; -public -showindashboard=false -category=Internal" DotNetPath="$(DotNetPath)" />
 				<DotNet Arguments="run --project &quot;$(RootDir)/Engine/Source/Programs/Horde/HordeServer/HordeServer.csproj&quot; -- generate tooldata -id=horde-agent-win64-installer -name=&quot;Horde Agent (Windows Installer)&quot; -description=&quot;Installer for the Horde Agent on Windows, with tray application.&quot; -version=&quot;$(InformationalVersion)&quot; -inputdir=&quot;$(ServerToolsStagingDir_HordeAgentMsi)&quot; -serverdir=&quot;$(StagingDir)/Server/$(ServerArch)&quot; -public" DotNetPath="$(DotNetPath)" />
-				<DotNet Arguments="run --project &quot;$(RootDir)/Engine/Source/Programs/Horde/HordeServer/HordeServer.csproj&quot; -- generate tooldata -id=p4v-utils -name=&quot;P4V Extensions&quot; -description=&quot;Update package for the P4VUtils tool.&quot; -version=&quot;$(InformationalVersion)&quot; -inputdir=&quot;$(ServerToolsStagingDir_P4VUtils)&quot; -serverdir=&quot;$(StagingDir)/Server/$(ServerArch)&quot; -public -showinugs -category=Internal" DotNetPath="$(DotNetPath)" />
-				<DotNet Arguments="run --project &quot;$(RootDir)/Engine/Source/Programs/Horde/HordeServer/HordeServer.csproj&quot; -- generate tooldata -id=ugs-cmd -name=&quot;UnrealGameSync (Command-Line)&quot; -description=&quot;Command-line version of UnrealGameSync. Supports Windows, MacOS and Linux.&quot; -version=&quot;$(InformationalVersion)&quot; -inputdir=&quot;$(ServerToolsStagingDir_UnrealGameSyncCmd)&quot; -serverdir=&quot;$(StagingDir)/Server/$(ServerArch)&quot; -public" DotNetPath="$(DotNetPath)" />
-				<DotNet Arguments="run --project &quot;$(RootDir)/Engine/Source/Programs/Horde/HordeServer/HordeServer.csproj&quot; -- generate tooldata -id=ugs-msi -name=&quot;UnrealGameSync (Windows Installer)&quot; -description=&quot;Installer for UnrealGameSync.&quot; -version=&quot;$(InformationalVersion)&quot; -inputdir=&quot;$(ServerToolsStagingDir_UnrealGameSyncMsi)&quot; -serverdir=&quot;$(StagingDir)/Server/$(ServerArch)&quot; -public" DotNetPath="$(DotNetPath)" />
+				<DotNet Arguments="run --project &quot;$(RootDir)/Engine/Source/Programs/Horde/HordeServer/HordeServer.csproj&quot; -- generate tooldata -id=horde-tray-app -name=&quot;Horde Tray App&quot; -description=&quot;Update package for the Horde Tray App.&quot; -version=&quot;$(InformationalVersion)&quot; -inputdir=&quot;$(ServerToolsStagingDir_HordeTrayApp)&quot; -serverdir=&quot;$(StagingDir)/Server/$(ServerArch)&quot; -public -showinugs -Category=Internal" DotNetPath="$(DotNetPath)" />
+				<DotNet Arguments="run --project &quot;$(RootDir)/Engine/Source/Programs/Horde/HordeServer/HordeServer.csproj&quot; -- generate tooldata -id=horde-tray-app-msi -name=&quot;Horde Tray App (Windows Installer)&quot; -description=&quot;Windows installer for the Horde Tray App.&quot; -version=&quot;$(InformationalVersion)&quot; -inputdir=&quot;$(ServerToolsStagingDir_HordeTrayAppMsi)&quot; -serverdir=&quot;$(StagingDir)/Server/$(ServerArch)&quot; -public -showinugs" DotNetPath="$(DotNetPath)" />
+				<DotNet Arguments="run --project &quot;$(RootDir)/Engine/Source/Programs/Horde/HordeServer/HordeServer.csproj&quot; -- generate tooldata -id=p4v-utils -name=&quot;P4V Extensions&quot; -description=&quot;Update package for the P4VUtils tool.&quot; -version=&quot;$(InformationalVersion)&quot; -inputdir=&quot;$(ServerToolsStagingDir_P4VUtils)&quot; -serverdir=&quot;$(StagingDir)/Server/$(ServerArch)&quot; -public -showinugs -showinlauncher -category=Internal" DotNetPath="$(DotNetPath)" />
+				<DotNet Arguments="run --project &quot;$(RootDir)/Engine/Source/Programs/Horde/HordeServer/HordeServer.csproj&quot; -- generate tooldata -id=ugs-cmd -name=&quot;UnrealGameSync (Command-Line)&quot; -description=&quot;Command-line version of UnrealGameSync. Supports Windows, MacOS and Linux.&quot; -version=&quot;$(InformationalVersion)&quot; -inputdir=&quot;$(ServerToolsStagingDir_UnrealGameSyncCmd)&quot; -serverdir=&quot;$(StagingDir)/Server/$(ServerArch)&quot; -public -showinlauncher" DotNetPath="$(DotNetPath)" />
+				<DotNet Arguments="run --project &quot;$(RootDir)/Engine/Source/Programs/Horde/HordeServer/HordeServer.csproj&quot; -- generate tooldata -id=ugs-msi -name=&quot;UnrealGameSync (Windows Installer)&quot; -description=&quot;Installer for UnrealGameSync.&quot; -version=&quot;$(InformationalVersion)&quot; -inputdir=&quot;$(ServerToolsStagingDir_UnrealGameSyncMsi)&quot; -serverdir=&quot;$(StagingDir)/Server/$(ServerArch)&quot; -public -showinlauncher" DotNetPath="$(DotNetPath)" />
 				<DotNet Arguments="run --project &quot;$(RootDir)/Engine/Source/Programs/Horde/HordeServer/HordeServer.csproj&quot; -- generate tooldata -id=ugs-win -name=&quot;UnrealGameSync (Windows Update Package)&quot; -description=&quot;Update package for UGS on Windows.&quot; -version=&quot;$(InformationalVersion)&quot; -inputdir=&quot;$(ServerToolsStagingDir_UnrealGameSync)&quot; -serverdir=&quot;$(StagingDir)/Server/$(ServerArch)&quot; -public -showindashboard=false -category=Internal" DotNetPath="$(DotNetPath)" />
-				<DotNet Arguments="run --project &quot;$(RootDir)/Engine/Source/Programs/Horde/HordeServer/HordeServer.csproj&quot; -- generate tooldata -id=ushell -name=&quot;Ushell&quot; -description=&quot;Update package for ushell.&quot; -version=&quot;$(InformationalVersion)&quot; -inputdir=&quot;$(RootDir)/Engine/Extras/ushell&quot; -serverdir=&quot;$(StagingDir)/Server/$(ServerArch)&quot; -public -showinugs -Category=Internal" DotNetPath="$(DotNetPath)" />
+				<DotNet Arguments="run --project &quot;$(RootDir)/Engine/Source/Programs/Horde/HordeServer/HordeServer.csproj&quot; -- generate tooldata -id=ushell -name=&quot;Ushell&quot; -description=&quot;Update package for ushell.&quot; -version=&quot;$(InformationalVersion)&quot; -inputdir=&quot;$(RootDir)/Engine/Extras/ushell&quot; -serverdir=&quot;$(StagingDir)/Server/$(ServerArch)&quot; -public -showinugs -showinlauncher -Category=Internal" DotNetPath="$(DotNetPath)" />
 
 				<!-- Separate copy of the command line tool for installation with the server on Windows -->
 				<Delete Files="$(StagingDir)/Tools/..."/>
