Horde: Rename HordeTrayApp to HordeToolbox.

diff --git a/Engine/Source/Programs/Horde/BuildHorde.xml b/Engine/Source/Programs/Horde/BuildHorde.xml
--- a/Engine/Source/Programs/Horde/BuildHorde.xml
+++ b/Engine/Source/Programs/Horde/BuildHorde.xml
@@ -101,12 +101,12 @@
 		</Node>
 	</Agent>
 	
-	<Agent Name="HordeTrayApp" Type="Win64">
-		<Node Name="Analyze HordeTrayApp">
-			<DotNet Arguments="build &quot;$(EngineDir)/Source/Programs/Horde/HordeTrayApp/HordeTrayApp.csproj&quot; -p:Configuration=Analyze --nologo" DotNetPath="$(DotNetPath)"/>
+	<Agent Name="HordeToolbox" Type="Win64">
+		<Node Name="Analyze HordeToolbox">
+			<DotNet Arguments="build &quot;$(EngineDir)/Source/Programs/Horde/HordeToolbox/HordeToolbox.csproj&quot; -p:Configuration=Analyze --nologo" DotNetPath="$(DotNetPath)"/>
 		</Node>	
 		
-		<Node Name="Check HordeTrayApp Licenses" Requires="Analyze HordeTrayApp">
+		<Node Name="Check HordeToolbox Licenses" Requires="Analyze HordeToolbox">
 			<Property Name="IgnorePackages">
 				[REDACTED_EMAIL]
 				[REDACTED_EMAIL]
@@ -114,7 +114,7 @@
 				[REDACTED_EMAIL]
 				[REDACTED_EMAIL]
 			</Property>
-			<NuGet-LicenseCheck BaseDir="Engine/Source/Programs/Horde/HordeTrayApp" IgnorePackages="$(IgnorePackages)" LicenseDir="Engine/Source/Programs/Horde/ThirdParty/Licenses" DotNetPath="$(DotNetPath)"/>
+			<NuGet-LicenseCheck BaseDir="Engine/Source/Programs/Horde/HordeToolbox" IgnorePackages="$(IgnorePackages)" LicenseDir="Engine/Source/Programs/Horde/ThirdParty/Licenses" DotNetPath="$(DotNetPath)"/>
 		</Node>
 	</Agent>
 	
@@ -129,8 +129,8 @@
 	<Property Name="ServerToolsStagingDir_HordeAgent" Value="$(ServerToolsStagingDir)/HordeAgent"/>
 	<Property Name="ServerToolsStagingDir_HordeAgentWin64" Value="$(ServerToolsStagingDir)/HordeAgentWin64"/>
 	<Property Name="ServerToolsStagingDir_HordeAgentMsi" Value="$(ServerToolsStagingDir)/HordeAgentMsi"/>
-	<Property Name="ServerToolsStagingDir_HordeTrayApp" Value="$(ServerToolsStagingDir)/HordeTrayApp"/>
-	<Property Name="ServerToolsStagingDir_HordeTrayAppMsi" Value="$(ServerToolsStagingDir)/HordeTrayAppMsi"/>
+	<Property Name="ServerToolsStagingDir_HordeToolbox" Value="$(ServerToolsStagingDir)/HordeToolbox"/>
+	<Property Name="ServerToolsStagingDir_HordeToolboxMsi" Value="$(ServerToolsStagingDir)/HordeToolboxMsi"/>
 	<Property Name="ServerToolsStagingDir_P4VUtils" Value="$(ServerToolsStagingDir)/P4VUtils"/>
 	<Property Name="ServerToolsStagingDir_UnrealGameSync" Value="$(ServerToolsStagingDir)/UnrealGameSync"/>
 	<Property Name="ServerToolsStagingDir_UnrealGameSyncCmd" Value="$(ServerToolsStagingDir)/UnrealGameSyncCmd"/>
@@ -166,40 +166,40 @@
 			<DotNet Arguments="publish &quot;$(RootDir)/Engine/Source/Programs/Horde/Drivers/JobDriver/JobDriver.csproj&quot; --output &quot;$(ServerToolsStagingDir_HordeAgentWin64)/JobDriver&quot; --runtime win-x64 --self-contained $(VersionArguments)"/>
 			<Sign Files="$(ServerToolsStagingDir_HordeAgentWin64)/JobDriver/JobDriver.exe;$(ServerToolsStagingDir_HordeAgentWin64)/JobDriver/JobDriver.dll;$(ServerToolsStagingDir_HordeAgentWin64)/JobDriver/EpicGames.*.dll" If="'$(SignOutputs)' == 'true'"/>
 
-			<!-- Build the tray app -->
-			<DotNet Arguments="publish &quot;$(RootDir)/Engine/Source/Programs/Horde/HordeTrayApp/HordeTrayApp.csproj&quot; --output &quot;$(ServerToolsStagingDir_HordeAgentWin64)/TrayApp&quot; --runtime win-x64 --self-contained $(VersionArguments)"/>
-			<Sign Files="$(ServerToolsStagingDir_HordeAgentWin64)/TrayApp/Horde.Agent.TrayApp.exe;$(ServerToolsStagingDir_HordeAgentWin64)/TrayApp/EpicGames.*.dll" If="'$(SignOutputs)' == 'true'"/>
+			<!-- Build the toolbox -->
+			<DotNet Arguments="publish &quot;$(RootDir)/Engine/Source/Programs/Horde/HordeToolbox/HordeToolbox.csproj&quot; --output &quot;$(ServerToolsStagingDir_HordeAgentWin64)/Toolbox&quot; --runtime win-x64 --self-contained $(VersionArguments)"/>
+			<Sign Files="$(ServerToolsStagingDir_HordeAgentWin64)/Toolbox/HordeToolbox.exe;$(ServerToolsStagingDir_HordeAgentWin64)/Toolbox/EpicGames.*.dll" If="'$(SignOutputs)' == 'true'"/>
 
 			<Delete Files="$(ServerToolsStagingDir_HordeAgentWin64)/appsettings.Local.json"/>
 			<Tag Files="$(ServerToolsStagingDir_HordeAgentWin64)/..." With="#HordeAgentWin64Files"/>
 		</Node>
-		<Node Name="Build HordeTrayApp" Produces="#HordeTrayAppFiles">
-			<Delete Files="$(ServerToolsStagingDir_HordeTrayApp)/..."/>
-			<DotNet	Arguments="publish &quot;$(RootDir)/Engine/Source/Programs/Horde/HordeTrayApp/HordeTrayApp.csproj&quot; --output &quot;$(ServerToolsStagingDir_HordeTrayApp)&quot; --runtime win-x64 --self-contained $(VersionArguments)" DotNetPath="$(DotNetPath)" />
-			<Sign Files="$(ServerToolsStagingDir_HordeTrayApp)/Horde.Agent.TrayApp.exe;$(ServerToolsStagingDir_HordeTrayApp)/Horde.dll;$(ServerToolsStagingDir_HordeTrayApp)/EpicGames.*.dll" If="'$(SignOutputs)' == 'true'"/>
-			<Tag Files="$(ServerToolsStagingDir_HordeTrayApp)/..." With="#HordeTrayAppFiles"/>
+		<Node Name="Build HordeToolbox" Produces="#HordeToolboxFiles">
+			<Delete Files="$(ServerToolsStagingDir_HordeToolbox)/..."/>
+			<DotNet	Arguments="publish &quot;$(RootDir)/Engine/Source/Programs/Horde/HordeToolbox/HordeToolbox.csproj&quot; --output &quot;$(ServerToolsStagingDir_HordeToolbox)&quot; --runtime win-x64 --self-contained $(VersionArguments)" DotNetPath="$(DotNetPath)" />
+			<Sign Files="$(ServerToolsStagingDir_HordeToolbox)/HordeToolbox.exe;$(ServerToolsStagingDir_HordeToolbox)/Horde.dll;$(ServerToolsStagingDir_HordeToolbox)/EpicGames.*.dll" If="'$(SignOutputs)' == 'true'"/>
+			<Tag Files="$(ServerToolsStagingDir_HordeToolbox)/..." With="#HordeToolboxFiles"/>
 		</Node>
-		<Node Name="Build HordeTrayApp Msi" Requires="Build HordeTrayApp" Produces="#HordeTrayAppMsiFiles">
+		<Node Name="Build HordeToolbox Msi" Requires="Build HordeToolbox" Produces="#HordeToolboxMsiFiles">
 			<Property Name="CommonArgs" Value="-dPlatform=x64 -arch x64"/>
 
-			<!-- Create the tray app file list -->
-			<Delete Files="$(StagingDir)/TrayApp-Bulk/..."/>
-			<Tag Files="#HordeTrayAppFiles" Except="Horde.Agent.TrayApp.exe" With="#TrayAppBulkFiles"/>
-			<Copy Files="#TrayAppBulkFiles" From="$(ServerToolsStagingDir_HordeTrayApp)" To="$(StagingDir)/TrayApp-Bulk"/>
+			<!-- Create the toolbox file list -->
+			<Delete Files="$(StagingDir)/Toolbox-Bulk/..."/>
+			<Tag Files="#HordeToolboxFiles" Except="HordeToolbox.exe" With="#ToolboxBulkFiles"/>
+			<Copy Files="#ToolboxBulkFiles" From="$(ServerToolsStagingDir_HordeToolbox)" To="$(StagingDir)/Toolbox-Bulk"/>
 
-			<Spawn Exe="$(WixDir)/heat.exe" Arguments="dir &quot;$(StagingDir)/TrayApp-Bulk&quot; -cg HordeTrayAppFiles -scom -sreg -gg -dr TrayAppInstallDir -srd -var var.SourceDir -out &quot;$(InstallerDir)/HordeTrayAppFiles.wxs&quot;"/>
-			<Spawn Exe="$(WixDir)/candle.exe" Arguments="$(CommonArgs) HordeTrayAppFiles.wxs -dSourceDir=&quot;$(StagingDir)/TrayApp-Bulk&quot; -out &quot;obj/HordeTrayAppFiles.wixobj&quot;" WorkingDir="$(InstallerDir)"/>
+			<Spawn Exe="$(WixDir)/heat.exe" Arguments="dir &quot;$(StagingDir)/Toolbox-Bulk&quot; -cg HordeToolboxFiles -scom -sreg -gg -dr ToolboxInstallDir -srd -var var.SourceDir -out &quot;$(InstallerDir)/HordeToolboxFiles.wxs&quot;"/>
+			<Spawn Exe="$(WixDir)/candle.exe" Arguments="$(CommonArgs) HordeToolboxFiles.wxs -dSourceDir=&quot;$(StagingDir)/Toolbox-Bulk&quot; -out &quot;obj/HordeToolboxFiles.wixobj&quot;" WorkingDir="$(InstallerDir)"/>
 
 			<!-- Main component -->
-			<Spawn Exe="$(WixDir)/candle.exe" Arguments="$(CommonArgs) TrayAppInstaller.wxs -dStagingDir=&quot;$(ServerToolsStagingDir_HordeTrayApp)&quot; -out &quot;obj/HordeTrayApp.wixobj&quot; -ext WixUtilExtension" WorkingDir="$(InstallerDir)"/>
-			<Spawn Exe="$(WixDir)/light.exe" Arguments="-out &quot;bin/HordeTrayApp.msi&quot; -sw1076 -pdbout &quot;bin/HordeTrayApp.pdb&quot; &quot;obj/HordeTrayApp.wixobj&quot; &quot;obj/HordeTrayAppFiles.wixobj&quot; -ext WixUtilExtension -ext WixUIExtension" WorkingDir="$(InstallerDir)"/>
+			<Spawn Exe="$(WixDir)/candle.exe" Arguments="$(CommonArgs) ToolboxInstaller.wxs -dStagingDir=&quot;$(ServerToolsStagingDir_HordeToolbox)&quot; -out &quot;obj/HordeToolbox.wixobj&quot; -ext WixUtilExtension" WorkingDir="$(InstallerDir)"/>
+			<Spawn Exe="$(WixDir)/light.exe" Arguments="-out &quot;bin/HordeToolbox.msi&quot; -sw1076 -pdbout &quot;bin/HordeToolbox.pdb&quot; &quot;obj/HordeToolbox.wixobj&quot; &quot;obj/HordeToolboxFiles.wixobj&quot; -ext WixUtilExtension -ext WixUIExtension" WorkingDir="$(InstallerDir)"/>
 
 			<!-- We need to sign here as will be included in the server msi -->
-			<Sign Files="$(InstallerDir)/bin/HordeTrayApp.msi" Description="Horde Tray App" If="'$(SignOutputs)' == 'true'"/>
+			<Sign Files="$(InstallerDir)/bin/HordeToolbox.msi" Description="Horde Toolbox" If="'$(SignOutputs)' == 'true'"/>
 			
 			<!-- Copy files to the output directory -->
-			<Copy Files="HordeTrayApp.msi" From="$(InstallerDir)/bin" To="$(ServerToolsStagingDir_HordeTrayAppMsi)"/>
-			<Tag Files="$(ServerToolsStagingDir_HordeTrayAppMsi)/..." With="#HordeTrayAppMsiFiles"/>
+			<Copy Files="HordeToolbox.msi" From="$(InstallerDir)/bin" To="$(ServerToolsStagingDir_HordeToolboxMsi)"/>
+			<Tag Files="$(ServerToolsStagingDir_HordeToolboxMsi)/..." With="#HordeToolboxMsiFiles"/>
 		</Node>
 		<Node Name="Build P4VUtils" Produces="#P4VUtilsFiles">
 			<Delete Files="$(ServerToolsStagingDir_P4VUtils)/..."/>
@@ -273,7 +273,7 @@
 
 			<!-- Create the agent file list -->
 			<Delete Files="$(StagingDir)/Agent-Bulk/..."/>
-			<Tag Files="#HordeAgentWin64Files" Except="HordeAgent.exe;Horde.Agent.TrayApp.exe" With="#AgentFiles"/>
+			<Tag Files="#HordeAgentWin64Files" Except="HordeAgent.exe;HordeToolbox.exe" With="#AgentFiles"/>
 			<Copy Files="#AgentFiles" From="$(ServerToolsStagingDir_HordeAgentWin64)" To="$(StagingDir)/Agent-Bulk/Agent"/>
 
 			<Spawn Exe="$(WixDir)/heat.exe" Arguments="dir &quot;$(StagingDir)/Agent-Bulk&quot; -cg HordeAgentFiles -scom -sreg -gg -dr InstallDir -srd -var var.SourceDir -out &quot;$(InstallerDir)/HordeAgentFiles.wxs&quot;"/>
@@ -293,7 +293,7 @@
 		</Node>
 
 		<ForEach Name="ServerArch" Values="$(ServerArchList)">
-			<Node Name="Stage Horde Server Installer ($(ServerArch))" Requires="#DashboardOutputFiles;#HordeCmdFiles;#HordeAgentFiles;#HordeAgentMsiFiles;#HordeTrayAppFiles;#HordeTrayAppMsiFiles;#P4VUtilsFiles;#UnrealGameSyncFiles;#UnrealGameSyncCmdFiles;#UnrealGameSyncMsiFiles;" Produces="#ServerWithToolsFiles_$(ServerArch)">
+			<Node Name="Stage Horde Server Installer ($(ServerArch))" Requires="#DashboardOutputFiles;#HordeCmdFiles;#HordeAgentFiles;#HordeAgentMsiFiles;#HordeToolboxFiles;#HordeToolboxMsiFiles;#P4VUtilsFiles;#UnrealGameSyncFiles;#UnrealGameSyncCmdFiles;#UnrealGameSyncMsiFiles;" Produces="#ServerWithToolsFiles_$(ServerArch)">
 		
 				<!-- Stage the server -->
 				<Delete Files="$(StagingDir)/Server/$(ServerArch)/..."/>
@@ -306,9 +306,9 @@
 				<DotNet Arguments="run --project &quot;$(RootDir)/Engine/Source/Programs/Horde/HordeServer/HordeServer.csproj&quot; -- generate tooldata -id=horde-cmd -name=&quot;Horde (Command-Line)&quot; -description=&quot;Update package for the Horde command-line tool.&quot; -version=&quot;$(InformationalVersion)&quot; -inputdir=&quot;$(ServerToolsStagingDir_HordeCmd)&quot; -serverdir=&quot;$(StagingDir)/Server/$(ServerArch)&quot; -public -showinugs -showinlauncher -Category=Internal" DotNetPath="$(DotNetPath)" />
 				<DotNet Arguments="run --project &quot;$(RootDir)/Engine/Source/Programs/Horde/HordeServer/HordeServer.csproj&quot; -- generate tooldata -id=horde-agent -name=&quot;Horde Agent (Cross-Platform)&quot; -description=&quot;Cross-platform Horde Agent distribution. Allows hosts to take on work from the Horde Server.&quot; -version=&quot;$(InformationalVersion)&quot; -inputdir=&quot;$(ServerToolsStagingDir_HordeAgent)&quot; -serverdir=&quot;$(StagingDir)/Server/$(ServerArch)&quot; -public" DotNetPath="$(DotNetPath)" />
 				<DotNet Arguments="run --project &quot;$(RootDir)/Engine/Source/Programs/Horde/HordeServer/HordeServer.csproj&quot; -- generate tooldata -id=horde-agent-win64 -name=&quot;Horde Agent (Windows)&quot; -description=&quot;Update package for older installations of Horde Agent through the MSI installer.&quot; -version=&quot;$(InformationalVersion)&quot; -inputdir=&quot;$(ServerToolsStagingDir_HordeAgentWin64)&quot; -serverdir=&quot;$(StagingDir)/Server/$(ServerArch)&quot; -public -showindashboard=false -category=Internal" DotNetPath="$(DotNetPath)" />
-				<DotNet Arguments="run --project &quot;$(RootDir)/Engine/Source/Programs/Horde/HordeServer/HordeServer.csproj&quot; -- generate tooldata -id=horde-agent-win64-installer -name=&quot;Horde Agent (Windows Installer)&quot; -description=&quot;Installer for the Horde Agent on Windows, with tray application.&quot; -version=&quot;$(InformationalVersion)&quot; -inputdir=&quot;$(ServerToolsStagingDir_HordeAgentMsi)&quot; -serverdir=&quot;$(StagingDir)/Server/$(ServerArch)&quot; -public" DotNetPath="$(DotNetPath)" />
-				<DotNet Arguments="run --project &quot;$(RootDir)/Engine/Source/Programs/Horde/HordeServer/HordeServer.csproj&quot; -- generate tooldata -id=horde-tray-app -name=&quot;Horde Tray App&quot; -description=&quot;Update package for the Horde Tray App.&quot; -version=&quot;$(InformationalVersion)&quot; -inputdir=&quot;$(ServerToolsStagingDir_HordeTrayApp)&quot; -serverdir=&quot;$(StagingDir)/Server/$(ServerArch)&quot; -public -showinugs -Category=Internal" DotNetPath="$(DotNetPath)" />
-				<DotNet Arguments="run --project &quot;$(RootDir)/Engine/Source/Programs/Horde/HordeServer/HordeServer.csproj&quot; -- generate tooldata -id=horde-tray-app-msi -name=&quot;Horde Tray App (Windows Installer)&quot; -description=&quot;Windows installer for the Horde Tray App.&quot; -version=&quot;$(InformationalVersion)&quot; -inputdir=&quot;$(ServerToolsStagingDir_HordeTrayAppMsi)&quot; -serverdir=&quot;$(StagingDir)/Server/$(ServerArch)&quot; -public -showinugs" DotNetPath="$(DotNetPath)" />
+				<DotNet Arguments="run --project &quot;$(RootDir)/Engine/Source/Programs/Horde/HordeServer/HordeServer.csproj&quot; -- generate tooldata -id=horde-agent-win64-installer -name=&quot;Horde Agent (Windows Installer)&quot; -description=&quot;Installer for the Horde Agent on Windows, with toolbox application.&quot; -version=&quot;$(InformationalVersion)&quot; -inputdir=&quot;$(ServerToolsStagingDir_HordeAgentMsi)&quot; -serverdir=&quot;$(StagingDir)/Server/$(ServerArch)&quot; -public" DotNetPath="$(DotNetPath)" />
+				<DotNet Arguments="run --project &quot;$(RootDir)/Engine/Source/Programs/Horde/HordeServer/HordeServer.csproj&quot; -- generate tooldata -id=horde-toolbox -name=&quot;Horde Toolbox&quot; -description=&quot;Update package for the Horde Toolbox.&quot; -version=&quot;$(InformationalVersion)&quot; -inputdir=&quot;$(ServerToolsStagingDir_HordeToolbox)&quot; -serverdir=&quot;$(StagingDir)/Server/$(ServerArch)&quot; -public -showinugs -Category=Internal" DotNetPath="$(DotNetPath)" />
+				<DotNet Arguments="run --project &quot;$(RootDir)/Engine/Source/Programs/Horde/HordeServer/HordeServer.csproj&quot; -- generate tooldata -id=horde-toolbox-msi -name=&quot;Horde Toolbox (Windows Installer)&quot; -description=&quot;Windows installer for the Horde Toolbox.&quot; -version=&quot;$(InformationalVersion)&quot; -inputdir=&quot;$(ServerToolsStagingDir_HordeToolboxMsi)&quot; -serverdir=&quot;$(StagingDir)/Server/$(ServerArch)&quot; -public -showinugs" DotNetPath="$(DotNetPath)" />
 				<DotNet Arguments="run --project &quot;$(RootDir)/Engine/Source/Programs/Horde/HordeServer/HordeServer.csproj&quot; -- generate tooldata -id=p4v-utils -name=&quot;P4V Extensions&quot; -description=&quot;Update package for the P4VUtils tool.&quot; -version=&quot;$(InformationalVersion)&quot; -inputdir=&quot;$(ServerToolsStagingDir_P4VUtils)&quot; -serverdir=&quot;$(StagingDir)/Server/$(ServerArch)&quot; -public -showinugs -showinlauncher -category=Internal" DotNetPath="$(DotNetPath)" />
 				<DotNet Arguments="run --project &quot;$(RootDir)/Engine/Source/Programs/Horde/HordeServer/HordeServer.csproj&quot; -- generate tooldata -id=ugs-cmd -name=&quot;UnrealGameSync (Command-Line)&quot; -description=&quot;Command-line version of UnrealGameSync. Supports Windows, MacOS and Linux.&quot; -version=&quot;$(InformationalVersion)&quot; -inputdir=&quot;$(ServerToolsStagingDir_UnrealGameSyncCmd)&quot; -serverdir=&quot;$(StagingDir)/Server/$(ServerArch)&quot; -public -showinlauncher" DotNetPath="$(DotNetPath)" />
 				<DotNet Arguments="run --project &quot;$(RootDir)/Engine/Source/Programs/Horde/HordeServer/HordeServer.csproj&quot; -- generate tooldata -id=ugs-msi -name=&quot;UnrealGameSync (Windows Installer)&quot; -description=&quot;Installer for UnrealGameSync.&quot; -version=&quot;$(InformationalVersion)&quot; -inputdir=&quot;$(ServerToolsStagingDir_UnrealGameSyncMsi)&quot; -serverdir=&quot;$(StagingDir)/Server/$(ServerArch)&quot; -public -showinlauncher" DotNetPath="$(DotNetPath)" />
diff --git a/Engine/Source/Programs/Horde/HordeToolbox/App.axaml.cs b/Engine/Source/Programs/Horde/HordeToolbox/App.axaml.cs
new file mode 100644
--- /dev/null
+++ b/Engine/Source/Programs/Horde/HordeToolbox/App.axaml.cs
@@ -0,0 +1,226 @@
+// Copyright Epic Games, Inc. All Rights Reserved.
+
+using System.Diagnostics;
+using Avalonia;
+using Avalonia.Controls;
+using Avalonia.Controls.ApplicationLifetimes;
+using Avalonia.Markup.Xaml;
+using Avalonia.Threading;
+using EpicGames.Core;
+using HordeToolbox.Plugins.HordeAgent;
+using HordeToolbox.Plugins.HordeProxy;
+using Microsoft.Extensions.DependencyInjection;
+
+namespace HordeToolbox
+{
+	/// <summary>
+	/// Main application class
+	/// </summary>
+	public sealed partial class App : Application, ITrayAppHost
+	{
+		readonly ServiceProvider _serviceProvider;
+
+		WindowIcon? _normalIcon;
+		WindowIcon? _busyIcon;
+		WindowIcon? _pausedIcon;
+		WindowIcon? _errorIcon;
+
+		SettingsWindow? _settingsWindow;
+
+		/// <summary>
+		/// Constructor
+		/// </summary>
+		public App()
+		{
+			ServiceCollection serviceCollection = new ServiceCollection();
+			serviceCollection.AddLogging(builder => builder.AddEpicDefault());
+			serviceCollection.AddSingleton<IHordeClientProvider, HordeClientProvider>();
+			serviceCollection.AddSingleton<ToolCatalog>();
+			serviceCollection.AddSingleton<IToolCatalog, ToolCatalog>(sp => sp.GetRequiredService<ToolCatalog>());
+			serviceCollection.AddSingleton<ITrayAppHost>(this);
+
+			// Configure plugins
+			serviceCollection.AddSingleton<ITrayAppPlugin, HordeAgentPlugin>();
+			serviceCollection.AddSingleton<ITrayAppPlugin, HordeProxyPlugin>();
+
+			_serviceProvider = serviceCollection.BuildServiceProvider();
+		}
+
+		/// <inheritdoc/>
+		public override void Initialize()
+		{
+			AvaloniaXamlLoader.Load(this);
+
+			_normalIcon = (WindowIcon)Resources["StatusNormal"]!;
+			_busyIcon = (WindowIcon)Resources["StatusBusy"]!;
+			_pausedIcon = (WindowIcon)Resources["StatusPaused"]!;
+			_errorIcon = (WindowIcon)Resources["StatusError"]!;
+
+			UpdateMenu();
+
+			ToolCatalog toolCatalog = _serviceProvider.GetRequiredService<ToolCatalog>();
+			toolCatalog.OnItemsChanged += UpdateMenu;
+			toolCatalog.Start();
+		}
+
+		private void UpdateMenu()
+		{
+			NativeMenuItem settingsMenuItem = new NativeMenuItem("Settings...");
+			settingsMenuItem.Click += TrayIcon_Settings;
+
+			NativeMenuItem exitMenuItem = new NativeMenuItem("Exit");
+			exitMenuItem.Click += TrayIcon_Exit;
+
+			NativeMenu contextMenu = new NativeMenu();
+
+			TrayIcon trayIcon = TrayIcon.GetIcons(this)![0];
+			trayIcon.Menu = contextMenu;
+
+			foreach (ITrayAppPlugin plugin in _serviceProvider.GetServices<ITrayAppPlugin>())
+			{
+				plugin.PopulateContextMenu(contextMenu);
+			}
+
+			int numItems = contextMenu.Items.Count;
+			if (numItems > 0)
+			{
+				contextMenu.Items.Add(new NativeMenuItemSeparator());
+				numItems++;
+			}
+
+			IToolCatalog toolCatalog = _serviceProvider.GetRequiredService<IToolCatalog>();
+			foreach (IToolCatalogItem item in toolCatalog.Items)
+			{
+				CurrentToolDeploymentInfo? current = item.Current;
+				if (current != null)
+				{
+					ToolConfig? toolConfig = current.Config;
+					if (toolConfig != null && toolConfig.PopupMenu != null)
+					{
+						AddMenuItems(item, toolConfig.PopupMenu, contextMenu);
+					}
+				}
+			}
+
+			if (contextMenu.Items.Count > numItems)
+			{
+				contextMenu.Items.Add(new NativeMenuItemSeparator());
+			}
+			contextMenu.Items.Add(settingsMenuItem);
+			contextMenu.Items.Add(exitMenuItem);
+		}
+
+		static void AddMenuItems(IToolCatalogItem item, ToolMenuItem toolMenuItem, NativeMenu menu)
+		{
+			if (!String.IsNullOrEmpty(toolMenuItem.Label))
+			{
+				NativeMenuItem menuItem = new NativeMenuItem(toolMenuItem.Label);
+				if (toolMenuItem.Children != null && toolMenuItem.Children.Count > 0)
+				{
+					NativeMenu subMenu = new NativeMenu();
+					foreach (ToolMenuItem child in toolMenuItem.Children)
+					{
+						AddMenuItems(item, child, subMenu);
+					}
+					menuItem.Menu = subMenu;
+				}
+				else
+				{
+					if (toolMenuItem.Command != null)
+					{
+						menuItem.Click += (_, _) => TrayIcon_RunCommand(item, toolMenuItem.Command);
+					}
+				}
+				menu.Items.Add(menuItem);
+			}
+		}
+
+		static void TrayIcon_RunCommand(IToolCatalogItem item, ToolCommand command)
+		{
+			ProcessStartInfo startInfo = new ProcessStartInfo(command.FileName);
+			if (command.Arguments != null)
+			{
+				foreach (string argument in command.Arguments)
+				{
+					startInfo.ArgumentList.Add(argument);
+				}
+			}
+			startInfo.WorkingDirectory = item.Current!.Dir.FullName;
+			Process.Start(startInfo);
+		}
+
+		private void TrayIcon_Click(object? sender, EventArgs e)
+		{
+			OpenSettings();
+		}
+
+		private void TrayIcon_Settings(object? sender, EventArgs e)
+		{
+			OpenSettings();
+		}
+
+		private void OpenSettings()
+		{
+			if (_settingsWindow != null)
+			{
+				_settingsWindow.BringIntoView();
+				_settingsWindow.Show();
+			}
+			else
+			{
+				_settingsWindow = new SettingsWindow(_serviceProvider);
+				_settingsWindow.Closed += SettingsWindow_Closed;
+				_settingsWindow.Show();
+			}
+		}
+
+		private void SettingsWindow_Closed(object? sender, EventArgs e)
+		{
+			_settingsWindow = null;
+		}
+
+		private void TrayIcon_Exit(object? sender, EventArgs e)
+		{
+			((IClassicDesktopStyleApplicationLifetime)ApplicationLifetime!).Shutdown();
+		}
+
+		/// <inheritdoc/>
+		public void UpdateStatus()
+		{
+			Dispatcher.UIThread.Post(() => UpdateStatusMainThread());
+		}
+
+		private void UpdateStatusMainThread()
+		{
+			TrayAppPluginState state = TrayAppPluginState.Undefined;
+			List<string> messages = new List<string>();
+
+			foreach (ITrayAppPlugin plugin in _serviceProvider.GetServices<ITrayAppPlugin>())
+			{
+				TrayAppPluginStatus status = plugin.GetStatus();
+				if (status.State >= state)
+				{
+					if (status.State > state)
+					{
+						messages.Clear();
+						state = status.State;
+					}
+					if (status.Message != null)
+					{
+						messages.Add(status.Message);
+					}
+				}
+			}
+
+			TrayIcon trayIcon = TrayIcon.GetIcons(this)![0];
+			trayIcon!.Icon = state switch
+			{
+				TrayAppPluginState.Busy => _busyIcon,
+				TrayAppPluginState.Paused => _pausedIcon,
+				TrayAppPluginState.Error => _errorIcon,
+				_ => _normalIcon
+			};
+			trayIcon.ToolTipText = String.Join("\n", messages);
+		}
+	}
+}
diff --git a/Engine/Source/Programs/Horde/HordeToolbox/GeneralSettingsPage.axaml.cs b/Engine/Source/Programs/Horde/HordeToolbox/GeneralSettingsPage.axaml.cs
new file mode 100644
--- /dev/null
+++ b/Engine/Source/Programs/Horde/HordeToolbox/GeneralSettingsPage.axaml.cs
@@ -0,0 +1,30 @@
+// Copyright Epic Games, Inc. All Rights Reserved.
+
+using Avalonia.Controls;
+
+namespace HordeToolbox
+{
+	partial class GeneralSettingsPage : UserControl, IDisposable
+	{
+		readonly GeneralSettingsViewModel _viewModel;
+
+		public GeneralSettingsPage()
+			: this(SettingsContext.Default)
+		{
+		}
+
+		public GeneralSettingsPage(SettingsContext context)
+		{
+			InitializeComponent();
+
+			_viewModel = new GeneralSettingsViewModel(context);
+
+			DataContext = _viewModel;
+		}
+
+		public void Dispose()
+		{
+			_viewModel.Dispose();
+		}
+	}
+}
diff --git a/Engine/Source/Programs/Horde/HordeToolbox/GeneralSettingsViewModel.cs b/Engine/Source/Programs/Horde/HordeToolbox/GeneralSettingsViewModel.cs
new file mode 100644
--- /dev/null
+++ b/Engine/Source/Programs/Horde/HordeToolbox/GeneralSettingsViewModel.cs
@@ -0,0 +1,164 @@
+// Copyright Epic Games, Inc. All Rights Reserved.
+
+using System.Diagnostics;
+using Avalonia.Controls;
+using CommunityToolkit.Mvvm.ComponentModel;
+using EpicGames.Horde;
+using FluentAvalonia.UI.Controls;
+using Microsoft.Extensions.DependencyInjection;
+
+namespace HordeToolbox
+{
+	partial class GeneralSettingsViewModel : ObservableObject, IDisposable
+	{
+		record StateRecord(Uri ServerUrl, bool Connected);
+
+		readonly SettingsContext _context;
+		readonly IHordeClientProvider _hordeClientProvider;
+
+		[ObservableProperty]
+		string _serverUrl = String.Empty;
+
+		[ObservableProperty]
+		string _serverStatus = String.Empty;
+
+		[ObservableProperty]
+		bool? _isLoginEnabled;
+
+		public ToolCatalogViewModel ToolCatalog { get; }
+
+		public GeneralSettingsViewModel()
+			: this(SettingsContext.Default)
+		{
+		}
+
+		public GeneralSettingsViewModel(SettingsContext context)
+		{
+			_context = context;
+
+			_hordeClientProvider = context.ServiceProvider.GetRequiredService<IHordeClientProvider>();
+			_hordeClientProvider.OnStateChanged += OnStateChangedAsync;
+			_hordeClientProvider.OnAccessTokenStateChanged += OnStateChangedAsync;
+
+			ToolCatalog = new ToolCatalogViewModel(context.ServiceProvider.GetService<IToolCatalog>());
+
+			OnStateChangedAsync();
+		}
+
+		public void Dispose()
+		{
+			_hordeClientProvider.OnStateChanged -= OnStateChangedAsync;
+			_hordeClientProvider.OnAccessTokenStateChanged -= OnStateChangedAsync;
+		}
+
+		async void OnStateChangedAsync()
+		{
+			using IHordeClientRef? hordeClient = await _hordeClientProvider.GetClientRefAsync();
+			if (hordeClient == null)
+			{
+				ServerUrl = "No server configured";
+				ServerStatus = "Status unavailable";
+				IsLoginEnabled = false;
+			}
+			else
+			{
+				ServerUrl = hordeClient.Client.ServerUrl.ToString();
+				ServerStatus = hordeClient.Client.HasValidAccessToken() ? "Connected" : "Session Expired";
+				IsLoginEnabled = !hordeClient.Client.HasValidAccessToken();
+			}
+		}
+
+		public void OpenServer()
+		{
+			OpenBrowser(ServerUrl);
+		}
+
+		public async Task LoginAsync()
+		{
+			using IHordeClientRef? hordeClientRef = await _hordeClientProvider.GetClientRefAsync();
+			if (hordeClientRef != null)
+			{
+				await hordeClientRef.Client.LoginAsync(true, CancellationToken.None);
+			}
+		}
+
+		public async Task ConfigureServerAsync()
+		{
+			string serverUrl = ServerUrl;
+
+			string? errorMessage = null;
+			for (; ; )
+			{
+				TextBox serverUrlTextBox = new TextBox();
+				serverUrlTextBox.Text = serverUrl;
+				serverUrlTextBox.SelectionStart = serverUrl.Length;
+				serverUrlTextBox.SelectionEnd = serverUrl.Length;
+
+				Control content = serverUrlTextBox;
+				if (!String.IsNullOrEmpty(errorMessage))
+				{
+					TextBlock errorBlock = new TextBlock();
+					errorBlock.Text = errorMessage;
+
+					StackPanel stackPanel = new StackPanel() { Spacing = 10 };
+					stackPanel.Children.Add(serverUrlTextBox);
+					stackPanel.Children.Add(errorBlock);
+
+					content = stackPanel;
+				}
+
+				ContentDialog dialog = new ContentDialog()
+				{
+					Title = "Connect to Server",
+					Content = content,
+					PrimaryButtonText = "Connect",
+					CloseButtonText = "Cancel"
+				};
+
+				ContentDialogResult result = await dialog.ShowAsync(_context.SettingsWindow);
+				if (result == ContentDialogResult.None)
+				{
+					return;
+				}
+
+				serverUrl = serverUrlTextBox.Text;
+
+				try
+				{
+					HordeOptions.SetDefaultServerUrl(new Uri(serverUrl));
+
+					await _hordeClientProvider.RecreateAsync();
+
+					using IHordeClientRef? clientRef = await _hordeClientProvider.GetClientRefAsync();
+					if (clientRef != null)
+					{
+						await clientRef.Client.LoginAsync(true, CancellationToken.None);
+					}
+
+					return;
+				}
+				catch (Exception ex)
+				{
+					errorMessage = ex.ToString();
+				}
+			}
+		}
+
+		static void OpenBrowser(string url)
+		{
+			if (OperatingSystem.IsWindows())
+			{
+				string escapedUrl = url.Replace("&", "^&", StringComparison.Ordinal);
+				Process.Start(new ProcessStartInfo("cmd", $"/c start {escapedUrl}") { CreateNoWindow = true });
+			}
+			else if (OperatingSystem.IsLinux())
+			{
+				Process.Start("xdg-open", url);
+			}
+			else if (OperatingSystem.IsMacOS())
+			{
+				Process.Start("open", url);
+			}
+		}
+	}
+}
diff --git a/Engine/Source/Programs/Horde/HordeToolbox/HordeClientProvider.cs b/Engine/Source/Programs/Horde/HordeToolbox/HordeClientProvider.cs
new file mode 100644
--- /dev/null
+++ b/Engine/Source/Programs/Horde/HordeToolbox/HordeClientProvider.cs
@@ -0,0 +1,204 @@
+// Copyright Epic Games, Inc. All Rights Reserved.
+
+using System.Diagnostics;
+using EpicGames.Horde;
+using Microsoft.Extensions.DependencyInjection;
+using Microsoft.Extensions.Logging;
+
+namespace HordeToolbox
+{
+	/// <summary>
+	/// Implements a mechanism for creating scoped HordeClient instances, and modifying configuration settings
+	/// </summary>
+	class HordeClientProvider : IHordeClientProvider, IAsyncDisposable
+	{
+		class HordeClientLifetime : IAsyncDisposable
+		{
+			readonly TaskCompletionSource _refZeroTcs = new TaskCompletionSource();
+
+			readonly Task<IHordeClient?> _clientTask;
+			int _refCount;
+			bool _disposed;
+
+			public Task<IHordeClient?> ClientTask => _clientTask;
+
+			public HordeClientLifetime(Task<IHordeClient?> clientTask)
+			{
+				_clientTask = clientTask;
+				_refCount = 1;
+			}
+
+			public async ValueTask DisposeAsync()
+			{
+				if (!_disposed)
+				{
+					Release();
+					await _refZeroTcs.Task;
+					_disposed = true;
+				}
+			}
+
+			public void AddRef()
+			{
+				int refCount = Interlocked.Increment(ref _refCount);
+				Debug.Assert(refCount > 1);
+			}
+
+			public void Release()
+			{
+				if (Interlocked.Decrement(ref _refCount) == 0)
+				{
+					_refZeroTcs.SetResult();
+				}
+			}
+		}
+
+		class HordeClientRef : IHordeClientRef
+		{
+			IHordeClient _client;
+			HordeClientLifetime? _lifetime;
+
+			public IHordeClient Client
+				=> _client ?? throw new ObjectDisposedException(null);
+
+			public HordeClientRef(IHordeClient client, HordeClientLifetime lifetime)
+			{
+				_client = client;
+				_lifetime = lifetime;
+			}
+
+			public void Dispose()
+			{
+				_lifetime?.Release();
+				_lifetime = null;
+				_client = null!;
+			}
+		}
+
+		readonly object _lockObject = new object();
+		readonly ILoggerFactory _loggerFactory;
+		ServiceProvider? _serviceProvider;
+		IHordeClient? _hordeClient;
+		HordeClientLifetime _lifetime;
+		readonly ILogger _logger;
+
+		/// <summary>
+		/// Event signalled whenever the connection state changes
+		/// </summary>
+		public event Action? OnStateChanged;
+
+		/// <summary>
+		/// Event signalled whenever the access token state changes
+		/// </summary>
+		public event Action? OnAccessTokenStateChanged;
+
+		/// <summary>
+		/// Constructor
+		/// </summary>
+		public HordeClientProvider(ILoggerFactory loggerFactory, ILogger<HordeClientProvider> logger)
+		{
+			_loggerFactory = loggerFactory;
+			_logger = logger;
+
+			CreateServices();
+
+			_lifetime = new HordeClientLifetime(Task.FromResult(_hordeClient));
+		}
+
+		/// <inheritdoc/>
+		public async ValueTask DisposeAsync()
+		{
+			await _lifetime.DisposeAsync();
+			await DestroyServicesAsync();
+		}
+
+		/// <inheritdoc/>
+		public async Task<IHordeClientRef?> GetClientRefAsync()
+		{
+			HordeClientLifetime? lifetime = null;
+			try
+			{
+				lock (_lockObject)
+				{
+					lifetime = _lifetime;
+					lifetime.AddRef();
+				}
+
+				IHordeClient? client = await lifetime.ClientTask;
+				if (client == null)
+				{
+					lifetime.Release();
+					return null;
+				}
+
+				return new HordeClientRef(client, lifetime);
+			}
+			catch
+			{
+				lifetime?.Release();
+				throw;
+			}
+		}
+
+		/// <inheritdoc/>
+		public async Task RecreateAsync()
+		{
+			lock (_lockObject)
+			{
+				HordeClientLifetime prevLifetime = _lifetime;
+				Task<IHordeClient?> clientTask = Task.Run(async () =>
+				{
+					await prevLifetime.DisposeAsync();
+					await DestroyServicesAsync();
+					CreateServices();
+					return _hordeClient;
+				});
+				_lifetime = new HordeClientLifetime(clientTask);
+			}
+			await _lifetime.ClientTask;
+		}
+
+		void OnAccessTokenStateChangedInternal()
+		{
+			OnAccessTokenStateChanged?.Invoke();
+		}
+
+		void CreateServices()
+		{
+			Debug.Assert(_serviceProvider == null);
+			Debug.Assert(_hordeClient == null);
+
+			ServiceCollection serviceCollection = new ServiceCollection();
+			serviceCollection.AddHorde(options => options.AllowAuthPrompt = false);
+			serviceCollection.AddSingleton<ILoggerFactory>(_loggerFactory);
+			serviceCollection.AddSingleton(typeof(ILogger<>), typeof(Logger<>));
+			_serviceProvider = serviceCollection.BuildServiceProvider();
+
+			try
+			{
+				_hordeClient = _serviceProvider.GetRequiredService<IHordeClient>();
+				_hordeClient.OnAccessTokenStateChanged += OnAccessTokenStateChangedInternal;
+			}
+			catch (Exception ex)
+			{
+				_logger.LogWarning(ex, "Horde client could not be created: {Message}", ex.Message);
+			}
+
+			OnStateChanged?.Invoke();
+		}
+
+		async ValueTask DestroyServicesAsync()
+		{
+			if (_hordeClient != null)
+			{
+				_hordeClient.OnAccessTokenStateChanged -= OnAccessTokenStateChangedInternal;
+				_hordeClient = null;
+			}
+			if (_serviceProvider != null)
+			{
+				await _serviceProvider.DisposeAsync();
+				_serviceProvider = null;
+			}
+		}
+	}
+}
diff --git a/Engine/Source/Programs/Horde/HordeToolbox/HordeToolbox.csproj b/Engine/Source/Programs/Horde/HordeToolbox/HordeToolbox.csproj
new file mode 100644
--- /dev/null
+++ b/Engine/Source/Programs/Horde/HordeToolbox/HordeToolbox.csproj
@@ -0,0 +1,78 @@
+﻿<Project Sdk="Microsoft.NET.Sdk">
+
+  <ItemGroup>
+    <FrameworkReference Include="Microsoft.AspNetCore.App" />
+  </ItemGroup>
+
+  <PropertyGroup>
+    <OutputType>WinExe</OutputType>
+    <TargetFramework>net8.0-windows</TargetFramework>
+    <Nullable>enable</Nullable>
+    <ImplicitUsings>enable</ImplicitUsings>
+    <NeutralLanguage>en-US</NeutralLanguage>
+    <Configurations>Debug;Release;Analyze</Configurations>
+	<AssemblyName>HordeToolbox</AssemblyName>
+    <AvaloniaUseCompiledBindingsByDefault>true</AvaloniaUseCompiledBindingsByDefault>
+    <NoWarn>IDE0005;CA1501</NoWarn>
+  </PropertyGroup>
+
+  <PropertyGroup Condition=" '$(Configuration)|$(Platform)' == 'Analyze|AnyCPU' ">
+    <RunAnalyzersDuringBuild>True</RunAnalyzersDuringBuild>
+    <TreatWarningsAsErrors>false</TreatWarningsAsErrors>
+    <EnforceCodeStyleInBuild>true</EnforceCodeStyleInBuild>
+    <GenerateDocumentationFile>true</GenerateDocumentationFile>
+    <DocumentationFile></DocumentationFile>
+  </PropertyGroup>
+
+  <ItemGroup>
+    <Compile Include="..\HordeAgent\IpcMessages.cs" Link="Plugins\HordeAgent\IpcMessages.cs" />
+  </ItemGroup>
+
+  <ItemGroup>
+    <ProjectReference Include="..\..\Shared\EpicGames.Core\EpicGames.Core.csproj" />
+    <ProjectReference Include="..\..\Shared\EpicGames.Horde\EpicGames.Horde.csproj" />
+  </ItemGroup>
+
+  <ItemGroup>
+    <Compile Update="GeneralSettingsPage.axaml.cs">
+      <DependentUpon>GeneralSettingsPage.axaml</DependentUpon>
+    </Compile>
+    <Compile Update="Plugins\HordeAgent\HordeAgentSettingsPage.axaml.cs">
+      <DependentUpon>HordeAgentSettingsPage.axaml</DependentUpon>
+    </Compile>
+    <Compile Update="Properties\Resources.Designer.cs">
+      <DesignTime>True</DesignTime>
+      <AutoGen>True</AutoGen>
+      <DependentUpon>Resources.resx</DependentUpon>
+    </Compile>
+    <Compile Update="HordeServerSettings.axaml.cs">
+      <DependentUpon>HordeServerSettings.axaml</DependentUpon>
+    </Compile>
+    <Compile Update="SettingsWindow.axaml.cs">
+      <DependentUpon>SettingsWindow.axaml</DependentUpon>
+    </Compile>
+  </ItemGroup>
+
+  <ItemGroup>
+    <EmbeddedResource Update="Properties\Resources.resx">
+      <Generator>ResXFileCodeGenerator</Generator>
+      <LastGenOutput>Resources.Designer.cs</LastGenOutput>
+    </EmbeddedResource>
+  </ItemGroup>
+  
+  <ItemGroup>
+    <AvaloniaResource Include="Resources\**" />
+  </ItemGroup>
+
+  <ItemGroup>
+    <PackageReference Include="Avalonia" Version="11.1.3" />
+    <PackageReference Include="Avalonia.Desktop" Version="11.1.3" />
+    <PackageReference Include="Avalonia.Themes.Fluent" Version="11.1.3" />
+    <PackageReference Include="Avalonia.Fonts.Inter" Version="11.1.3" />
+    <!--Condition below is needed to remove Avalonia.Diagnostics package from build output in Release configuration.-->
+    <PackageReference Condition="'$(Configuration)' == 'Debug'" Include="Avalonia.Diagnostics" Version="11.1.0" />
+    <PackageReference Include="FluentAvaloniaUI" Version="2.1.0" />
+    <PackageReference Include="CommunityToolkit.Mvvm" Version="8.3.1" />
+  </ItemGroup>
+
+</Project>
\ No newline at end of file
diff --git a/Engine/Source/Programs/Horde/HordeToolbox/IHordeClientProvider.cs b/Engine/Source/Programs/Horde/HordeToolbox/IHordeClientProvider.cs
new file mode 100644
--- /dev/null
+++ b/Engine/Source/Programs/Horde/HordeToolbox/IHordeClientProvider.cs
@@ -0,0 +1,43 @@
+// Copyright Epic Games, Inc. All Rights Reserved.
+
+using EpicGames.Horde;
+
+namespace HordeToolbox
+{
+	/// <summary>
+	/// Reference to a Horde client instance. Should be disposed of when done.
+	/// </summary>
+	interface IHordeClientRef : IDisposable
+	{
+		/// <summary>
+		/// The client instance
+		/// </summary>
+		IHordeClient Client { get; }
+	}
+
+	/// <summary>
+	/// Accessor for a <see cref="IHordeClient"/> instance which can be recreated in response to config changes. 
+	/// </summary>
+	interface IHordeClientProvider
+	{
+		/// <summary>
+		/// Event signalled whenever the connection state changes
+		/// </summary>
+		event Action? OnStateChanged;
+
+		/// <summary>
+		/// Event signalled whenever the access token state changes
+		/// </summary>
+		event Action? OnAccessTokenStateChanged;
+
+		/// <summary>
+		/// Gets a reference to the current client instance. These references should be kept as short as possible.
+		/// </summary>
+		Task<IHordeClientRef?> GetClientRefAsync();
+
+		/// <summary>
+		/// Recreates the client instance.
+		/// </summary>
+		Task RecreateAsync();
+	}
+}
diff --git a/Engine/Source/Programs/Horde/HordeToolbox/IToolCatalog.cs b/Engine/Source/Programs/Horde/HordeToolbox/IToolCatalog.cs
new file mode 100644
--- /dev/null
+++ b/Engine/Source/Programs/Horde/HordeToolbox/IToolCatalog.cs
@@ -0,0 +1,25 @@
+// Copyright Epic Games, Inc. All Rights Reserved.
+
+namespace HordeToolbox
+{
+	/// <summary>
+	/// Interface for the tool catalog
+	/// </summary>
+	public interface IToolCatalog
+	{
+		/// <summary>
+		/// Whether to auto-update tools to latest
+		/// </summary>
+		bool AutoUpdate { get; set; }
+
+		/// <summary>
+		/// List of available tools
+		/// </summary>
+		IReadOnlyList<IToolCatalogItem> Items { get; }
+
+		/// <summary>
+		/// Notification that some property of the items list has changed
+		/// </summary>
+		event Action? OnItemsChanged;
+	}
+}
diff --git a/Engine/Source/Programs/Horde/HordeToolbox/IToolCatalogItem.cs b/Engine/Source/Programs/Horde/HordeToolbox/IToolCatalogItem.cs
new file mode 100644
--- /dev/null
+++ b/Engine/Source/Programs/Horde/HordeToolbox/IToolCatalogItem.cs
@@ -0,0 +1,87 @@
+// Copyright Epic Games, Inc. All Rights Reserved.
+
+using EpicGames.Core;
+using EpicGames.Horde.Tools;
+
+namespace HordeToolbox
+{
+	/// <summary>
+	/// Describes a tool available to be installed
+	/// </summary>
+	public interface IToolCatalogItem
+	{
+		/// <summary>
+		/// Identifier for the tool
+		/// </summary>
+		ToolId Id { get; }
+
+		/// <summary>
+		/// Name of the tool
+		/// </summary>
+		string Name { get; }
+
+		/// <summary>
+		/// Description for the tool
+		/// </summary>
+		string Description { get; }
+
+		/// <summary>
+		/// Latest available version of the tool
+		/// </summary>
+		ToolDeploymentInfo? Latest { get; }
+
+		/// <summary>
+		/// The current tool deployment
+		/// </summary>
+		CurrentToolDeploymentInfo? Current { get; }
+
+		/// <summary>
+		/// Pending tool deployment, if an install or uninstall operation 
+		/// </summary>
+		PendingToolDeploymentInfo? Pending { get; }
+
+		/// <summary>
+		/// Callback for the item state changing
+		/// </summary>
+		event Action? OnItemChanged;
+
+		/// <summary>
+		/// Cancel any pending installation
+		/// </summary>
+		void Cancel();
+
+		/// <summary>
+		/// Start installing this tool or updating it to the latest version
+		/// </summary>
+		void Install();
+
+		/// <summary>
+		/// Start uninstalling this tool
+		/// </summary>
+		void Uninstall();
+	}
+
+	/// <summary>
+	/// Information about a particular tool deployment
+	/// </summary>
+	/// <param name="Id">Identifier for the deployment</param>
+	/// <param name="Version">Version number for this deployment</param>
+	public record class ToolDeploymentInfo(ToolDeploymentId Id, string Version);
+
+	/// <summary>
+	/// Information about an installed tool deployment
+	/// </summary>
+	/// <param name="Id">Identifier for the deployment</param>
+	/// <param name="Version">Version number for this deployment</param>
+	/// <param name="Dir">Directory containing the tool</param>
+	/// <param name="Config">Configuration for the tool</param>
+	public record class CurrentToolDeploymentInfo(ToolDeploymentId Id, string Version, DirectoryReference Dir, ToolConfig Config) : ToolDeploymentInfo(Id, Version);
+
+	/// <summary>
+	/// Information about a pending tool installation
+	/// </summary>
+	/// <param name="Failed">Whether the installation failed</param>
+	/// <param name="Message">Current status message</param>
+	/// <param name="Deployment">Information about the pending deployment</param>
+	public record class PendingToolDeploymentInfo(bool Failed, string? Message, ToolDeploymentInfo? Deployment);
+}
diff --git a/Engine/Source/Programs/Horde/HordeToolbox/IToolboxPlugin.cs b/Engine/Source/Programs/Horde/HordeToolbox/IToolboxPlugin.cs
new file mode 100644
--- /dev/null
+++ b/Engine/Source/Programs/Horde/HordeToolbox/IToolboxPlugin.cs
@@ -0,0 +1,55 @@
+// Copyright Epic Games, Inc. All Rights Reserved.
+
+using Avalonia.Controls;
+using FluentAvalonia.UI.Controls;
+
+namespace HordeToolbox
+{
+	enum TrayAppPluginState
+	{
+		Undefined,
+		Ok,
+		Busy,
+		Paused,
+		Error
+	}
+
+	/// <summary>
+	/// Reported status of a plugin
+	/// </summary>
+	record class TrayAppPluginStatus(TrayAppPluginState State, string? Message)
+	{
+		public static TrayAppPluginStatus Default { get; } = new TrayAppPluginStatus(TrayAppPluginState.Undefined, null);
+	}
+
+	/// <summary>
+	/// Plugin for the tray app
+	/// </summary>
+	interface ITrayAppPlugin : IAsyncDisposable
+	{
+		/// <summary>
+		/// Name of the plugin to show in the settings page
+		/// </summary>
+		string Name { get; }
+
+		/// <summary>
+		/// Icon to show on the settings page
+		/// </summary>
+		IconSource Icon { get; }
+
+		/// <summary>
+		/// Get the current status of this plugin
+		/// </summary>
+		TrayAppPluginStatus GetStatus();
+
+		/// <summary>
+		/// Create a settings page for this plugin
+		/// </summary>
+		Control CreateSettingsPage(SettingsContext context);
+
+		/// <summary>
+		/// Allow the plugin to customize the tray icon context menu
+		/// </summary>
+		void PopulateContextMenu(NativeMenu contextMenu);
+	}
+}
diff --git a/Engine/Source/Programs/Horde/HordeToolbox/IToolboxPluginHost.cs b/Engine/Source/Programs/Horde/HordeToolbox/IToolboxPluginHost.cs
new file mode 100644
--- /dev/null
+++ b/Engine/Source/Programs/Horde/HordeToolbox/IToolboxPluginHost.cs
@@ -0,0 +1,15 @@
+// Copyright Epic Games, Inc. All Rights Reserved.
+
+namespace HordeToolbox
+{
+	/// <summary>
+	/// Interface for the tray app host application
+	/// </summary>
+	interface ITrayAppHost
+	{
+		/// <summary>
+		/// Notifies the host that a status change has ocurred
+		/// </summary>
+		void UpdateStatus();
+	}
+}
diff --git a/Engine/Source/Programs/Horde/HordeToolbox/JsonConfig.cs b/Engine/Source/Programs/Horde/HordeToolbox/JsonConfig.cs
new file mode 100644
--- /dev/null
+++ b/Engine/Source/Programs/Horde/HordeToolbox/JsonConfig.cs
@@ -0,0 +1,79 @@
+// Copyright Epic Games, Inc. All Rights Reserved.
+
+using System;
+using System.Collections.Generic;
+using System.Linq;
+using System.Text;
+using System.Text.Json;
+using System.Text.Json.Serialization;
+using System.Threading.Tasks;
+using EpicGames.Core;
+
+namespace HordeToolbox
+{
+	class JsonConfig<T> where T : class, new()
+	{
+		readonly FileReference _file;
+		static readonly JsonSerializerOptions s_jsonSerializerOptions = GetJsonSerializerOptions();
+		T _current;
+		byte[] _data;
+
+		public T Current => _current;
+
+		public JsonConfig(FileReference file)
+		{
+			_file = file;
+			_current = new T();
+			_data = Array.Empty<byte>();
+		}
+
+		static JsonSerializerOptions GetJsonSerializerOptions()
+		{
+			JsonSerializerOptions options = new JsonSerializerOptions();
+			options.PropertyNamingPolicy = JsonNamingPolicy.CamelCase;
+			options.PropertyNameCaseInsensitive = true;
+			options.DefaultIgnoreCondition = JsonIgnoreCondition.WhenWritingNull;
+			options.AllowTrailingCommas = true;
+			options.WriteIndented = true;
+			options.Converters.Add(new JsonStringEnumConverter());
+			return options;
+		}
+
+		public bool LoadSettings()
+		{
+			if (FileReference.Exists(_file))
+			{
+				try
+				{
+					byte[] data = FileReference.ReadAllBytes(_file);
+					if (!data.SequenceEqual(_data))
+					{
+						_current = JsonSerializer.Deserialize<T>(data, s_jsonSerializerOptions)!;
+						_data = data;
+						return true;
+					}
+				}
+				catch (Exception)
+				{
+				}
+			}
+			return false;
+		}
+
+		public bool UpdateSettings(T settings)
+		{
+			byte[] data = JsonSerializer.SerializeToUtf8Bytes(settings, s_jsonSerializerOptions);
+			if (!data.SequenceEqual(_data))
+			{
+				_current = settings;
+				_data = data;
+
+				DirectoryReference.CreateDirectory(_file.Directory);
+				FileReference.WriteAllBytes(_file, data);
+
+				return true;
+			}
+			return false;
+		}
+	}
+}
diff --git a/Engine/Source/Programs/Horde/HordeToolbox/Plugins/HordeAgent/HordeAgentPlugin.cs b/Engine/Source/Programs/Horde/HordeToolbox/Plugins/HordeAgent/HordeAgentPlugin.cs
new file mode 100644
--- /dev/null
+++ b/Engine/Source/Programs/Horde/HordeToolbox/Plugins/HordeAgent/HordeAgentPlugin.cs
@@ -0,0 +1,477 @@
+// Copyright Epic Games, Inc. All Rights Reserved.
+
+using System.Diagnostics;
+using System.IO.Pipes;
+using System.Runtime.InteropServices;
+using System.Text.Json;
+using System.Text.Json.Serialization;
+using Avalonia.Controls;
+using EpicGames.Core;
+using FluentAvalonia.UI.Controls;
+using HordeAgent;
+using Microsoft.Win32;
+
+namespace HordeToolbox.Plugins.HordeAgent
+{
+	class HordeAgentPlugin : ITrayAppPlugin
+	{
+		record struct IdleStat(string Name, long Value, long MinValue);
+
+		readonly ITrayAppHost _host;
+
+		readonly BackgroundTask _clientTask;
+		readonly BackgroundTask _tickPauseStateTask;
+
+		readonly FileReference _settingsFile;
+
+		byte[] _settingsData = Array.Empty<byte>();
+		HordeAgentSettings _settings = new HordeAgentSettings();
+
+		public HordeAgentSettings Settings => _settings;
+
+		AgentSettingsMessage? _agentSettings;
+
+		public string Name => "Horde Agent";
+
+		public IconSource Icon => new SymbolIconSource() { Symbol = Symbol.People };
+
+		public Control CreateSettingsPage(SettingsContext context)
+			=> new HordeAgentSettingsPage(context, this);
+
+		public void UpdateSettings(HordeAgentSettings settings)
+		{
+			byte[] data = JsonSerializer.SerializeToUtf8Bytes(settings, GetJsonSerializerOptions());
+			if (!data.SequenceEqual(_settingsData))
+			{
+				_settings = settings;
+				_settingsData = data;
+
+				DirectoryReference.CreateDirectory(_settingsFile.Directory);
+				FileReference.WriteAllBytes(_settingsFile, data);
+
+				_statusChangedEvent.Set();
+			}
+		}
+
+		void EnrollWithServer()
+		{
+			Uri? serverUrl = _agentSettings?.ServerUrl;
+			if (serverUrl != null)
+			{
+				Process.Start(new ProcessStartInfo(new Uri(serverUrl, "agents/registration").ToString()) { UseShellExecute = true });
+			}
+		}
+
+		public HordeAgentPlugin(ITrayAppHost host)
+		{
+			_host = host;
+
+			DirectoryReference? settingsRoot = DirectoryReference.GetSpecialFolder(Environment.SpecialFolder.CommonApplicationData);
+			settingsRoot ??= DirectoryReference.GetCurrentDirectory();
+			_settingsFile = FileReference.Combine(settingsRoot, "Epic", "Horde", "TrayApp", "Settings.json");
+
+			LoadSettings();
+
+			_clientTask = BackgroundTask.StartNew(StatusTaskAsync);
+			_tickPauseStateTask = BackgroundTask.StartNew(ctx => TickPauseStateAsync(ctx));
+		}
+
+		static JsonSerializerOptions GetJsonSerializerOptions()
+		{
+			JsonSerializerOptions options = new JsonSerializerOptions();
+			options.PropertyNamingPolicy = JsonNamingPolicy.CamelCase;
+			options.PropertyNameCaseInsensitive = true;
+			options.DefaultIgnoreCondition = JsonIgnoreCondition.WhenWritingNull;
+			options.AllowTrailingCommas = true;
+			options.WriteIndented = true;
+			options.Converters.Add(new JsonStringEnumConverter());
+			return options;
+		}
+
+		void LoadSettings()
+		{
+			if (FileReference.Exists(_settingsFile))
+			{
+				try
+				{
+					byte[] data = FileReference.ReadAllBytes(_settingsFile);
+					if (!data.SequenceEqual(_settingsData))
+					{
+						_settings = JsonSerializer.Deserialize<HordeAgentSettings>(data, GetJsonSerializerOptions())!;
+						if (_settings.Mode == null)
+						{
+							_settings.Mode = ReadLegacyMode();
+							UpdateSettings(_settings);
+						}
+						_settingsData = data;
+						return;
+					}
+				}
+				catch (Exception)
+				{
+				}
+			}
+
+			UpdateSettings(new HordeAgentSettings());
+		}
+
+		static AgentMode? ReadLegacyMode()
+		{
+			if (!OperatingSystem.IsWindows())
+			{
+				return null;
+			}
+
+			const string RegistryKey = "HKEY_CURRENT_USER\\Software\\Epic Games\\Horde\\TrayApp";
+			const string RegistryStatusValue = "Status";
+
+			int? status = Registry.GetValue(RegistryKey, RegistryStatusValue, null) as int?;
+			if (status == null)
+			{
+				return null;
+			}
+
+			return status.Value switch
+			{
+				0 => AgentMode.Dedicated,
+				1 => AgentMode.Disabled,
+				2 => AgentMode.Workstation,
+				_ => null
+			};
+		}
+
+		public void PopulateContextMenu(NativeMenu contextMenu)
+		{
+			NativeMenuItem enrollMenuItem = new NativeMenuItem("Enroll with Server...");
+			enrollMenuItem.Click += (s, e) => EnrollWithServer();
+			contextMenu.Items.Add(enrollMenuItem);
+		}
+
+		public async ValueTask DisposeAsync()
+		{
+			await _tickPauseStateTask.DisposeAsync();
+			await _clientTask.DisposeAsync();
+		}
+
+		private void OnOpenLogs(object? sender, EventArgs e)
+		{
+			DirectoryReference? programDataDir = DirectoryReference.GetSpecialFolder(Environment.SpecialFolder.CommonApplicationData);
+			if (programDataDir != null)
+			{
+				DirectoryReference logsDir = DirectoryReference.Combine(programDataDir, "Epic", "Horde", "Agent");
+				if (DirectoryReference.Exists(logsDir))
+				{
+					Process.Start(new ProcessStartInfo { FileName = logsDir.FullName, UseShellExecute = true });
+				}
+			}
+		}
+
+		TrayAppPluginStatus? _status;
+
+		void SetStatus(AgentStatusMessage status)
+		{
+			if (!status.Healthy)
+			{
+				string message = String.IsNullOrEmpty(status.Detail) ? "Error. Check logs." : status.Detail.Length > 100 ? status.Detail.Substring(0, 100) : status.Detail;
+				_status = new TrayAppPluginStatus(TrayAppPluginState.Error, message);
+			}
+			else if (status.NumLeases > 0)
+			{
+				string message = status.NumLeases == 1 ? "Currently handling 1 lease" : $"Currently handling {status.NumLeases} leases";
+				_status = new TrayAppPluginStatus(TrayAppPluginState.Busy, message);
+			}
+			else if (_enabled)
+			{
+				string message = "Agent is operating normally";
+				_status = new TrayAppPluginStatus(TrayAppPluginState.Ok, message);
+			}
+			else
+			{
+				string message = "Agent is paused";
+				_status = new TrayAppPluginStatus(TrayAppPluginState.Paused, message);
+			}
+			_host.UpdateStatus();
+		}
+
+		public TrayAppPluginStatus GetStatus()
+			=> _status ?? TrayAppPluginStatus.Default;
+
+		async Task StatusTaskAsync(CancellationToken cancellationToken)
+		{
+			SetStatus(AgentStatusMessage.Starting);
+			for (; ; )
+			{
+				try
+				{
+					await PollForStatusUpdatesAsync(cancellationToken);
+				}
+				catch (OperationCanceledException) when (cancellationToken.IsCancellationRequested)
+				{
+					break;
+				}
+				catch
+				{
+					SetStatus(new AgentStatusMessage(false, 0, "Unable to connect to Agent. Check the service is running."));
+					await Task.Delay(TimeSpan.FromSeconds(5.0), cancellationToken);
+				}
+			}
+		}
+
+#pragma warning disable IDE1006
+		[StructLayout(LayoutKind.Sequential)]
+		struct LASTINPUTINFO
+		{
+			public int cbSize;
+			public uint dwTime;
+		}
+
+		[DllImport("user32.dll")]
+		static extern bool GetLastInputInfo(ref LASTINPUTINFO plii);
+
+		[DllImport("kernel32.dll")]
+		static extern uint GetTickCount();
+
+		[StructLayout(LayoutKind.Sequential)]
+		struct FILETIME
+		{
+			public uint dwLowDateTime;
+			public uint dwHighDateTime;
+
+			public readonly ulong Total => dwLowDateTime | (ulong)dwHighDateTime << 32;
+		};
+
+		[StructLayout(LayoutKind.Sequential)]
+		struct MEMORYSTATUSEX
+		{
+			public int dwLength;
+			public uint dwMemoryLoad;
+			public ulong ullTotalPhys;
+			public ulong ullAvailPhys;
+			public ulong ullTotalPageFile;
+			public ulong ullAvailPageFile;
+			public ulong ullTotalVirtual;
+			public ulong ullAvailVirtual;
+			public ulong ullAvailExtendedVirtual;
+		}
+
+		[DllImport("kernel32.dll", CharSet = CharSet.Auto, SetLastError = true)]
+		static extern bool GlobalMemoryStatusEx(ref MEMORYSTATUSEX lpBuffer);
+
+		[DllImport("kernel32.dll", SetLastError = true)]
+		static extern bool GetSystemTimes(out FILETIME lpIdleTime, out FILETIME lpKernelTime, out FILETIME lpUserTime);
+#pragma warning restore IDE1006
+
+		bool _enabled;
+		readonly AsyncEvent _statusChangedEvent = new AsyncEvent();
+		readonly AsyncEvent _enabledChangedEvent = new AsyncEvent();
+
+		async Task TickPauseStateAsync(CancellationToken cancellationToken)
+		{
+			await using BackgroundTask cpuStatsTask = BackgroundTask.StartNew(ctx => TickCpuStatsAsync(ctx));
+			await using BackgroundTask criticalProcessTask = BackgroundTask.StartNew(ctx => TickCriticalProcessAsync(ctx));
+
+			TimeSpan pollInterval = TimeSpan.FromSeconds(0.25);
+
+			Stopwatch stateChangeTimer = Stopwatch.StartNew();
+			while (!cancellationToken.IsCancellationRequested)
+			{
+				Task statusChangedTask = _statusChangedEvent.Task;
+
+				AgentMode mode = _settings.Mode ?? AgentMode.Disabled;
+				if (mode == AgentMode.Dedicated)
+				{
+					if (!_enabled)
+					{
+						_enabled = true;
+						_enabledChangedEvent.Set();
+					}
+				}
+				else if (mode == AgentMode.Disabled)
+				{
+					if (_enabled)
+					{
+						_enabled = false;
+						_enabledChangedEvent.Set();
+					}
+				}
+
+				DateTime utcNow = DateTime.UtcNow;
+				IEnumerable<IdleStat> idleStats = GetIdleStats();
+
+				bool idle = idleStats.All(x => x.Value >= x.MinValue);
+				if (idle == _enabled)
+				{
+					stateChangeTimer.Restart();
+				}
+
+				const int WakeTimeSecs = 2;
+				const int IdleTimeSecs = 30;
+				int stateChangeTime = (int)stateChangeTimer.Elapsed.TotalSeconds;
+				int stateChangeMaxTime = _enabled ? WakeTimeSecs : IdleTimeSecs;
+				//				_idleForm?.TickStats(_enabled, stateChangeTime, stateChangeMaxTime, idleStats);
+
+				if (mode == AgentMode.Workstation && stateChangeTime >= stateChangeMaxTime)
+				{
+					_enabled ^= true;
+					_enabledChangedEvent.Set();
+					stateChangeTimer.Restart();
+				}
+
+				await Task.WhenAny(statusChangedTask, Task.Delay(pollInterval, cancellationToken));
+			}
+		}
+
+		IEnumerable<IdleStat> GetIdleStats()
+		{
+			// Check there has been no input for a while
+			LASTINPUTINFO lastInputInfo = new LASTINPUTINFO();
+			lastInputInfo.cbSize = Marshal.SizeOf<LASTINPUTINFO>();
+
+			if (GetLastInputInfo(ref lastInputInfo))
+			{
+				yield return new IdleStat("LastInputTime", (GetTickCount() - lastInputInfo.dwTime) / 1000, _settings.Idle.MinIdleTimeSecs);
+			}
+
+			// Check that no critical processes are running
+			if (_settings.Idle.CriticalProcesses.Any())
+			{
+				yield return new IdleStat("CriticalProcCount", -_idleCriticalProcessCount, 0);
+			}
+
+			// Only look at memory/CPU usage if we're not paused; executing jobs will increase them
+			if (!_enabled)
+			{
+				// Check the CPU usage doesn't exceed the limit
+				yield return new IdleStat("IdleCpuPct", _idleCpuPct, _settings.Idle.MinIdleCpuPct);
+
+				// Check there's enough available virtual memory 
+				MEMORYSTATUSEX memoryStatus = new MEMORYSTATUSEX();
+				memoryStatus.dwLength = Marshal.SizeOf<MEMORYSTATUSEX>();
+
+				if (GlobalMemoryStatusEx(ref memoryStatus))
+				{
+					yield return new IdleStat("VirtualMemMb", (long)(memoryStatus.ullAvailPhys + memoryStatus.ullAvailPageFile) / (1024 * 1024), _settings.Idle.MinFreeVirtualMemMb);
+				}
+			}
+		}
+
+		int _idleCpuPct = 0;
+		int _idleCriticalProcessCount = 0;
+
+		async Task TickCpuStatsAsync(CancellationToken cancellationToken)
+		{
+			const int NumSamples = 10;
+			TimeSpan sampleInterval = TimeSpan.FromSeconds(0.2);
+			(ulong IdleTime, ulong TotalTime)[] samples = new (ulong IdleTime, ulong TotalTime)[NumSamples];
+
+			int sampleIdx = 0;
+			for (; ; )
+			{
+				if (GetSystemTimes(out FILETIME idleTime, out FILETIME kernelTime, out FILETIME userTime))
+				{
+					(ulong prevIdleTime, ulong prevTotalTime) = samples[sampleIdx];
+					(ulong nextIdleTime, ulong nextTotalTime) = (idleTime.Total, kernelTime.Total + userTime.Total);
+
+					samples[sampleIdx] = (nextIdleTime, nextTotalTime);
+					sampleIdx = (sampleIdx + 1) % NumSamples;
+
+					if (prevTotalTime > 0 && nextTotalTime > prevTotalTime)
+					{
+						_idleCpuPct = (int)((nextIdleTime - prevIdleTime) * 100 / (nextTotalTime - prevTotalTime));
+					}
+				}
+				await Task.Delay(sampleInterval, cancellationToken);
+			}
+		}
+
+		async Task TickCriticalProcessAsync(CancellationToken cancellationToken)
+		{
+			TimeSpan sampleInterval = TimeSpan.FromSeconds(1.0);
+
+			for (; ; )
+			{
+				try
+				{
+					if (RuntimeInformation.IsOSPlatform(OSPlatform.Windows) && _settings.Idle.CriticalProcesses.Any())
+					{
+						IEnumerable<int> hordeProcessIds = Process.GetProcessesByName("HordeAgent").Select(x => x.Id);
+						IEnumerable<Process> criticalProcesses = _settings.Idle.CriticalProcesses
+							.Select(x => Path.GetFileNameWithoutExtension(x).ToUpperInvariant())
+							.Distinct()
+							.SelectMany(x => Process.GetProcessesByName(x))
+							.Where(x => !x.HasExited);
+
+						// Ignore processes that are descendants of HordeAgent
+						if (hordeProcessIds.Any() && criticalProcesses.Any())
+						{
+							criticalProcesses = criticalProcesses
+								.Where(x => !ProcessUtils.GetAncestorProcesses(x)
+								.Select(x => x.Id).Intersect(hordeProcessIds).Any())
+								.Where(x => !x.HasExited);
+						}
+
+						_idleCriticalProcessCount = criticalProcesses.Count();
+					}
+				}
+				catch (InvalidOperationException)
+				{
+					// If a process stops running Process.Id will throw an exception
+				}
+				await Task.Delay(sampleInterval, cancellationToken);
+			}
+		}
+
+		async Task PollForStatusUpdatesAsync(CancellationToken cancellationToken)
+		{
+			AgentMessageBuffer message = new AgentMessageBuffer();
+			using (NamedPipeClientStream pipeClient = new NamedPipeClientStream(".", AgentMessagePipe.PipeName, PipeDirection.InOut))
+			{
+				SetStatus(new AgentStatusMessage(false, 0, "Connecting to agent..."));
+				await pipeClient.ConnectAsync(cancellationToken);
+
+				SetStatus(new AgentStatusMessage(false, 0, "Waiting for status update."));
+				for (; ; )
+				{
+					Task idleChangeTask = _enabledChangedEvent.Task;
+
+					bool enabled = _enabled;
+					message.Set(AgentMessageType.SetEnabledRequest, new AgentEnabledMessage(enabled));
+					await message.SendAsync(pipeClient, cancellationToken);
+
+					if (_agentSettings == null)
+					{
+						message.Set(AgentMessageType.GetSettingsRequest);
+						await message.SendAsync(pipeClient, cancellationToken);
+
+						if (!await message.TryReadAsync(pipeClient, cancellationToken))
+						{
+							break;
+						}
+						if (message.Type == AgentMessageType.GetSettingsResponse)
+						{
+							_agentSettings = message.Parse<AgentSettingsMessage>();
+						}
+					}
+
+					message.Set(AgentMessageType.GetStatusRequest);
+					await message.SendAsync(pipeClient, cancellationToken);
+
+					if (!await message.TryReadAsync(pipeClient, cancellationToken))
+					{
+						break;
+					}
+
+					switch (message.Type)
+					{
+						case AgentMessageType.GetStatusResponse:
+							AgentStatusMessage status = message.Parse<AgentStatusMessage>();
+							SetStatus(status);
+							break;
+					}
+
+					await Task.WhenAny(idleChangeTask, Task.Delay(TimeSpan.FromSeconds(5.0), cancellationToken));
+				}
+			}
+		}
+	}
+}
diff --git a/Engine/Source/Programs/Horde/HordeToolbox/Plugins/HordeAgent/HordeAgentSettings.cs b/Engine/Source/Programs/Horde/HordeToolbox/Plugins/HordeAgent/HordeAgentSettings.cs
new file mode 100644
--- /dev/null
+++ b/Engine/Source/Programs/Horde/HordeToolbox/Plugins/HordeAgent/HordeAgentSettings.cs
@@ -0,0 +1,30 @@
+// Copyright Epic Games, Inc. All Rights Reserved.
+
+#pragma warning disable CS1591
+
+namespace HordeToolbox.Plugins.HordeAgent
+{
+	public enum AgentMode
+	{
+		Workstation = 0,
+		Dedicated = 1,
+		Disabled = 2,
+	}
+
+	[Serializable]
+	public class HordeAgentSettings
+	{
+		public AgentMode? Mode { get; set; }
+		public IdleSettings Idle { get; set; } = new();
+	}
+
+	[Serializable]
+	public class IdleSettings
+	{
+		public int MinIdleTimeSecs { get; set; } = 2;
+		public int MinIdleCpuPct { get; set; } = 70;
+		public int MinFreeVirtualMemMb { get; set; } = 256;
+
+		public string[] CriticalProcesses { get; set; } = Array.Empty<string>();
+	}
+}
diff --git a/Engine/Source/Programs/Horde/HordeToolbox/Plugins/HordeAgent/HordeAgentSettingsPage.axaml.cs b/Engine/Source/Programs/Horde/HordeToolbox/Plugins/HordeAgent/HordeAgentSettingsPage.axaml.cs
new file mode 100644
--- /dev/null
+++ b/Engine/Source/Programs/Horde/HordeToolbox/Plugins/HordeAgent/HordeAgentSettingsPage.axaml.cs
@@ -0,0 +1,21 @@
+// Copyright Epic Games, Inc. All Rights Reserved.
+
+using Avalonia.Controls;
+
+namespace HordeToolbox.Plugins.HordeAgent
+{
+	partial class HordeAgentSettingsPage : UserControl
+	{
+		public HordeAgentSettingsPage()
+			: this(SettingsContext.Default, null)
+		{
+		}
+
+		public HordeAgentSettingsPage(SettingsContext context, HordeAgentPlugin? plugin)
+		{
+			InitializeComponent();
+
+			DataContext = new HordeAgentSettingsViewModel(context, plugin);
+		}
+	}
+}
\ No newline at end of file
diff --git a/Engine/Source/Programs/Horde/HordeToolbox/Plugins/HordeAgent/HordeAgentSettingsViewModel.cs b/Engine/Source/Programs/Horde/HordeToolbox/Plugins/HordeAgent/HordeAgentSettingsViewModel.cs
new file mode 100644
--- /dev/null
+++ b/Engine/Source/Programs/Horde/HordeToolbox/Plugins/HordeAgent/HordeAgentSettingsViewModel.cs
@@ -0,0 +1,130 @@
+// Copyright Epic Games, Inc. All Rights Reserved.
+
+using System.ComponentModel;
+using System.Diagnostics;
+using Avalonia.Platform.Storage;
+using CommunityToolkit.Mvvm.ComponentModel;
+
+namespace HordeToolbox.Plugins.HordeAgent
+{
+	partial class HordeAgentSettingsViewModel : ObservableObject
+	{
+		readonly SettingsContext _context;
+		readonly HordeAgentPlugin? _plugin;
+
+		AgentMode _mode;
+
+		static string GetPropertyName(AgentMode mode)
+			=> mode switch
+			{
+				AgentMode.Dedicated => nameof(IsDedicatedMode),
+				AgentMode.Workstation => nameof(IsWorkstationMode),
+				_ => nameof(IsDisabled)
+			};
+
+		void SetModeIfTrue(AgentMode mode, bool value)
+		{
+			if (value && mode != _mode)
+			{
+				string oldName = GetPropertyName(_mode);
+				string newName = GetPropertyName(mode);
+
+				OnPropertyChanging(oldName);
+				OnPropertyChanging(newName);
+
+				_mode = mode;
+
+				OnPropertyChanged(oldName);
+				OnPropertyChanged(newName);
+			}
+		}
+
+		public bool IsDedicatedMode
+		{
+			get => _mode == AgentMode.Dedicated;
+			set => SetModeIfTrue(AgentMode.Dedicated, value);
+		}
+
+		public bool IsWorkstationMode
+		{
+			get => _mode == AgentMode.Workstation;
+			set => SetModeIfTrue(AgentMode.Workstation, value);
+		}
+
+		public bool IsDisabled
+		{
+			get => _mode == AgentMode.Disabled;
+			set => SetModeIfTrue(AgentMode.Disabled, value);
+		}
+
+		[ObservableProperty]
+		int _minIdleTimeSecs;
+
+		[ObservableProperty]
+		int _minFreeMemoryMb;
+
+		[ObservableProperty]
+		int _minFreeCpuPct;
+
+		[ObservableProperty]
+		string _sandboxDir = "C:\\Dump"; // TODO
+
+		public HordeAgentSettingsViewModel(SettingsContext context, HordeAgentPlugin? plugin)
+		{
+			_context = context;
+			_plugin = plugin;
+
+			HordeAgentSettings settings = plugin?.Settings ?? new HordeAgentSettings();
+			_mode = settings.Mode ?? AgentMode.Disabled;
+			_minIdleTimeSecs = settings.Idle.MinIdleTimeSecs;
+			_minFreeMemoryMb = settings.Idle.MinFreeVirtualMemMb;
+			_minFreeCpuPct = settings.Idle.MinIdleCpuPct;
+		}
+
+		protected override void OnPropertyChanged(PropertyChangedEventArgs e)
+		{
+			base.OnPropertyChanged(e);
+
+			if (_plugin != null)
+			{
+				HordeAgentSettings settings = new HordeAgentSettings();
+
+				settings.Mode = _mode;
+				settings.Idle.MinIdleTimeSecs = MinIdleTimeSecs;
+				settings.Idle.MinFreeVirtualMemMb = MinFreeMemoryMb;
+				settings.Idle.MinIdleCpuPct = MinFreeCpuPct;
+
+				_plugin.UpdateSettings(settings);
+			}
+		}
+
+		public void OpenSandbox()
+		{
+			ProcessStartInfo startInfo = new ProcessStartInfo("explorer.exe");
+			startInfo.ArgumentList.Add(SandboxDir);
+			Process.Start(startInfo);
+		}
+
+		public async Task MoveSandboxAsync()
+		{
+			if (_context.SettingsWindow != null)
+			{
+				IStorageProvider storageProvider = _context.SettingsWindow.StorageProvider;
+
+				FolderPickerOpenOptions options = new FolderPickerOpenOptions();
+				options.AllowMultiple = false;
+				options.SuggestedStartLocation = await storageProvider.TryGetFolderFromPathAsync(SandboxDir);
+
+				IReadOnlyList<IStorageFolder> folders = await storageProvider.OpenFolderPickerAsync(options);
+				if (folders.Count > 0)
+				{
+					string? path = folders[0].TryGetLocalPath();
+					if (path != null)
+					{
+						SandboxDir = path;
+					}
+				}
+			}
+		}
+	}
+}
diff --git a/Engine/Source/Programs/Horde/HordeToolbox/Plugins/HordeProxy/HordeProxyPlugin.cs b/Engine/Source/Programs/Horde/HordeToolbox/Plugins/HordeProxy/HordeProxyPlugin.cs
new file mode 100644
--- /dev/null
+++ b/Engine/Source/Programs/Horde/HordeToolbox/Plugins/HordeProxy/HordeProxyPlugin.cs
@@ -0,0 +1,165 @@
+// Copyright Epic Games, Inc. All Rights Reserved.
+
+using System.Net;
+using System.Text.Json;
+using System.Text.Json.Serialization;
+using Avalonia.Controls;
+using EpicGames.Core;
+using EpicGames.Horde;
+using FluentAvalonia.UI.Controls;
+using Microsoft.AspNetCore.Builder;
+using Microsoft.AspNetCore.Hosting;
+using Microsoft.AspNetCore.Http;
+using Microsoft.Extensions.DependencyInjection;
+using Microsoft.Extensions.Hosting;
+
+namespace HordeToolbox.Plugins.HordeProxy
+{
+	class HordeProxyPlugin : ITrayAppPlugin
+	{
+		readonly BackgroundTask _serverTask;
+		readonly IHordeClientProvider _hordeClientProvider;
+		readonly AsyncEvent _restartServerEvent = new AsyncEvent();
+		readonly SynchronizationContext? _synchronizationContext;
+		readonly JsonConfig<HordeProxySettings> _settings;
+
+		string _status = "Starting...";
+
+		public HordeProxySettings Settings => _settings.Current;
+
+		public event Action? OnStatusChanged;
+
+		public string Name => "Horde Proxy";
+
+		public string Status => _status;
+
+		public IconSource Icon => new SymbolIconSource() { Symbol = Symbol.Filter };
+
+		public HordeProxyPlugin(IHordeClientProvider hordeClientProvider)
+		{
+			_hordeClientProvider = hordeClientProvider;
+			_hordeClientProvider.OnStateChanged += OnStateChanged;
+			_synchronizationContext = SynchronizationContext.Current;
+
+			DirectoryReference ? settingsRoot = DirectoryReference.GetSpecialFolder(Environment.SpecialFolder.CommonApplicationData);
+			settingsRoot ??= DirectoryReference.GetCurrentDirectory();
+			_settings = new JsonConfig<HordeProxySettings>(FileReference.Combine(settingsRoot, "Epic", "Horde", "Proxy", "Settings.json"));
+			_settings.LoadSettings();
+
+			_serverTask = BackgroundTask.StartNew(RunServerAsync);
+		}
+
+		public Control CreateSettingsPage(SettingsContext context)
+			=> new HordeProxySettingsPage(this);
+
+		void SetStatus(string status)
+		{
+			_status = status;
+
+			if (_synchronizationContext == null)
+			{
+				OnStatusChanged?.Invoke();
+			}
+			else
+			{
+				_synchronizationContext.Post(_ => OnStatusChanged?.Invoke(), null);
+			}
+		}
+
+		public void UpdateSettings(HordeProxySettings settings)
+		{
+			if (_settings.UpdateSettings(settings))
+			{
+				_restartServerEvent.Pulse();
+			}
+		}
+
+		public async ValueTask DisposeAsync()
+		{
+			_hordeClientProvider.OnStateChanged -= OnStateChanged;
+			await _serverTask.DisposeAsync();
+		}
+
+		void OnStateChanged()
+		{
+			_restartServerEvent.Pulse();
+		}
+
+		public void PopulateContextMenu(NativeMenu contextMenu)
+		{
+		}
+
+		public TrayAppPluginStatus GetStatus()
+			=> TrayAppPluginStatus.Default;
+
+		async Task RunServerAsync(CancellationToken cancellationToken)
+		{
+			while (!cancellationToken.IsCancellationRequested)
+			{
+				Task restartServerTask = _restartServerEvent.Task;
+
+				HordeProxySettings settings = _settings.Current;
+				if (settings.Enabled)
+				{
+					using CancellationTokenSource cancellationSource = CancellationTokenSource.CreateLinkedTokenSource(cancellationToken);
+					Task stoppedTask = _restartServerEvent.Task.ContinueWith(_ => cancellationSource.Cancel(), cancellationSource.Token, TaskContinuationOptions.None, TaskScheduler.Default);
+
+					try
+					{
+						SetStatus($"Running on http://localhost:{settings.Port}");
+
+						WebApplicationBuilder builder = WebApplication.CreateBuilder();
+						builder.Services.AddHttpClient();
+						builder.Services.AddHorde();
+						builder.WebHost.ConfigureKestrel(options => options.ListenLocalhost(settings.Port));
+
+						await using WebApplication app = builder.Build();
+						app.Map("/{*Route}", ForwardRequestAsync);
+						await app.RunAsync(cancellationSource.Token);
+					}
+					catch (OperationCanceledException)
+					{
+						throw;
+					}
+					catch (Exception ex)
+					{
+						SetStatus($"Error: {ex.Message}");
+					}
+
+					await stoppedTask;
+				}
+				else
+				{
+					SetStatus("Stopped");
+					await restartServerTask.WaitAsync(cancellationToken);
+				}
+				await Task.Delay(TimeSpan.FromSeconds(2.0), cancellationToken);
+			}
+		}
+
+		static async Task ForwardRequestAsync(HttpContext context)
+		{
+			IHordeClient hordeClient = context.RequestServices.GetRequiredService<IHordeClient>();
+			HttpClient httpClient = hordeClient.CreateHttpClient().HttpClient;
+
+			if (context.Request.Method != "GET")
+			{
+				context.Response.StatusCode = (int)HttpStatusCode.MethodNotAllowed;
+				return;
+			}
+
+			using HttpRequestMessage request = new HttpRequestMessage();
+			request.Method = HttpMethod.Parse(context.Request.Method);
+			request.RequestUri = new Uri(hordeClient.ServerUrl, $"{context.Request.Path}");
+			request.Content = new StreamContent(context.Request.Body);
+
+			using HttpResponseMessage response = await httpClient.SendAsync(request, context.RequestAborted);
+			context.Response.StatusCode = (int)response.StatusCode;
+			context.Response.ContentType = response.Content.Headers.ContentType?.MediaType;
+			context.Response.ContentLength = response.Content.Headers.ContentLength;
+
+			using Stream stream = await response.Content.ReadAsStreamAsync(context.RequestAborted);
+			await stream.CopyToAsync(context.Response.Body, context.RequestAborted);
+		}
+	}
+}
diff --git a/Engine/Source/Programs/Horde/HordeToolbox/Plugins/HordeProxy/HordeProxySettings.cs b/Engine/Source/Programs/Horde/HordeToolbox/Plugins/HordeProxy/HordeProxySettings.cs
new file mode 100644
--- /dev/null
+++ b/Engine/Source/Programs/Horde/HordeToolbox/Plugins/HordeProxy/HordeProxySettings.cs
@@ -0,0 +1,13 @@
+// Copyright Epic Games, Inc. All Rights Reserved.
+
+#pragma warning disable CS1591
+
+namespace HordeToolbox.Plugins.HordeProxy
+{
+	[Serializable]
+	public class HordeProxySettings
+	{
+		public bool Enabled { get; set; }
+		public int Port { get; set; } = 13344;
+	}
+}
diff --git a/Engine/Source/Programs/Horde/HordeToolbox/Plugins/HordeProxy/HordeProxySettingsPage.axaml.cs b/Engine/Source/Programs/Horde/HordeToolbox/Plugins/HordeProxy/HordeProxySettingsPage.axaml.cs
new file mode 100644
--- /dev/null
+++ b/Engine/Source/Programs/Horde/HordeToolbox/Plugins/HordeProxy/HordeProxySettingsPage.axaml.cs
@@ -0,0 +1,21 @@
+// Copyright Epic Games, Inc. All Rights Reserved.
+
+using Avalonia.Controls;
+
+namespace HordeToolbox.Plugins.HordeProxy
+{
+	partial class HordeProxySettingsPage : UserControl
+	{
+		public HordeProxySettingsPage()
+			: this(null)
+		{
+		}
+
+		public HordeProxySettingsPage(HordeProxyPlugin? plugin)
+		{
+			InitializeComponent();
+
+			DataContext = new HordeProxySettingsViewModel(plugin);
+		}
+	}
+}
\ No newline at end of file
diff --git a/Engine/Source/Programs/Horde/HordeToolbox/Plugins/HordeProxy/HordeProxySettingsViewModel.cs b/Engine/Source/Programs/Horde/HordeToolbox/Plugins/HordeProxy/HordeProxySettingsViewModel.cs
new file mode 100644
--- /dev/null
+++ b/Engine/Source/Programs/Horde/HordeToolbox/Plugins/HordeProxy/HordeProxySettingsViewModel.cs
@@ -0,0 +1,76 @@
+// Copyright Epic Games, Inc. All Rights Reserved.
+
+using System.ComponentModel;
+using System.Diagnostics;
+using Avalonia.Platform.Storage;
+using Avalonia.Threading;
+using CommunityToolkit.Mvvm.ComponentModel;
+
+namespace HordeToolbox.Plugins.HordeProxy
+{
+	partial class HordeProxySettingsViewModel : ObservableObject
+	{
+		readonly HordeProxyPlugin? _plugin;
+
+		int _pendingUpdate;
+
+		[ObservableProperty]
+		bool _enabled;
+
+		[ObservableProperty]
+		int _port;
+
+		[ObservableProperty]
+		string _status = String.Empty;
+
+		public HordeProxySettingsViewModel(HordeProxyPlugin? plugin)
+		{
+			_plugin = plugin;
+
+			HordeProxySettings settings = plugin?.Settings ?? new HordeProxySettings();
+			_enabled = settings.Enabled;
+			_port = settings.Port;
+
+			if (_plugin != null)
+			{
+				_status = _plugin.Status;
+				_plugin.OnStatusChanged += OnStatusChanged;
+			}
+		}
+
+		void OnStatusChanged()
+			=> Status = _plugin?.Status ?? String.Empty;
+
+		protected override void OnPropertyChanged(PropertyChangedEventArgs e)
+		{
+			base.OnPropertyChanged(e);
+
+			if (_plugin != null)
+			{
+				int currentUpdate = Interlocked.Increment(ref _pendingUpdate);
+				ApplySettingsAsync(currentUpdate);
+			}
+		}
+
+		async void ApplySettingsAsync(int currentUpdate)
+		{
+			await Task.Delay(TimeSpan.FromSeconds(1.0));
+
+			if (Interlocked.CompareExchange(ref _pendingUpdate, 0, 0) == currentUpdate)
+			{
+				ApplySettings();
+			}
+		}
+
+		void ApplySettings()
+		{
+			if (_plugin != null)
+			{
+				HordeProxySettings settings = new HordeProxySettings();
+				settings.Enabled = Enabled;
+				settings.Port = Port;
+				_plugin.UpdateSettings(settings);
+			}
+		}
+	}
+}
diff --git a/Engine/Source/Programs/Horde/HordeToolbox/Program.cs b/Engine/Source/Programs/Horde/HordeToolbox/Program.cs
new file mode 100644
--- /dev/null
+++ b/Engine/Source/Programs/Horde/HordeToolbox/Program.cs
@@ -0,0 +1,142 @@
+// Copyright Epic Games, Inc. All Rights Reserved.
+
+using System.Diagnostics;
+using System.Reflection;
+using System.Runtime.InteropServices;
+using Avalonia;
+using Avalonia.Controls;
+using EpicGames.Core;
+
+#pragma warning disable CA1806 // Not checking return code from MessageBox
+
+namespace HordeToolbox
+{
+	static class Program
+	{
+		const string MutexName = "HordeToolbox-Mutex";
+		const string EventName = "HordeToolbox-Exit";
+
+		[DllImport("user32.dll", SetLastError = true, CharSet = CharSet.Auto)]
+		public static extern int MessageBox(IntPtr hWnd, string text, string? caption, uint type);
+
+		[STAThread]
+		public static int Main(string[] args)
+		{
+			using EventWaitHandle closeEvent = new EventWaitHandle(false, EventResetMode.AutoReset, EventName);
+			using SingleInstanceMutex mutex = new SingleInstanceMutex(MutexName);
+
+			if (args.Any(x => x.Equals("-close", StringComparison.OrdinalIgnoreCase)))
+			{
+				closeEvent.Set();
+				return mutex.Wait(5000) ? 0 : 1;
+			}
+
+			if (!mutex.Wait(0))
+			{
+				return 1;
+			}
+
+			if (args.Any(x => x.Equals("-shadowcopy", StringComparison.OrdinalIgnoreCase)))
+			{
+				DirectoryReference? localAppData = DirectoryReference.GetSpecialFolder(Environment.SpecialFolder.LocalApplicationData);
+				if (localAppData != null)
+				{
+					FileReference sourceExe = new FileReference(Assembly.GetExecutingAssembly().Location).ChangeExtension(".exe");
+					if (FileReference.Exists(sourceExe) && !sourceExe.IsUnderDirectory(localAppData))
+					{
+						DirectoryReference sourceDir = sourceExe.Directory;
+						DirectoryReference targetDir = DirectoryReference.Combine(localAppData, "Epic Games", "HordeToolbox");
+						try
+						{
+							DirectoryReference.CreateDirectory(targetDir);
+							FileUtils.ForceDeleteDirectoryContents(targetDir);
+
+							CopyFiles(sourceDir, targetDir);
+							mutex.Release();
+
+							FileReference targetExe = FileReference.Combine(targetDir, sourceExe.MakeRelativeTo(sourceDir)).ChangeExtension(".exe");
+							using Process process = Process.Start(targetExe.FullName, args);
+
+							return 0;
+						}
+						catch (Exception ex)
+						{
+							MessageBox(IntPtr.Zero, $"Unable to copy app to temp location. Error:\n\n{ex}", null, 0);
+							return 1;
+						}
+					}
+				}
+			}
+
+			BuildAvaloniaApp()
+				.StartWithClassicDesktopLifetime(args, ShutdownMode.OnExplicitShutdown);
+
+			return 0;
+		}
+
+		static void CopyFiles(DirectoryReference sourceDir, DirectoryReference targetDir)
+		{
+			DirectoryReference.CreateDirectory(targetDir);
+			foreach (DirectoryReference sourceSubDir in DirectoryReference.EnumerateDirectories(sourceDir))
+			{
+				DirectoryReference targetSubDir = DirectoryReference.Combine(targetDir, sourceSubDir.GetDirectoryName());
+				CopyFiles(sourceSubDir, targetSubDir);
+			}
+			foreach (FileReference sourceFile in DirectoryReference.EnumerateFiles(sourceDir))
+			{
+				FileReference targetFile = FileReference.Combine(targetDir, sourceFile.GetFileName());
+				FileReference.Copy(sourceFile, targetFile, true);
+			}
+		}
+
+		// Avalonia configuration, don't remove; also used by visual designer.
+		public static AppBuilder BuildAvaloniaApp()
+		{
+			return AppBuilder.Configure<App>()
+				.UsePlatformDetect()
+				.LogToTrace();
+		}
+	}
+
+	class SingleInstanceMutex : IDisposable
+	{
+		readonly Mutex _mutex;
+		bool _locked;
+
+		public SingleInstanceMutex(string name)
+		{
+			_mutex = new Mutex(true, name);
+		}
+
+		public void Release()
+		{
+			if (_locked)
+			{
+				_mutex.ReleaseMutex();
+				_locked = false;
+			}
+		}
+
+		public bool Wait(int timeout)
+		{
+			if (!_locked)
+			{
+				try
+				{
+					_locked = _mutex.WaitOne(timeout);
+				}
+				catch (AbandonedMutexException)
+				{
+					_locked = true;
+				}
+			}
+			return _locked;
+		}
+
+		public void Dispose()
+		{
+			Release();
+			_mutex.Dispose();
+		}
+	}
+}
diff --git a/Engine/Source/Programs/Horde/HordeToolbox/SettingsWindow.axaml.cs b/Engine/Source/Programs/Horde/HordeToolbox/SettingsWindow.axaml.cs
new file mode 100644
--- /dev/null
+++ b/Engine/Source/Programs/Horde/HordeToolbox/SettingsWindow.axaml.cs
@@ -0,0 +1,68 @@
+// Copyright Epic Games, Inc. All Rights Reserved.
+
+using Avalonia.Controls;
+using FluentAvalonia.UI.Controls;
+using Microsoft.Extensions.DependencyInjection;
+
+namespace HordeToolbox
+{
+	record class SettingsContext(Window? SettingsWindow, IServiceProvider ServiceProvider)
+	{
+		public static SettingsContext Default { get; } = new SettingsContext(null!, new ServiceCollection().BuildServiceProvider());
+	}
+
+	partial class SettingsWindow : Window
+	{
+		readonly IServiceProvider? _serviceProvider;
+		readonly Dictionary<string, ITrayAppPlugin> _typeToPlugin = new Dictionary<string, ITrayAppPlugin>();
+
+		public SettingsWindow()
+			: this(null!)
+		{ }
+
+		public SettingsWindow(IServiceProvider serviceProvider)
+		{
+			InitializeComponent();
+
+			_serviceProvider = serviceProvider;
+
+			foreach (ITrayAppPlugin plugin in _serviceProvider.GetServices<ITrayAppPlugin>())
+			{
+				string typeName = plugin.GetType().FullName!;
+				_navView.MenuItems.Add(new NavigationViewItem(){ Content = plugin.Name, IconSource = plugin.Icon, Tag = typeName });
+				_typeToPlugin.Add(typeName, plugin);
+			}
+
+			_navView.SelectedItem = _navView.MenuItems[0];
+			_navView.SelectionChanged += NavView_SelectionChanged;
+
+			NavView_UpdateContent();
+		}
+
+		private void NavView_SelectionChanged(object? sender, NavigationViewSelectionChangedEventArgs e)
+		{
+			if (sender == _navView)
+			{
+				NavView_UpdateContent();
+			}
+		}
+
+		private void NavView_UpdateContent()
+		{
+			if (_navView.SelectedItem is NavigationViewItem nvi)
+			{
+				SettingsContext context = new SettingsContext(this, _serviceProvider!);
+
+				string? typeName = (string?)nvi.Tag;
+				if (typeName != null && _typeToPlugin.TryGetValue(typeName, out ITrayAppPlugin? plugin))
+				{
+					_navView.Content = plugin.CreateSettingsPage(context);
+				}
+				else
+				{
+					_navView.Content = new GeneralSettingsPage(context);
+				}
+			}
+		}
+	}
+}
diff --git a/Engine/Source/Programs/Horde/HordeToolbox/ToolCatalog.cs b/Engine/Source/Programs/Horde/HordeToolbox/ToolCatalog.cs
new file mode 100644
--- /dev/null
+++ b/Engine/Source/Programs/Horde/HordeToolbox/ToolCatalog.cs
@@ -0,0 +1,870 @@
+// Copyright Epic Games, Inc. All Rights Reserved.
+
+using System.Diagnostics;
+using System.Runtime.InteropServices;
+using System.Text.Json;
+using EpicGames.Core;
+using EpicGames.Horde;
+using EpicGames.Horde.Storage.Nodes;
+using EpicGames.Horde.Tools;
+using Microsoft.Extensions.Logging;
+
+#nullable enable
+
+namespace HordeToolbox
+{
+	class ToolCatalog : IToolCatalog, IAsyncDisposable
+	{
+		record class ItemState(string Name, string Description, CurrentToolDeploymentInfo? Current, PendingToolDeploymentInfo? Pending, IToolDeployment? Latest);
+
+		[DebuggerDisplay("{Id}")]
+		class Item : IToolCatalogItem
+		{
+			public ToolId Id { get; }
+
+			readonly ToolCatalog _catalog;
+
+			// Keep a separate copy of the state on the main thread; we will update it explicitly. This prevents tearing of updates on UI threads.
+			public ItemState _mainThreadState;
+			public ToolDeploymentInfo? _mainThreadLatest;
+
+			// Public properties only expose main thread data
+			string IToolCatalogItem.Name => _mainThreadState.Name;
+			string IToolCatalogItem.Description => _mainThreadState.Description;
+			ToolDeploymentInfo? IToolCatalogItem.Latest => _mainThreadLatest;
+			CurrentToolDeploymentInfo? IToolCatalogItem.Current => _mainThreadState.Current;
+			PendingToolDeploymentInfo? IToolCatalogItem.Pending => _mainThreadState.Pending;
+
+			public ItemState _state;
+			public DeploymentTask? _deploymentTask;
+			public event Action? OnItemChanged;
+
+			public Item(ToolCatalog catalog, ToolId id, ItemState state)
+			{
+				_catalog = catalog;
+				Id = id;
+				_state = state;
+
+				_mainThreadState = null!;
+				UpdateMainThreadState();
+			}
+
+			public void UpdateMainThreadState()
+			{
+				_mainThreadState = _state;
+				_mainThreadLatest = (_state.Latest == null) ? null : new ToolDeploymentInfo(_state.Latest.Id, _state.Latest.Version);
+			}
+
+			public void NotifyChanged()
+			{
+				OnItemChanged?.Invoke();
+			}
+
+			public void SetState(ItemState state)
+			{
+				_state = state;
+				OnItemChanged?.Invoke();
+			}
+
+			public void Cancel()
+				=> _catalog.CancelDeployment(this, _mainThreadState.Pending?.Deployment);
+
+			// Installing or uninstalling changes the pending task state immediately, and anyone copying the pending state will always see the value of that state.
+			// UI things may wait on that state, and will get notified from that state.
+			// So we should basically post a notice to the main thread to say we've queued up the next state.
+
+			public void Install()
+				=> _catalog.UpdateDeployment(this, _mainThreadState.Latest);
+
+			public void Uninstall()
+				=> _catalog.UpdateDeployment(this, null);
+		}
+
+#pragma warning disable CA1001 // Cancellation source is disposed by the task itself
+		class DeploymentTask
+#pragma warning restore CA1001
+		{
+			public ToolDeploymentInfo Deployment { get; }
+
+			CancellationTokenSource? _cancellationSource;
+			readonly Task _task;
+
+			public DeploymentTask(ToolDeploymentInfo deployment, Func<CancellationToken, Task> func, DeploymentTask? lastDeploymentTask)
+			{
+				Deployment = deployment;
+
+				CancellationTokenSource cancellationSource = new CancellationTokenSource();
+				_task = Task.Run(() => RunAsync(func, lastDeploymentTask, cancellationSource), CancellationToken.None);
+
+				_cancellationSource = cancellationSource;
+			}
+
+			async Task RunAsync(Func<CancellationToken, Task> func, DeploymentTask? lastDeploymentTask, CancellationTokenSource cancellationSource)
+			{
+				CancellationToken cancellationToken = cancellationSource!.Token;
+				try
+				{
+					if (lastDeploymentTask != null)
+					{
+						lastDeploymentTask.Cancel();
+						await lastDeploymentTask.WaitAsync(CancellationToken.None);
+					}
+					await func(cancellationToken);
+				}
+				catch (OperationCanceledException)
+				{
+				}
+				finally
+				{
+					lock (_task)
+					{
+						cancellationSource.Dispose(); // Note; the instance passed into this function, not the class member.
+						_cancellationSource = null;
+					}
+				}
+			}
+
+			public Task WaitAsync(CancellationToken cancellationToken)
+				=> _task.WaitAsync(cancellationToken);
+
+			public void Cancel()
+			{
+				lock (_task)
+				{
+					_cancellationSource?.Cancel();
+					_cancellationSource = null;
+				}
+			}
+		}
+
+		#region Json serialization
+
+		class JsonState
+		{
+			public bool AutoUpdate { get; set; }
+			public List<JsonItemState> Items { get; set; } = new List<JsonItemState>();
+		}
+
+		class JsonItemState
+		{
+			public ToolId Id { get; set; }
+			public string Name { get; set; }
+			public string Description { get; set; }
+			public JsonDeploymentState? Current { get; set; }
+
+			public JsonItemState()
+			{
+				Name = String.Empty;
+				Description = String.Empty;
+			}
+
+			public JsonItemState(ToolId id, ItemState state)
+			{
+				Id = id;
+				Name = state.Name;
+				Description = state.Description;
+				Current = (state.Current == null) ? null : new JsonDeploymentState(state.Current);
+			}
+		}
+
+		class JsonDeploymentState
+		{
+			public ToolDeploymentId Id { get; set; }
+			public string Version { get; set; }
+
+			public JsonDeploymentState()
+			{
+				Version = String.Empty;
+			}
+
+			public JsonDeploymentState(ToolDeploymentInfo current)
+			{
+				Id = current.Id;
+				Version = current.Version;
+			}
+
+			public ToolDeploymentInfo ToToolDeploymentInfo()
+				=> new ToolDeploymentInfo(Id, Version);
+		}
+
+		#endregion
+
+		readonly IHordeClientProvider _hordeClientProvider;
+		readonly ILogger _logger;
+		readonly object _workerThreadLockObject = new object();
+		readonly DirectoryReference _baseDir;
+		readonly FileReference _configFile;
+		readonly AsyncEvent _updateEvent;
+		readonly BackgroundTask _updateTask;
+		readonly AsyncEvent _cleanupEvent;
+		readonly BackgroundTask _cleanupTask;
+
+		readonly SynchronizationContext? _mainThreadSynchronizationContext;
+
+		IReadOnlyDictionary<ToolId, Item> _mainThreadItems;
+		IReadOnlyList<IToolCatalogItem> _mainThreadItemList;
+		IReadOnlyList<IToolCatalogItem> IToolCatalog.Items => _mainThreadItemList;
+
+		IReadOnlyDictionary<ToolId, Item> _items;
+
+		Uri? _serverUri;
+		bool _autoUpdate = true;
+
+		public DirectoryReference BaseDir => _baseDir;
+
+		public event Action? OnItemsChanged;
+
+		static readonly JsonSerializerOptions s_jsonSerializerOptions = new JsonSerializerOptions { AllowTrailingCommas = true, PropertyNameCaseInsensitive = true, PropertyNamingPolicy = JsonNamingPolicy.CamelCase, WriteIndented = true };
+
+		public bool AutoUpdate
+		{
+			get => _autoUpdate;
+			set
+			{
+				if (_autoUpdate != value)
+				{
+					_autoUpdate = value;
+					_ = SaveStateAsync(CancellationToken.None);
+				}
+			}
+		}
+
+		public ToolCatalog(IHordeClientProvider hordeClientProvider, ILogger<ToolCatalog> logger)
+		{
+			DirectoryReference? toolsDir = DirectoryReference.GetSpecialFolder(Environment.SpecialFolder.LocalApplicationData);
+			toolsDir ??= DirectoryReference.GetCurrentDirectory();
+
+			_hordeClientProvider = hordeClientProvider;
+			_logger = logger;
+			_baseDir = DirectoryReference.Combine(toolsDir, "Epic Games", "Horde", "TrayApp");
+			_configFile = FileReference.Combine(_baseDir, "Tools.json");
+			_updateEvent = new AsyncEvent();
+			_updateTask = new BackgroundTask(UpdateAsync);
+			_cleanupEvent = new AsyncEvent();
+			_cleanupTask = new BackgroundTask(CleanupAsync);
+			_mainThreadSynchronizationContext = SynchronizationContext.Current;
+
+			_items = LoadState();
+			_mainThreadItems = null!;
+			_mainThreadItemList = null!;
+
+			UpdateMainThread();
+		}
+
+		void PostMainThreadUpdate()
+		{
+			if (_mainThreadSynchronizationContext == null)
+			{
+				UpdateMainThread();
+			}
+			else
+			{
+				_mainThreadSynchronizationContext.Post(_ => UpdateMainThread(), null);
+			}
+		}
+
+		void UpdateMainThread()
+		{
+			bool updatedCollection = false;
+			List<Item> updatedItems = new List<Item>();
+
+			lock (_workerThreadLockObject)
+			{
+				if (_items != _mainThreadItems)
+				{
+					UpdateMainThreadItemsState();
+					updatedCollection = true;
+				}
+
+				foreach (Item item in _items.Values)
+				{
+					if (item._state != item._mainThreadState)
+					{
+						item.UpdateMainThreadState();
+						updatedItems.Add(item);
+					}
+				}
+			}
+
+			if (updatedCollection)
+			{
+				OnItemsChanged?.Invoke();
+			}
+			foreach (Item updatedItem in updatedItems)
+			{
+				updatedItem.NotifyChanged();
+			}
+		}
+
+		void UpdateMainThreadItemsState()
+		{
+			_mainThreadItems = _items;
+			_mainThreadItemList = _items.Values.OrderBy(x => x._mainThreadState.Name).ToArray();
+		}
+
+		public void Start()
+		{
+			_updateTask.Start();
+			_cleanupTask.Start();
+
+			_hordeClientProvider.OnStateChanged += OnServerChange;
+		}
+
+		public async Task StopAsync(CancellationToken cancellationToken)
+		{
+			_hordeClientProvider.OnStateChanged -= OnServerChange;
+
+			foreach (Item item in _items.Values)
+			{
+				DeploymentTask? deploymentTask = item._deploymentTask;
+				if (deploymentTask != null)
+				{
+					deploymentTask.Cancel();
+					await deploymentTask.WaitAsync(cancellationToken);
+				}
+			}
+
+			await _updateTask.StopAsync(cancellationToken);
+			await _cleanupTask.StopAsync(cancellationToken);
+		}
+
+		public async ValueTask DisposeAsync()
+		{
+			await StopAsync(CancellationToken.None);
+			await _updateTask.DisposeAsync();
+			await _cleanupTask.DisposeAsync();
+		}
+
+		void OnServerChange()
+		{
+			_updateEvent.Pulse();
+		}
+
+		DirectoryReference GetToolDir(ToolId toolId, ToolDeploymentId deploymentId)
+		{
+			return DirectoryReference.Combine(_baseDir, toolId.ToString(), deploymentId.ToString());
+		}
+
+		ToolConfig LoadToolConfig(DirectoryReference toolDir)
+		{
+			FileReference configFile = FileReference.Combine(toolDir, "Launcher.json");
+			if (FileReference.Exists(configFile))
+			{
+				JsonSerializerOptions jsonSerializerOptions = new JsonSerializerOptions();
+				jsonSerializerOptions.PropertyNameCaseInsensitive = true;
+				jsonSerializerOptions.AllowTrailingCommas = true;
+				jsonSerializerOptions.ReadCommentHandling = JsonCommentHandling.Skip;
+
+				try
+				{
+					byte[] data = FileReference.ReadAllBytes(configFile);
+					return JsonSerializer.Deserialize<ToolConfig>(data, jsonSerializerOptions) ?? new ToolConfig();
+				}
+				catch (Exception ex)
+				{
+					_logger.LogError(ex, "Error reading {File}: {Message}", configFile, ex.Message);
+				}
+			}
+			return new ToolConfig();
+		}
+
+		IReadOnlyDictionary<ToolId, Item> LoadState()
+		{
+			Dictionary<ToolId, Item> items = new Dictionary<ToolId, Item>();
+
+			byte[]? data = FileTransaction.ReadAllBytes(_configFile);
+			if (data != null)
+			{
+				JsonState? state = JsonSerializer.Deserialize<JsonState>(data, s_jsonSerializerOptions);
+				if (state != null)
+				{
+					_autoUpdate = state.AutoUpdate;
+					foreach (JsonItemState jsonItem in state.Items)
+					{
+						CurrentToolDeploymentInfo? current = null;
+						if (jsonItem.Current != null)
+						{
+							DirectoryReference toolDir = GetToolDir(jsonItem.Id, jsonItem.Current.Id);
+							ToolConfig toolConfig = LoadToolConfig(toolDir);
+							current = new CurrentToolDeploymentInfo(jsonItem.Current.Id, jsonItem.Current.Version, toolDir, toolConfig);
+						}
+
+						ItemState itemState = new ItemState(jsonItem.Name, jsonItem.Description, current, null, null);
+						items.Add(jsonItem.Id, new Item(this, jsonItem.Id, itemState));
+					}
+				}
+			}
+
+			return items;
+		}
+
+		async Task SaveStateAsync(CancellationToken cancellationToken)
+		{
+			JsonState state;
+			lock (_workerThreadLockObject)
+			{
+				state = new JsonState();
+				state.AutoUpdate = _autoUpdate;
+				foreach (Item item in _items.Values)
+				{
+					if (item._state.Current != null)
+					{
+						state.Items.Add(new JsonItemState(item.Id, item._state));
+					}
+				}
+			}
+
+			byte[] data = JsonSerializer.SerializeToUtf8Bytes(state, s_jsonSerializerOptions);
+
+			DirectoryReference.CreateDirectory(_configFile.Directory);
+			await FileTransaction.WriteAllBytesAsync(_configFile, data, cancellationToken);
+		}
+
+		void CancelDeployment(Item item, ToolDeploymentInfo? deployment)
+		{
+			if (deployment != null)
+			{
+				lock (_workerThreadLockObject)
+				{
+					if (item._deploymentTask != null && item._deploymentTask.Deployment == deployment)
+					{
+						item._deploymentTask.Cancel();
+					}
+				}
+			}
+		}
+
+		void UpdateDeployment(Item item, IToolDeployment? deployment)
+		{
+			lock (_workerThreadLockObject)
+			{
+				DeploymentTask? lastDeploymentTask = item._deploymentTask;
+				if (deployment != null)
+				{
+					ToolDeploymentInfo deploymentInfo = new ToolDeploymentInfo(deployment.Id, deployment.Version);
+					item._deploymentTask = new DeploymentTask(deploymentInfo, ctx => UpdateCurrentDeploymentAsync(item, deployment, deploymentInfo, ctx), lastDeploymentTask);
+				}
+				else if (item._state.Current != null)
+				{
+					ToolDeploymentInfo deploymentInfo = item._state.Current;
+					item._deploymentTask = new DeploymentTask(deploymentInfo, ctx => UpdateCurrentDeploymentAsync(item, deployment, deploymentInfo, ctx), lastDeploymentTask);
+				}
+			}
+		}
+
+		async Task UpdateCurrentDeploymentAsync(Item item, IToolDeployment? deployment, ToolDeploymentInfo deploymentInfo, CancellationToken cancellationToken)
+		{
+			try
+			{
+				if (deployment == null)
+				{
+					await UninstallAsync(item, cancellationToken);
+					_cleanupEvent.Set();
+				}
+				else
+				{
+					UpdateStateAndNotify(item, state => state with { Pending = new PendingToolDeploymentInfo(false, $"Updating to {deployment.Version}", deploymentInfo) });
+
+					DirectoryReference toolDir = GetToolDir(item.Id, deploymentInfo.Id);
+					DirectoryReference.CreateDirectory(toolDir);
+					await deployment.Content.ExtractAsync(toolDir.ToDirectoryInfo(), _logger, cancellationToken);
+
+					await InstallAsync(item, toolDir, deploymentInfo, cancellationToken);
+				}
+
+				await SaveStateAsync(cancellationToken);
+			}
+			catch (OperationCanceledException)
+			{
+				_logger.LogInformation("Deployment cancelled");
+				UpdateStateAndNotify(item, state => state with { Pending = new PendingToolDeploymentInfo(true, $"Cancelled.", deploymentInfo) });
+			}
+			catch (Exception ex)
+			{
+				_logger.LogError(ex, "Error updating deployment: {Message}", ex.Message);
+				UpdateStateAndNotify(item, state => state with { Pending = new PendingToolDeploymentInfo(true, $"Error updating. See log.", deploymentInfo) });
+			}
+		}
+
+		async Task InstallAsync(Item item, DirectoryReference toolDir, ToolDeploymentInfo deploymentInfo, CancellationToken cancellationToken)
+		{
+			ToolConfig toolConfig = LoadToolConfig(toolDir);
+
+			await UninstallAsync(item, cancellationToken);
+
+			if (toolConfig.InstallCommand != null)
+			{
+				int exitCode = await RunCommandAsync(item.Id.ToString(), toolConfig.InstallCommand, toolDir.FullName, cancellationToken);
+				if (exitCode != 0)
+				{
+					PendingToolDeploymentInfo pending = new PendingToolDeploymentInfo(true, $"Installation failed ({exitCode}). See log.", deploymentInfo);
+					UpdateStateAndNotify(item, state => state with { Pending = pending });
+					return;
+				}
+			}
+
+			CurrentToolDeploymentInfo current = new CurrentToolDeploymentInfo(deploymentInfo.Id, deploymentInfo.Version, toolDir, toolConfig);
+			UpdateStateAndNotify(item, state => state with { Pending = null, Current = current });
+		}
+
+		async Task UninstallAsync(Item item, CancellationToken cancellationToken)
+		{
+			CurrentToolDeploymentInfo? current = null;
+			UpdateStateAndNotify(item, state =>
+			{
+				current = state.Current;
+
+				PendingToolDeploymentInfo? pending = null;
+				if (current != null)
+				{
+					pending = new PendingToolDeploymentInfo(false, $"Removing {current.Version}", null);
+				}
+
+				return state with { Pending = pending, Current = null };
+			});
+
+			if (current != null)
+			{
+				ToolCommand? uninstallCommand = current.Config.UninstallCommand;
+				if (uninstallCommand != null)
+				{
+					await RunCommandAsync(current.Id.ToString(), uninstallCommand, current.Dir.FullName, cancellationToken);
+				}
+			}
+
+			UpdateStateAndNotify(item, state => state with { Pending = null, Current = null });
+		}
+
+		async Task<int> RunCommandAsync(string toolName, ToolCommand command, string workingDir, CancellationToken cancellationToken)
+		{
+			string arguments = String.Empty;
+			if (command.Arguments != null)
+			{
+				arguments = CommandLineArguments.Join(command.Arguments);
+			}
+
+			using (ManagedProcessGroup newProcessGroup = new ManagedProcessGroup())
+			using (ManagedProcess newProcess = new ManagedProcess(newProcessGroup, command.FileName, arguments, workingDir, null, null, ProcessPriorityClass.Normal))
+			{
+				for (; ; )
+				{
+					string? line = await newProcess.ReadLineAsync(cancellationToken);
+					if (line == null)
+					{
+						await newProcess.WaitForExitAsync(cancellationToken);
+						_logger.LogInformation("{ToolName}> Exit code {ExitCode}", toolName, newProcess.ExitCode);
+						return newProcess.ExitCode;
+					}
+					_logger.LogInformation("{ToolName}> {Line}", toolName, line);
+				}
+			}
+		}
+
+		void UpdateStateAndNotify(Item item, Func<ItemState, ItemState> update)
+		{
+			lock (_workerThreadLockObject)
+			{
+				item._state = update(item._state);
+			}
+			PostMainThreadUpdate();
+		}
+
+		#region Cleanup
+
+		async Task CleanupAsync(CancellationToken cancellationToken)
+		{
+			while (!cancellationToken.IsCancellationRequested)
+			{
+				Task cleanupTask = _cleanupEvent.Task;
+
+				bool success = false;
+				try
+				{
+					success = CleanupOnce();
+				}
+				catch (OperationCanceledException)
+				{
+					break;
+				}
+				catch (Exception ex)
+				{
+					_logger.LogError(ex, "Exception while running cleanup over installed tools: {Message}", ex.Message);
+				}
+
+				Task waitTask = cleanupTask.WaitAsync(cancellationToken);
+				if (!success)
+				{
+					waitTask = Task.WhenAny(waitTask, Task.Delay(TimeSpan.FromMinutes(1.0), cancellationToken));
+				}
+
+				await waitTask;
+			}
+		}
+
+		bool CleanupOnce()
+		{
+			bool result = true;
+
+			List<DirectoryInfo> deleteToolDirs = new List<DirectoryInfo>();
+			List<DirectoryInfo> deleteDeploymentDirs = new List<DirectoryInfo>();
+
+			DirectoryInfo baseDirInfo = _baseDir.ToDirectoryInfo();
+			if (baseDirInfo.Exists)
+			{
+				foreach (DirectoryInfo toolDir in baseDirInfo.EnumerateDirectories())
+				{
+					if (!toolDir.Name.StartsWith('.'))
+					{
+						DirectoryInfo[] deploymentDirs = toolDir.GetDirectories();
+
+						ToolId? toolId;
+						try
+						{
+							toolId = new ToolId(toolDir.Name);
+						}
+						catch
+						{
+							toolId = null;
+						}
+
+						lock (_workerThreadLockObject)
+						{
+							ItemState? itemState = null;
+							if (toolId != null && _items.TryGetValue(toolId.Value, out Item? item))
+							{
+								itemState = item._state;
+							}
+
+							bool deleteToolDir = true;
+							for (int idx = 0; idx < deploymentDirs.Length; idx++)
+							{
+								DirectoryInfo deploymentDir = deploymentDirs[idx];
+								if (!deploymentDir.Name.StartsWith('.'))
+								{
+									if (!BinaryId.TryParse(deploymentDir.Name, out BinaryId deploymentId))
+									{
+										deleteDeploymentDirs.Add(deploymentDir);
+									}
+									else if (itemState != null && itemState.Current != null && itemState.Current.Id.Id == deploymentId)
+									{
+										deleteToolDir = false;
+									}
+									else if (itemState != null && itemState.Pending?.Deployment != null && itemState.Pending.Deployment.Id.Id == deploymentId)
+									{
+										deleteToolDir = false;
+									}
+									else if (TryRenameDir(deploymentDir, deploymentDir.FullName + ".dead"))
+									{
+										deleteDeploymentDirs.Add(deploymentDir);
+									}
+									else
+									{
+										result = deleteToolDir = false;
+									}
+								}
+							}
+
+							if (deleteToolDir)
+							{
+								deleteToolDirs.Add(toolDir);
+							}
+						}
+					}
+				}
+			}
+
+			// Try to delete all the deployment directories
+			foreach (DirectoryInfo deleteDeploymentDir in deleteDeploymentDirs)
+			{
+				result &= TryDeleteDir(deleteDeploymentDir, true);
+			}
+
+			// Try to delete all the tool directories if they're empty
+			foreach (DirectoryInfo deleteToolDir in deleteToolDirs)
+			{
+				result &= TryDeleteDir(deleteToolDir, false);
+			}
+
+			return result;
+		}
+
+		bool TryRenameDir(DirectoryInfo dirInfo, string newName)
+		{
+			try
+			{
+				dirInfo.MoveTo(newName);
+				return true;
+			}
+			catch (Exception ex)
+			{
+				_logger.LogWarning(ex, "Unable to rename directory {DirName}: {Message}", dirInfo.FullName, ex.Message);
+				return false;
+			}
+		}
+
+		bool TryDeleteDir(DirectoryInfo dirInfo, bool recursive)
+		{
+			try
+			{
+				dirInfo.Delete(recursive);
+				return true;
+			}
+			catch (Exception ex)
+			{
+				_logger.LogWarning(ex, "Unable to delete directory {DirName}: {Message}", dirInfo.FullName, ex.Message);
+				return false;
+			}
+		}
+
+		#endregion
+
+		#region Updates
+
+		async Task UpdateAsync(CancellationToken cancellationToken)
+		{
+			while (!cancellationToken.IsCancellationRequested)
+			{
+				Task updateTask = _updateEvent.Task;
+				try
+				{
+					await PollForUpdatesOnceAsync(cancellationToken);
+				}
+				catch (OperationCanceledException)
+				{
+					break;
+				}
+				catch (Exception ex)
+				{
+					_logger.LogError(ex, "Exception while checking for tool updates: {Message}", ex.Message);
+				}
+				await Task.WhenAny(updateTask, Task.Delay(TimeSpan.FromMinutes(5.0), cancellationToken));
+			}
+		}
+
+		async Task PollForUpdatesOnceAsync(CancellationToken cancellationToken)
+		{
+			using IHordeClientRef? hordeClientRef = await _hordeClientProvider.GetClientRefAsync();
+			if (hordeClientRef == null)
+			{
+				SetTools(Enumerable.Empty<ITool>());
+			}
+			else
+			{
+				if (_serverUri != hordeClientRef.Client.ServerUrl)
+				{
+					SetTools(Enumerable.Empty<ITool>());
+					_serverUri = hordeClientRef.Client.ServerUrl;
+				}
+
+				IEnumerable<ITool> tools = await hordeClientRef.Client.Tools.GetAllAsync(cancellationToken);
+				SetTools(tools);
+			}
+		}
+
+		void SetTools(IEnumerable<ITool> tools)
+		{
+			lock (_workerThreadLockObject)
+			{
+				Dictionary<ToolId, Item> newItems = new Dictionary<ToolId, Item>();
+				foreach (ITool tool in tools)
+				{
+					if (ShouldIncludeTool(tool))
+					{
+						IToolDeployment? latestDeployment = (tool.Deployments.Count > 0) ? tool.Deployments[^1] : null;
+						if (_items.TryGetValue(tool.Id, out Item? item))
+						{
+							item._state = new ItemState(tool.Name, tool.Description, item._state.Current, item._state.Pending, latestDeployment);
+						}
+						else
+						{
+							item = new Item(this, tool.Id, new ItemState(tool.Name, tool.Description, null, null, latestDeployment));
+						}
+						newItems[item.Id] = item;
+					}
+				}
+				foreach (Item existingItem in _items.Values)
+				{
+					if (existingItem._state.Current != null && !newItems.ContainsKey(existingItem.Id))
+					{
+						newItems.Add(existingItem.Id, existingItem);
+					}
+				}
+				_items = newItems;
+
+				if (_autoUpdate)
+				{
+					foreach (Item item in _items.Values)
+					{
+						if (item._state.Current != null && item._state.Latest != null && item._state.Latest.Id != item._state.Current.Id)
+						{
+							UpdateDeployment(item, item._state.Latest);
+						}
+					}
+				}
+			}
+			PostMainThreadUpdate();
+		}
+
+		static IReadOnlySet<string> HostPlatformRids { get; } = GetHostPlatformRids();
+
+		static bool ShouldIncludeTool(ITool tool)
+		{
+			if (tool.Platforms != null && !tool.Platforms.Any(x => HostPlatformRids.Contains(x)))
+			{
+				return false;
+			}
+			return true; // TODO: ShowInUgs?
+		}
+
+		static HashSet<string> GetHostPlatformRids()
+		{
+			HashSet<string> platforms = new HashSet<string>(StringComparer.OrdinalIgnoreCase);
+			platforms.Add("any");
+			if (OperatingSystem.IsWindows())
+			{
+				platforms.Add("win");
+				switch (RuntimeInformation.OSArchitecture)
+				{
+					case Architecture.X64:
+						platforms.Add("win-x64");
+						break;
+					case Architecture.Arm64:
+						platforms.Add("win-arm64");
+						break;
+				}
+			}
+			else if (OperatingSystem.IsMacOS())
+			{
+				platforms.Add("osx");
+				switch (RuntimeInformation.OSArchitecture)
+				{
+					case Architecture.X64:
+						platforms.Add("osx-x64");
+						break;
+					case Architecture.Arm64:
+						platforms.Add("osx-arm64");
+						break;
+				}
+			}
+			else if (OperatingSystem.IsLinux())
+			{
+				platforms.Add("linux");
+				platforms.Add("unix");
+				switch (RuntimeInformation.OSArchitecture)
+				{
+					case Architecture.X64:
+						platforms.Add("linux-x64");
+						platforms.Add("unix-x64");
+						break;
+					case Architecture.Arm64:
+						platforms.Add("linux-arm64");
+						platforms.Add("unix-arm64");
+						break;
+				}
+			}
+			return platforms;
+		}
+
+		#endregion
+	}
+}
diff --git a/Engine/Source/Programs/Horde/HordeToolbox/ToolCatalogViewModel.cs b/Engine/Source/Programs/Horde/HordeToolbox/ToolCatalogViewModel.cs
new file mode 100644
--- /dev/null
+++ b/Engine/Source/Programs/Horde/HordeToolbox/ToolCatalogViewModel.cs
@@ -0,0 +1,172 @@
+// Copyright Epic Games, Inc. All Rights Reserved.
+
+using System.Collections.ObjectModel;
+using System.ComponentModel;
+using CommunityToolkit.Mvvm.ComponentModel;
+
+namespace HordeToolbox
+{
+	partial class ToolCatalogViewModel : ObservableObject, IDisposable
+	{
+		readonly IToolCatalog? _catalog;
+		readonly Dictionary<IToolCatalogItem, ToolCatalogItemViewModel> _items = new Dictionary<IToolCatalogItem, ToolCatalogItemViewModel>();
+
+		public ObservableCollection<ToolCatalogItemViewModel> Items { get; } = new ObservableCollection<ToolCatalogItemViewModel>();
+
+		public bool ShowTools => Items.Count > 0;
+
+		[ObservableProperty]
+		bool _autoUpdate = true;
+
+		public ToolCatalogViewModel()
+			: this(null)
+		{ }
+
+		public ToolCatalogViewModel(IToolCatalog? catalog)
+		{
+			if (catalog != null)
+			{
+				_catalog = catalog;
+				_catalog.OnItemsChanged += RefreshItems;
+
+				Items.CollectionChanged += (_, _) => OnPropertyChanged(nameof(ShowTools));
+
+				RefreshItems();
+			}
+		}
+
+		protected override void OnPropertyChanged(PropertyChangedEventArgs e)
+		{
+			base.OnPropertyChanged(e);
+
+			if (_catalog != null)
+			{
+				_catalog.AutoUpdate = AutoUpdate;
+			}
+		}
+
+		public void Dispose()
+		{
+			foreach (ToolCatalogItemViewModel item in _items.Values)
+			{
+				item.Dispose();
+			}
+			if (_catalog != null)
+			{
+				_catalog.OnItemsChanged -= RefreshItems;
+			}
+		}
+
+		void RefreshItems()
+		{
+			if (_catalog != null)
+			{
+				Items.Clear();
+
+				HashSet<IToolCatalogItem> newItemSet = new HashSet<IToolCatalogItem>(_catalog.Items);
+				foreach ((IToolCatalogItem item, ToolCatalogItemViewModel itemViewModel) in _items.ToArray())
+				{
+					if (!newItemSet.Contains(item))
+					{
+						_items.Remove(item);
+						itemViewModel.Dispose();
+					}
+				}
+
+				foreach (IToolCatalogItem item in _catalog.Items)
+				{
+					ToolCatalogItemViewModel? itemViewModel;
+					if (!_items.TryGetValue(item, out itemViewModel))
+					{
+						itemViewModel = new ToolCatalogItemViewModel(item);
+						_items.Add(item, itemViewModel);
+					}
+					Items.Add(itemViewModel);
+				}
+			}
+		}
+	}
+
+	partial class ToolCatalogItemViewModel : ObservableObject, IDisposable
+	{
+		readonly IToolCatalogItem _item;
+
+		public ToolCatalogItemViewModel(IToolCatalogItem item)
+		{
+			_item = item;
+			_item.OnItemChanged += UpdateFromModel;
+			UpdateFromModel();
+		}
+
+		public void Dispose()
+		{
+			_item.OnItemChanged -= UpdateFromModel;
+		}
+
+		[ObservableProperty]
+		string _name = String.Empty;
+
+		[ObservableProperty]
+		string _description = String.Empty;
+
+		[ObservableProperty]
+		string _status = String.Empty;
+
+		[ObservableProperty]
+		bool _showAdd;
+
+		[ObservableProperty]
+		bool _showUpdate;
+
+		[ObservableProperty]
+		bool _showRemove;
+
+		internal void UpdateFromModel()
+		{
+			Name = _item.Name;
+			Description = _item.Description;
+
+			PendingToolDeploymentInfo? pending = _item.Pending;
+			if (pending != null)
+			{
+				if (pending.Message != null)
+				{
+					Status = pending.Message;
+				}
+				else if (pending.Deployment == null)
+				{
+					Status = "Removing...";
+				}
+				else
+				{
+					Status = $"Updating to {pending.Deployment.Version}...";
+				}
+			}
+			else
+			{
+				ToolDeploymentInfo? current = _item.Current;
+				if (current != null)
+				{
+					Status = current.Version;
+				}
+				else
+				{
+					Status = String.Empty;
+				}
+			}
+
+			ShowAdd = _item.Latest != null && _item.Current == null;
+			ShowUpdate = _item.Latest != null && _item.Current != null && _item.Latest.Id != _item.Current.Id;
+			ShowRemove = _item.Current != null;
+		}
+
+		public void Add()
+			=> _item.Install();
+
+		public void Remove()
+			=> _item.Uninstall();
+
+		public void Update()
+			=> _item.Install();
+	}
+}
diff --git a/Engine/Source/Programs/Horde/HordeToolbox/ToolConfig.cs b/Engine/Source/Programs/Horde/HordeToolbox/ToolConfig.cs
new file mode 100644
--- /dev/null
+++ b/Engine/Source/Programs/Horde/HordeToolbox/ToolConfig.cs
@@ -0,0 +1,62 @@
+// Copyright Epic Games, Inc. All Rights Reserved.
+
+namespace HordeToolbox
+{
+	/// <summary>
+	/// Configuration for a tool that can be shown in the launcher
+	/// </summary>
+	public class ToolConfig
+	{
+		/// <summary>
+		/// Popup menu item for the tool
+		/// </summary>
+		public ToolMenuItem? PopupMenu { get; set; }
+
+		/// <summary>
+		/// Command to run when installing
+		/// </summary>
+		public ToolCommand? InstallCommand { get; set; }
+
+		/// <summary>
+		/// Command to run when uninstalling
+		/// </summary>
+		public ToolCommand? UninstallCommand { get; set; }
+	}
+
+	/// <summary>
+	/// Menu item for the tool
+	/// </summary>
+	public class ToolMenuItem
+	{
+		/// <summary>
+		/// Text to display for the menu item. If null, or a sequence of hyphens, will create a menu separator.
+		/// </summary>
+		public string? Label { get; set; }
+
+		/// <summary>
+		/// Command to run when the menu item is clicked. Will be executed in the root directory of the downloaded tool.
+		/// </summary>
+		public ToolCommand? Command { get; set; }
+
+		/// <summary>
+		/// Child menu items. Cannot be specified at the same time as a tool.
+		/// </summary>
+		public List<ToolMenuItem>? Children { get; set; }
+	}
+
+	/// <summary>
+	/// Command to run for a menu item.
+	/// </summary>
+	public class ToolCommand
+	{
+		/// <summary>
+		/// Executable to run
+		/// </summary>
+		public string FileName { get; set; } = "cmd.exe";
+
+		/// <summary>
+		/// Command line arguments for the tool
+		/// </summary>
+		public List<string>? Arguments { get; set; }
+	}
+}
diff --git a/Engine/Source/Programs/Horde/HordeTrayApp/App.axaml.cs b/Engine/Source/Programs/Horde/HordeTrayApp/App.axaml.cs
deleted file mode 100644
diff --git a/Engine/Source/Programs/Horde/HordeTrayApp/GeneralSettingsPage.axaml.cs b/Engine/Source/Programs/Horde/HordeTrayApp/GeneralSettingsPage.axaml.cs
deleted file mode 100644
diff --git a/Engine/Source/Programs/Horde/HordeTrayApp/GeneralSettingsViewModel.cs b/Engine/Source/Programs/Horde/HordeTrayApp/GeneralSettingsViewModel.cs
deleted file mode 100644
diff --git a/Engine/Source/Programs/Horde/HordeTrayApp/HordeClientProvider.cs b/Engine/Source/Programs/Horde/HordeTrayApp/HordeClientProvider.cs
deleted file mode 100644
diff --git a/Engine/Source/Programs/Horde/HordeTrayApp/HordeTrayApp.csproj b/Engine/Source/Programs/Horde/HordeTrayApp/HordeTrayApp.csproj
deleted file mode 100644
diff --git a/Engine/Source/Programs/Horde/HordeTrayApp/IHordeClientProvider.cs b/Engine/Source/Programs/Horde/HordeTrayApp/IHordeClientProvider.cs
deleted file mode 100644
diff --git a/Engine/Source/Programs/Horde/HordeTrayApp/IToolCatalog.cs b/Engine/Source/Programs/Horde/HordeTrayApp/IToolCatalog.cs
deleted file mode 100644
diff --git a/Engine/Source/Programs/Horde/HordeTrayApp/IToolCatalogItem.cs b/Engine/Source/Programs/Horde/HordeTrayApp/IToolCatalogItem.cs
deleted file mode 100644
diff --git a/Engine/Source/Programs/Horde/HordeTrayApp/ITrayAppHost.cs b/Engine/Source/Programs/Horde/HordeTrayApp/ITrayAppHost.cs
deleted file mode 100644
diff --git a/Engine/Source/Programs/Horde/HordeTrayApp/ITrayAppPlugin.cs b/Engine/Source/Programs/Horde/HordeTrayApp/ITrayAppPlugin.cs
deleted file mode 100644
diff --git a/Engine/Source/Programs/Horde/HordeTrayApp/JsonConfig.cs b/Engine/Source/Programs/Horde/HordeTrayApp/JsonConfig.cs
deleted file mode 100644
diff --git a/Engine/Source/Programs/Horde/HordeTrayApp/Plugins/HordeAgent/HordeAgentPlugin.cs b/Engine/Source/Programs/Horde/HordeTrayApp/Plugins/HordeAgent/HordeAgentPlugin.cs
deleted file mode 100644
diff --git a/Engine/Source/Programs/Horde/HordeTrayApp/Plugins/HordeAgent/HordeAgentSettings.cs b/Engine/Source/Programs/Horde/HordeTrayApp/Plugins/HordeAgent/HordeAgentSettings.cs
deleted file mode 100644
diff --git a/Engine/Source/Programs/Horde/HordeTrayApp/Plugins/HordeAgent/HordeAgentSettingsPage.axaml.cs b/Engine/Source/Programs/Horde/HordeTrayApp/Plugins/HordeAgent/HordeAgentSettingsPage.axaml.cs
deleted file mode 100644
diff --git a/Engine/Source/Programs/Horde/HordeTrayApp/Plugins/HordeAgent/HordeAgentSettingsViewModel.cs b/Engine/Source/Programs/Horde/HordeTrayApp/Plugins/HordeAgent/HordeAgentSettingsViewModel.cs
deleted file mode 100644
diff --git a/Engine/Source/Programs/Horde/HordeTrayApp/Plugins/HordeProxy/HordeProxyPlugin.cs b/Engine/Source/Programs/Horde/HordeTrayApp/Plugins/HordeProxy/HordeProxyPlugin.cs
deleted file mode 100644
diff --git a/Engine/Source/Programs/Horde/HordeTrayApp/Plugins/HordeProxy/HordeProxySettings.cs b/Engine/Source/Programs/Horde/HordeTrayApp/Plugins/HordeProxy/HordeProxySettings.cs
deleted file mode 100644
diff --git a/Engine/Source/Programs/Horde/HordeTrayApp/Plugins/HordeProxy/HordeProxySettingsPage.axaml.cs b/Engine/Source/Programs/Horde/HordeTrayApp/Plugins/HordeProxy/HordeProxySettingsPage.axaml.cs
deleted file mode 100644
diff --git a/Engine/Source/Programs/Horde/HordeTrayApp/Plugins/HordeProxy/HordeProxySettingsViewModel.cs b/Engine/Source/Programs/Horde/HordeTrayApp/Plugins/HordeProxy/HordeProxySettingsViewModel.cs
deleted file mode 100644
diff --git a/Engine/Source/Programs/Horde/HordeTrayApp/Program.cs b/Engine/Source/Programs/Horde/HordeTrayApp/Program.cs
deleted file mode 100644
diff --git a/Engine/Source/Programs/Horde/HordeTrayApp/SettingsWindow.axaml.cs b/Engine/Source/Programs/Horde/HordeTrayApp/SettingsWindow.axaml.cs
deleted file mode 100644
diff --git a/Engine/Source/Programs/Horde/HordeTrayApp/ToolCatalog.cs b/Engine/Source/Programs/Horde/HordeTrayApp/ToolCatalog.cs
deleted file mode 100644
diff --git a/Engine/Source/Programs/Horde/HordeTrayApp/ToolCatalogViewModel.cs b/Engine/Source/Programs/Horde/HordeTrayApp/ToolCatalogViewModel.cs
deleted file mode 100644
diff --git a/Engine/Source/Programs/Horde/HordeTrayApp/ToolConfig.cs b/Engine/Source/Programs/Horde/HordeTrayApp/ToolConfig.cs
deleted file mode 100644
