Horde: Fix warnings in CI.

#rnx

diff --git a/Engine/Source/Programs/Horde/BuildHorde.xml b/Engine/Source/Programs/Horde/BuildHorde.xml
--- a/Engine/Source/Programs/Horde/BuildHorde.xml
+++ b/Engine/Source/Programs/Horde/BuildHorde.xml
@@ -180,6 +180,7 @@
 			<Tag Files="$(ServerToolsStagingDir_HordeToolbox)/..." With="#HordeToolboxFiles"/>
 		</Node>
 		<Node Name="Build HordeToolbox Msi" Requires="Build HordeToolbox" Produces="#HordeToolboxMsiFiles">
+			<Property Name="StagingDir" Value="$(RootDir)/Staging"/>
 			<Property Name="CommonArgs" Value="-dPlatform=x64 -arch x64"/>
 
 			<!-- Create the toolbox file list -->
