Horde: Rename "ShowInLauncher" parameter to "ShowInToolbox"

#rnx

diff --git a/Engine/Source/Programs/Horde/BuildHorde.xml b/Engine/Source/Programs/Horde/BuildHorde.xml
--- a/Engine/Source/Programs/Horde/BuildHorde.xml
+++ b/Engine/Source/Programs/Horde/BuildHorde.xml
@@ -304,17 +304,17 @@
 				<Delete Files="$(StagingDir)/Server/$(ServerArch)/appsettings.Local.json"/>
 
 				<!-- Add the bundled tools into the installer -->
-				<DotNet Arguments="run --project &quot;$(RootDir)/Engine/Source/Programs/Horde/HordeServer/HordeServer.csproj&quot; -- generate tooldata -id=horde-cmd -name=&quot;Horde (Command-Line)&quot; -description=&quot;Update package for the Horde command-line tool.&quot; -version=&quot;$(InformationalVersion)&quot; -inputdir=&quot;$(ServerToolsStagingDir_HordeCmd)&quot; -serverdir=&quot;$(StagingDir)/Server/$(ServerArch)&quot; -public -showinugs -showinlauncher -Category=Internal" DotNetPath="$(DotNetPath)" />
+				<DotNet Arguments="run --project &quot;$(RootDir)/Engine/Source/Programs/Horde/HordeServer/HordeServer.csproj&quot; -- generate tooldata -id=horde-cmd -name=&quot;Horde (Command-Line)&quot; -description=&quot;Update package for the Horde command-line tool.&quot; -version=&quot;$(InformationalVersion)&quot; -inputdir=&quot;$(ServerToolsStagingDir_HordeCmd)&quot; -serverdir=&quot;$(StagingDir)/Server/$(ServerArch)&quot; -public -showinugs -showintoolbox -Category=Internal" DotNetPath="$(DotNetPath)" />
 				<DotNet Arguments="run --project &quot;$(RootDir)/Engine/Source/Programs/Horde/HordeServer/HordeServer.csproj&quot; -- generate tooldata -id=horde-agent -name=&quot;Horde Agent (Cross-Platform)&quot; -description=&quot;Cross-platform Horde Agent distribution. Allows hosts to take on work from the Horde Server.&quot; -version=&quot;$(InformationalVersion)&quot; -inputdir=&quot;$(ServerToolsStagingDir_HordeAgent)&quot; -serverdir=&quot;$(StagingDir)/Server/$(ServerArch)&quot; -public" DotNetPath="$(DotNetPath)" />
 				<DotNet Arguments="run --project &quot;$(RootDir)/Engine/Source/Programs/Horde/HordeServer/HordeServer.csproj&quot; -- generate tooldata -id=horde-agent-win64 -name=&quot;Horde Agent (Windows)&quot; -description=&quot;Update package for older installations of Horde Agent through the MSI installer.&quot; -version=&quot;$(InformationalVersion)&quot; -inputdir=&quot;$(ServerToolsStagingDir_HordeAgentWin64)&quot; -serverdir=&quot;$(StagingDir)/Server/$(ServerArch)&quot; -public -showindashboard=false -category=Internal" DotNetPath="$(DotNetPath)" />
 				<DotNet Arguments="run --project &quot;$(RootDir)/Engine/Source/Programs/Horde/HordeServer/HordeServer.csproj&quot; -- generate tooldata -id=horde-agent-win64-installer -name=&quot;Horde Agent (Windows Installer)&quot; -description=&quot;Installer for the Horde Agent on Windows, with toolbox application.&quot; -version=&quot;$(InformationalVersion)&quot; -inputdir=&quot;$(ServerToolsStagingDir_HordeAgentMsi)&quot; -serverdir=&quot;$(StagingDir)/Server/$(ServerArch)&quot; -public" DotNetPath="$(DotNetPath)" />
 				<DotNet Arguments="run --project &quot;$(RootDir)/Engine/Source/Programs/Horde/HordeServer/HordeServer.csproj&quot; -- generate tooldata -id=horde-toolbox -name=&quot;Horde Toolbox&quot; -description=&quot;Update package for the Horde Toolbox.&quot; -version=&quot;$(InformationalVersion)&quot; -inputdir=&quot;$(ServerToolsStagingDir_UnrealToolbox)&quot; -serverdir=&quot;$(StagingDir)/Server/$(ServerArch)&quot; -public -showinugs -Category=Internal" DotNetPath="$(DotNetPath)" />
 				<DotNet Arguments="run --project &quot;$(RootDir)/Engine/Source/Programs/Horde/HordeServer/HordeServer.csproj&quot; -- generate tooldata -id=horde-toolbox-msi -name=&quot;Horde Toolbox (Windows Installer)&quot; -description=&quot;Windows installer for the Horde Toolbox.&quot; -version=&quot;$(InformationalVersion)&quot; -inputdir=&quot;$(ServerToolsStagingDir_UnrealToolboxMsi)&quot; -serverdir=&quot;$(StagingDir)/Server/$(ServerArch)&quot; -public -showinugs" DotNetPath="$(DotNetPath)" />
-				<DotNet Arguments="run --project &quot;$(RootDir)/Engine/Source/Programs/Horde/HordeServer/HordeServer.csproj&quot; -- generate tooldata -id=p4v-utils -name=&quot;P4V Extensions&quot; -description=&quot;Update package for the P4VUtils tool.&quot; -version=&quot;$(InformationalVersion)&quot; -inputdir=&quot;$(ServerToolsStagingDir_P4VUtils)&quot; -serverdir=&quot;$(StagingDir)/Server/$(ServerArch)&quot; -public -showinugs -showinlauncher -category=Internal" DotNetPath="$(DotNetPath)" />
-				<DotNet Arguments="run --project &quot;$(RootDir)/Engine/Source/Programs/Horde/HordeServer/HordeServer.csproj&quot; -- generate tooldata -id=ugs-cmd -name=&quot;UnrealGameSync (Command-Line)&quot; -description=&quot;Command-line version of UnrealGameSync. Supports Windows, MacOS and Linux.&quot; -version=&quot;$(InformationalVersion)&quot; -inputdir=&quot;$(ServerToolsStagingDir_UnrealGameSyncCmd)&quot; -serverdir=&quot;$(StagingDir)/Server/$(ServerArch)&quot; -public -showinlauncher" DotNetPath="$(DotNetPath)" />
-				<DotNet Arguments="run --project &quot;$(RootDir)/Engine/Source/Programs/Horde/HordeServer/HordeServer.csproj&quot; -- generate tooldata -id=ugs-msi -name=&quot;UnrealGameSync (Windows Installer)&quot; -description=&quot;Installer for UnrealGameSync.&quot; -version=&quot;$(InformationalVersion)&quot; -inputdir=&quot;$(ServerToolsStagingDir_UnrealGameSyncMsi)&quot; -serverdir=&quot;$(StagingDir)/Server/$(ServerArch)&quot; -public -showinlauncher" DotNetPath="$(DotNetPath)" />
+				<DotNet Arguments="run --project &quot;$(RootDir)/Engine/Source/Programs/Horde/HordeServer/HordeServer.csproj&quot; -- generate tooldata -id=p4v-utils -name=&quot;P4V Extensions&quot; -description=&quot;Update package for the P4VUtils tool.&quot; -version=&quot;$(InformationalVersion)&quot; -inputdir=&quot;$(ServerToolsStagingDir_P4VUtils)&quot; -serverdir=&quot;$(StagingDir)/Server/$(ServerArch)&quot; -public -showinugs -showintoolbox -category=Internal" DotNetPath="$(DotNetPath)" />
+				<DotNet Arguments="run --project &quot;$(RootDir)/Engine/Source/Programs/Horde/HordeServer/HordeServer.csproj&quot; -- generate tooldata -id=ugs-cmd -name=&quot;UnrealGameSync (Command-Line)&quot; -description=&quot;Command-line version of UnrealGameSync. Supports Windows, MacOS and Linux.&quot; -version=&quot;$(InformationalVersion)&quot; -inputdir=&quot;$(ServerToolsStagingDir_UnrealGameSyncCmd)&quot; -serverdir=&quot;$(StagingDir)/Server/$(ServerArch)&quot; -public -showintoolbox" DotNetPath="$(DotNetPath)" />
+				<DotNet Arguments="run --project &quot;$(RootDir)/Engine/Source/Programs/Horde/HordeServer/HordeServer.csproj&quot; -- generate tooldata -id=ugs-msi -name=&quot;UnrealGameSync (Windows Installer)&quot; -description=&quot;Installer for UnrealGameSync.&quot; -version=&quot;$(InformationalVersion)&quot; -inputdir=&quot;$(ServerToolsStagingDir_UnrealGameSyncMsi)&quot; -serverdir=&quot;$(StagingDir)/Server/$(ServerArch)&quot; -public -showintoolbox" DotNetPath="$(DotNetPath)" />
 				<DotNet Arguments="run --project &quot;$(RootDir)/Engine/Source/Programs/Horde/HordeServer/HordeServer.csproj&quot; -- generate tooldata -id=ugs-win -name=&quot;UnrealGameSync (Windows Update Package)&quot; -description=&quot;Update package for UGS on Windows.&quot; -version=&quot;$(InformationalVersion)&quot; -inputdir=&quot;$(ServerToolsStagingDir_UnrealGameSync)&quot; -serverdir=&quot;$(StagingDir)/Server/$(ServerArch)&quot; -public -showindashboard=false -category=Internal" DotNetPath="$(DotNetPath)" />
-				<DotNet Arguments="run --project &quot;$(RootDir)/Engine/Source/Programs/Horde/HordeServer/HordeServer.csproj&quot; -- generate tooldata -id=ushell -name=&quot;Ushell&quot; -description=&quot;Update package for ushell.&quot; -version=&quot;$(InformationalVersion)&quot; -inputdir=&quot;$(RootDir)/Engine/Extras/ushell&quot; -serverdir=&quot;$(StagingDir)/Server/$(ServerArch)&quot; -public -showinugs -showinlauncher -Category=Internal" DotNetPath="$(DotNetPath)" />
+				<DotNet Arguments="run --project &quot;$(RootDir)/Engine/Source/Programs/Horde/HordeServer/HordeServer.csproj&quot; -- generate tooldata -id=ushell -name=&quot;Ushell&quot; -description=&quot;Update package for ushell.&quot; -version=&quot;$(InformationalVersion)&quot; -inputdir=&quot;$(RootDir)/Engine/Extras/ushell&quot; -serverdir=&quot;$(StagingDir)/Server/$(ServerArch)&quot; -public -showinugs -showintoolbox -Category=Internal" DotNetPath="$(DotNetPath)" />
 
 				<!-- Separate copy of the command line tool for installation with the server on Windows -->
 				<Delete Files="$(StagingDir)/Tools/..."/>
diff --git a/Engine/Source/Programs/Horde/Plugins/Tools/HordeServer.Tools/Commands/Generate/ToolDataCommand.cs b/Engine/Source/Programs/Horde/Plugins/Tools/HordeServer.Tools/Commands/Generate/ToolDataCommand.cs
--- a/Engine/Source/Programs/Horde/Plugins/Tools/HordeServer.Tools/Commands/Generate/ToolDataCommand.cs
+++ b/Engine/Source/Programs/Horde/Plugins/Tools/HordeServer.Tools/Commands/Generate/ToolDataCommand.cs
@@ -47,8 +47,8 @@ namespace HordeServer.Commands.Generate
 		public bool ShowInUgs { get; set; }
 
 		[CommandLine]
-		[Description("Shows the tool for download in the Horde Launcher")]
-		public bool ShowInLauncher { get; set; }
+		[Description("Shows the tool for download in the Unreal Toolbox")]
+		public bool ShowInToolbox { get; set; }
 
 		[CommandLine("-ShowInDashboard=")]
 		[Description("Shows the tool for download on the dashboard")]
@@ -130,9 +130,9 @@ namespace HordeServer.Commands.Generate
 				{
 					bundledTool[nameof(BundledToolConfig.Public)] = true;
 				}
-				if (ShowInLauncher)
+				if (ShowInToolbox)
 				{
-					bundledTool[nameof(BundledToolConfig.ShowInLauncher)] = true;
+					bundledTool[nameof(BundledToolConfig.ShowInToolbox)] = true;
 				}
 				if (ShowInUgs)
 				{
diff --git a/Engine/Source/Programs/Horde/Plugins/Tools/HordeServer.Tools/Tools/ToolCollection.cs b/Engine/Source/Programs/Horde/Plugins/Tools/HordeServer.Tools/Tools/ToolCollection.cs
--- a/Engine/Source/Programs/Horde/Plugins/Tools/HordeServer.Tools/Tools/ToolCollection.cs
+++ b/Engine/Source/Programs/Horde/Plugins/Tools/HordeServer.Tools/Tools/ToolCollection.cs
@@ -104,7 +104,7 @@ namespace HordeServer.Tools
 			public bool Bundled => _config is BundledToolConfig;
 			public bool ShowInUgs => _config.ShowInUgs;
 			public bool ShowInDashboard => _config.ShowInDashboard;
-			public bool ShowInLauncher => _config.ShowInLauncher;
+			public bool ShowInToolbox => _config.ShowInToolbox;
 			public IReadOnlyList<IToolDeployment> Deployments => _deployments;
 
 			public Tool(ToolCollection collection, ToolDocument document, ToolConfig config, DateTime utcNow)
diff --git a/Engine/Source/Programs/Horde/Plugins/Tools/HordeServer.Tools/Tools/ToolConfig.cs b/Engine/Source/Programs/Horde/Plugins/Tools/HordeServer.Tools/Tools/ToolConfig.cs
--- a/Engine/Source/Programs/Horde/Plugins/Tools/HordeServer.Tools/Tools/ToolConfig.cs
+++ b/Engine/Source/Programs/Horde/Plugins/Tools/HordeServer.Tools/Tools/ToolConfig.cs
@@ -71,9 +71,9 @@ namespace HordeServer.Tools
 		public bool ShowInDashboard { get; set; } = true;
 
 		/// <summary>
-		/// Whether to show this tool for download in the Horde launcher
+		/// Whether to show this tool for download in Unreal Toolbox
 		/// </summary>
-		public bool ShowInLauncher { get; set; }
+		public bool ShowInToolbox { get; set; }
 
 		/// <summary>
 		/// Default namespace for new deployments of this tool
diff --git a/Engine/Source/Programs/Horde/Plugins/Tools/HordeServer.Tools/Tools/ToolsController.cs b/Engine/Source/Programs/Horde/Plugins/Tools/HordeServer.Tools/Tools/ToolsController.cs
--- a/Engine/Source/Programs/Horde/Plugins/Tools/HordeServer.Tools/Tools/ToolsController.cs
+++ b/Engine/Source/Programs/Horde/Plugins/Tools/HordeServer.Tools/Tools/ToolsController.cs
@@ -204,7 +204,7 @@ namespace HordeServer.Tools
 		static GetToolSummaryResponse CreateGetToolSummaryResponse(ITool tool)
 		{
 			IToolDeployment? deployment = (tool.Deployments.Count == 0) ? null : tool.Deployments[^1];
-			return new GetToolSummaryResponse(tool.Id, tool.Name, tool.Description, tool.Category, tool.Group, tool.Platforms?.ToList(), deployment?.Version, deployment?.Id, deployment?.State, deployment?.Progress, tool.Bundled, tool.ShowInUgs, tool.ShowInDashboard, tool.ShowInLauncher);
+			return new GetToolSummaryResponse(tool.Id, tool.Name, tool.Description, tool.Category, tool.Group, tool.Platforms?.ToList(), deployment?.Version, deployment?.Id, deployment?.State, deployment?.Progress, tool.Bundled, tool.ShowInUgs, tool.ShowInDashboard, tool.ShowInToolbox);
 		}
 
 		/// <summary>
@@ -248,7 +248,7 @@ namespace HordeServer.Tools
 
 		static GetToolResponse CreateGetToolResponse(ITool tool, List<GetToolDeploymentResponse> deployments)
 		{
-			return new GetToolResponse(tool.Id, tool.Name, tool.Description, tool.Category, tool.Group, tool.Platforms?.ToList(), deployments, tool.Public, tool.Bundled, tool.ShowInUgs, tool.ShowInDashboard, tool.ShowInLauncher);
+			return new GetToolResponse(tool.Id, tool.Name, tool.Description, tool.Category, tool.Group, tool.Platforms?.ToList(), deployments, tool.Public, tool.Bundled, tool.ShowInUgs, tool.ShowInDashboard, tool.ShowInToolbox);
 		}
 
 		/// <summary>
diff --git a/Engine/Source/Programs/Shared/EpicGames.Horde/Tools/HttpToolCollection.cs b/Engine/Source/Programs/Shared/EpicGames.Horde/Tools/HttpToolCollection.cs
--- a/Engine/Source/Programs/Shared/EpicGames.Horde/Tools/HttpToolCollection.cs
+++ b/Engine/Source/Programs/Shared/EpicGames.Horde/Tools/HttpToolCollection.cs
@@ -32,7 +32,7 @@ namespace EpicGames.Horde.Tools
 			bool ITool.Bundled => _response.Bundled;
 			bool ITool.ShowInUgs => _response.ShowInUgs;
 			bool ITool.ShowInDashboard => _response.ShowInDashboard;
-			bool ITool.ShowInLauncher => _response.ShowInLauncher;
+			bool ITool.ShowInToolbox => _response.ShowInToolbox;
 			IReadOnlyList<IToolDeployment> ITool.Deployments => _deployments;
 
 			public Tool(HttpToolCollection collection, GetToolResponse response)
diff --git a/Engine/Source/Programs/Shared/EpicGames.Horde/Tools/ITool.cs b/Engine/Source/Programs/Shared/EpicGames.Horde/Tools/ITool.cs
--- a/Engine/Source/Programs/Shared/EpicGames.Horde/Tools/ITool.cs
+++ b/Engine/Source/Programs/Shared/EpicGames.Horde/Tools/ITool.cs
@@ -71,9 +71,9 @@ namespace EpicGames.Horde.Tools
 		bool ShowInDashboard { get; }
 
 		/// <summary>
-		/// Whether to show the tool for download in the launcher
+		/// Whether to show the tool for download in the toolbox
 		/// </summary>
-		bool ShowInLauncher { get; }
+		bool ShowInToolbox { get; }
 
 		/// <summary>
 		/// Current deployments of this tool, sorted by time.
diff --git a/Engine/Source/Programs/Shared/EpicGames.Horde/Tools/ToolMessages.cs b/Engine/Source/Programs/Shared/EpicGames.Horde/Tools/ToolMessages.cs
--- a/Engine/Source/Programs/Shared/EpicGames.Horde/Tools/ToolMessages.cs
+++ b/Engine/Source/Programs/Shared/EpicGames.Horde/Tools/ToolMessages.cs
@@ -21,8 +21,8 @@ namespace EpicGames.Horde.Tools
 	/// <param name="Bundled">Whether this tool is bundled with the server</param>
 	/// <param name="ShowInUgs">Whether to show this tool for download inside UGS</param>
 	/// <param name="ShowInDashboard">Whether to show this tool for download on the dashboard</param>
-	/// <param name="ShowInLauncher">Whether to show this tool for download in the launcher</param>
-	public record class GetToolResponse(ToolId Id, string Name, string Description, string? Category, string? Group, List<string>? Platforms, List<GetToolDeploymentResponse> Deployments, bool Public, bool Bundled, bool ShowInUgs, bool ShowInDashboard, bool ShowInLauncher);
+	/// <param name="ShowInToolbox">Whether to show this tool for download in Unreal Toolbox</param>
+	public record class GetToolResponse(ToolId Id, string Name, string Description, string? Category, string? Group, List<string>? Platforms, List<GetToolDeploymentResponse> Deployments, bool Public, bool Bundled, bool ShowInUgs, bool ShowInDashboard, bool ShowInToolbox);
 
 	/// <summary>
 	/// Summary for a particular tool.
@@ -40,8 +40,8 @@ namespace EpicGames.Horde.Tools
 	/// <param name="Bundled">Whether this tool is bundled with the server</param>
 	/// <param name="ShowInUgs">Whether to show this tool for download inside UGS</param>
 	/// <param name="ShowInDashboard">Whether to show this tool for download on the dashboard</param>
-	/// <param name="ShowInLauncher">Whether to show this tool for download in the launcher</param>
-	public record class GetToolSummaryResponse(ToolId Id, string Name, string Description, string? Category, string? Group, List<string>? Platforms, string? Version, ToolDeploymentId? DeploymentId, ToolDeploymentState? DeploymentState, double? DeploymentProgress, bool Bundled, bool ShowInUgs, bool ShowInDashboard, bool ShowInLauncher);
+	/// <param name="ShowInToolbox">Whether to show this tool for download in the launcher</param>
+	public record class GetToolSummaryResponse(ToolId Id, string Name, string Description, string? Category, string? Group, List<string>? Platforms, string? Version, ToolDeploymentId? DeploymentId, ToolDeploymentState? DeploymentState, double? DeploymentProgress, bool Bundled, bool ShowInUgs, bool ShowInDashboard, bool ShowInToolbox);
 
 	/// <summary>
 	/// Response when querying all tools
diff --git a/Engine/Source/Programs/UnrealToolbox/ToolCatalog.cs b/Engine/Source/Programs/UnrealToolbox/ToolCatalog.cs
--- a/Engine/Source/Programs/UnrealToolbox/ToolCatalog.cs
+++ b/Engine/Source/Programs/UnrealToolbox/ToolCatalog.cs
@@ -809,11 +809,15 @@ namespace UnrealToolbox
 
 		static bool ShouldIncludeTool(ITool tool)
 		{
+			if (!tool.ShowInToolbox)
+			{
+				return false;
+			}
 			if (tool.Platforms != null && !tool.Platforms.Any(x => HostPlatformRids.Contains(x)))
 			{
 				return false;
 			}
-			return true; // TODO: ShowInUgs?
+			return true;
 		}
 
 		static HashSet<string> GetHostPlatformRids()
