Horde: Add support for group and platforms to bundled tools, adjust the build accordingly

diff --git a/Engine/Source/Programs/Horde/BuildHorde.xml b/Engine/Source/Programs/Horde/BuildHorde.xml
--- a/Engine/Source/Programs/Horde/BuildHorde.xml
+++ b/Engine/Source/Programs/Horde/BuildHorde.xml
@@ -304,17 +304,21 @@
 				<Delete Files="$(StagingDir)/Server/$(ServerArch)/appsettings.Local.json"/>
 
 				<!-- Add the bundled tools into the installer -->
-				<DotNet Arguments="run --project &quot;$(RootDir)/Engine/Source/Programs/Horde/HordeServer/HordeServer.csproj&quot; -- generate tooldata -id=horde-cmd -name=&quot;Horde (Command-Line)&quot; -description=&quot;Update package for the Horde command-line tool.&quot; -version=&quot;$(InformationalVersion)&quot; -inputdir=&quot;$(ServerToolsStagingDir_HordeCmd)&quot; -serverdir=&quot;$(StagingDir)/Server/$(ServerArch)&quot; -public -showinugs -showintoolbox -Category=Internal" DotNetPath="$(DotNetPath)" />
-				<DotNet Arguments="run --project &quot;$(RootDir)/Engine/Source/Programs/Horde/HordeServer/HordeServer.csproj&quot; -- generate tooldata -id=horde-agent -name=&quot;Horde Agent (Cross-Platform)&quot; -description=&quot;Cross-platform Horde Agent distribution. Allows hosts to take on work from the Horde Server.&quot; -version=&quot;$(InformationalVersion)&quot; -inputdir=&quot;$(ServerToolsStagingDir_HordeAgent)&quot; -serverdir=&quot;$(StagingDir)/Server/$(ServerArch)&quot; -public" DotNetPath="$(DotNetPath)" />
-				<DotNet Arguments="run --project &quot;$(RootDir)/Engine/Source/Programs/Horde/HordeServer/HordeServer.csproj&quot; -- generate tooldata -id=horde-agent-win64 -name=&quot;Horde Agent (Windows)&quot; -description=&quot;Update package for older installations of Horde Agent through the MSI installer.&quot; -version=&quot;$(InformationalVersion)&quot; -inputdir=&quot;$(ServerToolsStagingDir_HordeAgentWin64)&quot; -serverdir=&quot;$(StagingDir)/Server/$(ServerArch)&quot; -public -showindashboard=false -category=Internal" DotNetPath="$(DotNetPath)" />
-				<DotNet Arguments="run --project &quot;$(RootDir)/Engine/Source/Programs/Horde/HordeServer/HordeServer.csproj&quot; -- generate tooldata -id=horde-agent-win64-installer -name=&quot;Horde Agent (Windows Installer)&quot; -description=&quot;Installer for the Horde Agent on Windows, with toolbox application.&quot; -version=&quot;$(InformationalVersion)&quot; -inputdir=&quot;$(ServerToolsStagingDir_HordeAgentMsi)&quot; -serverdir=&quot;$(StagingDir)/Server/$(ServerArch)&quot; -public" DotNetPath="$(DotNetPath)" />
-				<DotNet Arguments="run --project &quot;$(RootDir)/Engine/Source/Programs/Horde/HordeServer/HordeServer.csproj&quot; -- generate tooldata -id=horde-toolbox -name=&quot;Horde Toolbox&quot; -description=&quot;Update package for the Horde Toolbox.&quot; -version=&quot;$(InformationalVersion)&quot; -inputdir=&quot;$(ServerToolsStagingDir_UnrealToolbox)&quot; -serverdir=&quot;$(StagingDir)/Server/$(ServerArch)&quot; -public -showinugs -Category=Internal" DotNetPath="$(DotNetPath)" />
-				<DotNet Arguments="run --project &quot;$(RootDir)/Engine/Source/Programs/Horde/HordeServer/HordeServer.csproj&quot; -- generate tooldata -id=horde-toolbox-msi -name=&quot;Horde Toolbox (Windows Installer)&quot; -description=&quot;Windows installer for the Horde Toolbox.&quot; -version=&quot;$(InformationalVersion)&quot; -inputdir=&quot;$(ServerToolsStagingDir_UnrealToolboxMsi)&quot; -serverdir=&quot;$(StagingDir)/Server/$(ServerArch)&quot; -public -showinugs" DotNetPath="$(DotNetPath)" />
-				<DotNet Arguments="run --project &quot;$(RootDir)/Engine/Source/Programs/Horde/HordeServer/HordeServer.csproj&quot; -- generate tooldata -id=p4v-utils -name=&quot;P4V Extensions&quot; -description=&quot;Update package for the P4VUtils tool.&quot; -version=&quot;$(InformationalVersion)&quot; -inputdir=&quot;$(ServerToolsStagingDir_P4VUtils)&quot; -serverdir=&quot;$(StagingDir)/Server/$(ServerArch)&quot; -public -showinugs -showintoolbox -category=Internal" DotNetPath="$(DotNetPath)" />
-				<DotNet Arguments="run --project &quot;$(RootDir)/Engine/Source/Programs/Horde/HordeServer/HordeServer.csproj&quot; -- generate tooldata -id=ugs-cmd -name=&quot;UnrealGameSync (Command-Line)&quot; -description=&quot;Command-line version of UnrealGameSync. Supports Windows, MacOS and Linux.&quot; -version=&quot;$(InformationalVersion)&quot; -inputdir=&quot;$(ServerToolsStagingDir_UnrealGameSyncCmd)&quot; -serverdir=&quot;$(StagingDir)/Server/$(ServerArch)&quot; -public -showintoolbox" DotNetPath="$(DotNetPath)" />
-				<DotNet Arguments="run --project &quot;$(RootDir)/Engine/Source/Programs/Horde/HordeServer/HordeServer.csproj&quot; -- generate tooldata -id=ugs-msi -name=&quot;UnrealGameSync (Windows Installer)&quot; -description=&quot;Installer for UnrealGameSync.&quot; -version=&quot;$(InformationalVersion)&quot; -inputdir=&quot;$(ServerToolsStagingDir_UnrealGameSyncMsi)&quot; -serverdir=&quot;$(StagingDir)/Server/$(ServerArch)&quot; -public -showintoolbox" DotNetPath="$(DotNetPath)" />
-				<DotNet Arguments="run --project &quot;$(RootDir)/Engine/Source/Programs/Horde/HordeServer/HordeServer.csproj&quot; -- generate tooldata -id=ugs-win -name=&quot;UnrealGameSync (Windows Update Package)&quot; -description=&quot;Update package for UGS on Windows.&quot; -version=&quot;$(InformationalVersion)&quot; -inputdir=&quot;$(ServerToolsStagingDir_UnrealGameSync)&quot; -serverdir=&quot;$(StagingDir)/Server/$(ServerArch)&quot; -public -showindashboard=false -category=Internal" DotNetPath="$(DotNetPath)" />
-				<DotNet Arguments="run --project &quot;$(RootDir)/Engine/Source/Programs/Horde/HordeServer/HordeServer.csproj&quot; -- generate tooldata -id=ushell -name=&quot;Ushell&quot; -description=&quot;Update package for ushell.&quot; -version=&quot;$(InformationalVersion)&quot; -inputdir=&quot;$(RootDir)/Engine/Extras/ushell&quot; -serverdir=&quot;$(StagingDir)/Server/$(ServerArch)&quot; -public -showinugs -showintoolbox -Category=Internal" DotNetPath="$(DotNetPath)" />
+				
+				<!-- General Tools -->
+				<DotNet Arguments="run --project &quot;$(RootDir)/Engine/Source/Programs/Horde/HordeServer/HordeServer.csproj&quot; -- generate tooldata -id=horde-agent -name=&quot;Horde Agent (Cross-Platform)&quot; -group=&quot;horde-agent&quot; -platforms=any -description=&quot;Cross-platform Horde Agent distribution. Allows hosts to take on work from the Horde Server.&quot; -version=&quot;$(InformationalVersion)&quot; -inputdir=&quot;$(ServerToolsStagingDir_HordeAgent)&quot; -serverdir=&quot;$(StagingDir)/Server/$(ServerArch)&quot; -public" DotNetPath="$(DotNetPath)" />				
+				<DotNet Arguments="run --project &quot;$(RootDir)/Engine/Source/Programs/Horde/HordeServer/HordeServer.csproj&quot; -- generate tooldata -id=horde-agent-win64-installer -name=&quot;Horde Agent (Windows Installer)&quot; -group=&quot;horde-agent&quot; -platforms=win-x64 -description=&quot;Installer for the Horde Agent on Windows, with toolbox application.&quot; -version=&quot;$(InformationalVersion)&quot; -inputdir=&quot;$(ServerToolsStagingDir_HordeAgentMsi)&quot; -serverdir=&quot;$(StagingDir)/Server/$(ServerArch)&quot; -public" DotNetPath="$(DotNetPath)" />				
+				<DotNet Arguments="run --project &quot;$(RootDir)/Engine/Source/Programs/Horde/HordeServer/HordeServer.csproj&quot; -- generate tooldata -id=horde-toolbox-msi -name=&quot;Horde Toolbox (Windows Installer)&quot; -description=&quot;Windows installer for the Horde Toolbox.&quot; -version=&quot;$(InformationalVersion)&quot; -inputdir=&quot;$(ServerToolsStagingDir_UnrealToolboxMsi)&quot; -serverdir=&quot;$(StagingDir)/Server/$(ServerArch)&quot; -public -showinugs" DotNetPath="$(DotNetPath)" />				
+				<DotNet Arguments="run --project &quot;$(RootDir)/Engine/Source/Programs/Horde/HordeServer/HordeServer.csproj&quot; -- generate tooldata -id=ugs-cmd -name=&quot;UnrealGameSync (Command-Line)&quot; -group=&quot;ugs&quot; -platforms=any -description=&quot;Command-line version of UnrealGameSync. Supports Windows, MacOS and Linux.&quot; -version=&quot;$(InformationalVersion)&quot; -inputdir=&quot;$(ServerToolsStagingDir_UnrealGameSyncCmd)&quot; -serverdir=&quot;$(StagingDir)/Server/$(ServerArch)&quot; -public -showintoolbox" DotNetPath="$(DotNetPath)" />
+				<DotNet Arguments="run --project &quot;$(RootDir)/Engine/Source/Programs/Horde/HordeServer/HordeServer.csproj&quot; -- generate tooldata -id=ugs-msi -name=&quot;UnrealGameSync (Windows Installer)&quot; -group=&quot;ugs&quot; -platforms=win-x64 -description=&quot;Installer for UnrealGameSync.&quot; -version=&quot;$(InformationalVersion)&quot; -inputdir=&quot;$(ServerToolsStagingDir_UnrealGameSyncMsi)&quot; -serverdir=&quot;$(StagingDir)/Server/$(ServerArch)&quot; -public -showintoolbox" DotNetPath="$(DotNetPath)" />
+				
+				<!-- Internal Category -->
+				<DotNet Arguments="run --project &quot;$(RootDir)/Engine/Source/Programs/Horde/HordeServer/HordeServer.csproj&quot; -- generate tooldata -id=horde-cmd -name=&quot;Horde (Command-Line)&quot; -Category=Internal -description=&quot;Update package for the Horde command-line tool.&quot; -version=&quot;$(InformationalVersion)&quot; -inputdir=&quot;$(ServerToolsStagingDir_HordeCmd)&quot; -serverdir=&quot;$(StagingDir)/Server/$(ServerArch)&quot; -public -showinugs -showintoolbox" DotNetPath="$(DotNetPath)" />
+				<DotNet Arguments="run --project &quot;$(RootDir)/Engine/Source/Programs/Horde/HordeServer/HordeServer.csproj&quot; -- generate tooldata -id=horde-agent-win64 -name=&quot;Horde Agent (Windows)&quot; -category=Internal -description=&quot;Update package for older installations of Horde Agent through the MSI installer.&quot; -version=&quot;$(InformationalVersion)&quot; -inputdir=&quot;$(ServerToolsStagingDir_HordeAgentWin64)&quot; -serverdir=&quot;$(StagingDir)/Server/$(ServerArch)&quot; -public -showindashboard=false" DotNetPath="$(DotNetPath)" />
+				<DotNet Arguments="run --project &quot;$(RootDir)/Engine/Source/Programs/Horde/HordeServer/HordeServer.csproj&quot; -- generate tooldata -id=horde-toolbox -name=&quot;Horde Toolbox&quot; -Category=Internal -description=&quot;Update package for the Horde Toolbox.&quot; -version=&quot;$(InformationalVersion)&quot; -inputdir=&quot;$(ServerToolsStagingDir_UnrealToolbox)&quot; -serverdir=&quot;$(StagingDir)/Server/$(ServerArch)&quot; -public -showinugs" DotNetPath="$(DotNetPath)" />
+				<DotNet Arguments="run --project &quot;$(RootDir)/Engine/Source/Programs/Horde/HordeServer/HordeServer.csproj&quot; -- generate tooldata -id=p4v-utils -name=&quot;P4V Extensions&quot; -category=Internal -description=&quot;Update package for the P4VUtils tool.&quot; -version=&quot;$(InformationalVersion)&quot; -inputdir=&quot;$(ServerToolsStagingDir_P4VUtils)&quot; -serverdir=&quot;$(StagingDir)/Server/$(ServerArch)&quot; -public -showinugs -showintoolbox" DotNetPath="$(DotNetPath)" />
+				<DotNet Arguments="run --project &quot;$(RootDir)/Engine/Source/Programs/Horde/HordeServer/HordeServer.csproj&quot; -- generate tooldata -id=ugs-win -name=&quot;UnrealGameSync (Windows Update Package)&quot; -category=Internal -description=&quot;Update package for UGS on Windows.&quot; -version=&quot;$(InformationalVersion)&quot; -inputdir=&quot;$(ServerToolsStagingDir_UnrealGameSync)&quot; -serverdir=&quot;$(StagingDir)/Server/$(ServerArch)&quot; -public -showindashboard=false" DotNetPath="$(DotNetPath)" />
+				<DotNet Arguments="run --project &quot;$(RootDir)/Engine/Source/Programs/Horde/HordeServer/HordeServer.csproj&quot; -- generate tooldata -id=ushell -name=&quot;Ushell&quot; -Category=Internal -description=&quot;Update package for ushell.&quot; -version=&quot;$(InformationalVersion)&quot; -inputdir=&quot;$(RootDir)/Engine/Extras/ushell&quot; -serverdir=&quot;$(StagingDir)/Server/$(ServerArch)&quot; -public -showinugs -showintoolbox" DotNetPath="$(DotNetPath)" />			
 
 				<!-- Separate copy of the command line tool for installation with the server on Windows -->
 				<Delete Files="$(StagingDir)/Tools/..."/>
diff --git a/Engine/Source/Programs/Horde/Plugins/Tools/HordeServer.Tools/Commands/Generate/ToolDataCommand.cs b/Engine/Source/Programs/Horde/Plugins/Tools/HordeServer.Tools/Commands/Generate/ToolDataCommand.cs
--- a/Engine/Source/Programs/Horde/Plugins/Tools/HordeServer.Tools/Commands/Generate/ToolDataCommand.cs
+++ b/Engine/Source/Programs/Horde/Plugins/Tools/HordeServer.Tools/Commands/Generate/ToolDataCommand.cs
@@ -34,6 +34,13 @@ namespace HordeServer.Commands.Generate
 		[Description("Category for the tool")]
 		public string? Category { get; set; }
 
+		[CommandLine]
+		[Description("Group for the tool")]
+		public string? Group { get; set; }
+
+		[CommandLine(ListSeparator = ';')]
+		public string[]? Platforms { get; set; }
+
 		[CommandLine]
 		[Description("Version string for the tool")]
 		public string? Version { get; set; }
@@ -119,11 +126,25 @@ namespace HordeServer.Commands.Generate
 				{
 					bundledTool[nameof(BundledToolConfig.Category)] = Category;
 				}
+				if (!String.IsNullOrEmpty(Group))
+				{
+					bundledTool[nameof(BundledToolConfig.Group)] = Group;
+				}
 				if (!String.IsNullOrEmpty(Version))
 				{
 					bundledTool[nameof(BundledToolConfig.Version)] = Version;
 				}
 
+				if (Platforms != null && Platforms.Length > 0)
+				{
+					JsonArray platforms = new JsonArray();
+					foreach (string platform in Platforms) 
+					{
+						platforms.Add(platform);
+					}
+					bundledTool[nameof(BundledToolConfig.Platforms)] = platforms;
+				}
+
 				bundledTool[nameof(BundledToolConfig.RefName)] = refName.ToString();
 
 				if (Public)
