Horde: Add missing description for platforms command line option
#rnx

diff --git a/Engine/Source/Programs/Horde/Plugins/Tools/HordeServer.Tools/Commands/Generate/ToolDataCommand.cs b/Engine/Source/Programs/Horde/Plugins/Tools/HordeServer.Tools/Commands/Generate/ToolDataCommand.cs
--- a/Engine/Source/Programs/Horde/Plugins/Tools/HordeServer.Tools/Commands/Generate/ToolDataCommand.cs
+++ b/Engine/Source/Programs/Horde/Plugins/Tools/HordeServer.Tools/Commands/Generate/ToolDataCommand.cs
@@ -39,6 +39,7 @@ namespace HordeServer.Commands.Generate
 		public string? Group { get; set; }
 
 		[CommandLine(ListSeparator = ';')]
+		[Description("Semicolon delimited list of tool platforms")]
 		public string[]? Platforms { get; set; }
 
 		[CommandLine]
