Horde: Update other projects to net 8.

#rnx

diff --git a/Engine/Source/Programs/Horde/Drivers/JobDriver/Execution/JobExecutor.cs b/Engine/Source/Programs/Horde/Drivers/JobDriver/Execution/JobExecutor.cs
--- a/Engine/Source/Programs/Horde/Drivers/JobDriver/Execution/JobExecutor.cs
+++ b/Engine/Source/Programs/Horde/Drivers/JobDriver/Execution/JobExecutor.cs
@@ -554,7 +554,7 @@ namespace JobDriver.Execution
 					if (res.AbortRequested)
 					{
 						leaseLogger.LogDebug("Step was aborted by server (JobId={JobId} BatchId={BatchId} StepId={StepId})", jobId, batchId, stepId);
-						stepCancelSource.Cancel();
+						await stepCancelSource.CancelAsync();
 						break;
 					}
 				}
diff --git a/Engine/Source/Programs/Horde/Drivers/JobDriver/JobDriver.csproj b/Engine/Source/Programs/Horde/Drivers/JobDriver/JobDriver.csproj
--- a/Engine/Source/Programs/Horde/Drivers/JobDriver/JobDriver.csproj
+++ b/Engine/Source/Programs/Horde/Drivers/JobDriver/JobDriver.csproj
@@ -2,7 +2,7 @@
 
   <PropertyGroup>
     <OutputType>Exe</OutputType>
-    <TargetFramework>net6.0</TargetFramework>
+    <TargetFramework>net8.0</TargetFramework>
     <Nullable>enable</Nullable>
     <ImplicitUsings>enable</ImplicitUsings>
     <AssemblyName>JobDriver</AssemblyName>
diff --git a/Engine/Source/Programs/Horde/Drivers/JobDriver/Utility/JobStepLogger.cs b/Engine/Source/Programs/Horde/Drivers/JobDriver/Utility/JobStepLogger.cs
--- a/Engine/Source/Programs/Horde/Drivers/JobDriver/Utility/JobStepLogger.cs
+++ b/Engine/Source/Programs/Horde/Drivers/JobDriver/Utility/JobStepLogger.cs
@@ -59,7 +59,7 @@ namespace JobDriver.Utility
 		{
 			if (!_updateOutcomeTask.IsCompleted)
 			{
-				_cancellationSource.Cancel();
+				await _cancellationSource.CancelAsync();
 				await _updateOutcomeTask;
 				_cancellationSource.Dispose();
 			}
diff --git a/Engine/Source/Programs/Horde/Samples/RemoteClient/RemoteClient.csproj b/Engine/Source/Programs/Horde/Samples/RemoteClient/RemoteClient.csproj
--- a/Engine/Source/Programs/Horde/Samples/RemoteClient/RemoteClient.csproj
+++ b/Engine/Source/Programs/Horde/Samples/RemoteClient/RemoteClient.csproj
@@ -2,7 +2,7 @@
 
   <PropertyGroup>
     <OutputType>Exe</OutputType>
-    <TargetFramework>net6.0</TargetFramework>
+    <TargetFramework>net8.0</TargetFramework>
     <ImplicitUsings>enable</ImplicitUsings>
     <Nullable>enable</Nullable>
   </PropertyGroup>
diff --git a/Engine/Source/Programs/Horde/Samples/RemoteWorker/RemoteWorker.csproj b/Engine/Source/Programs/Horde/Samples/RemoteWorker/RemoteWorker.csproj
--- a/Engine/Source/Programs/Horde/Samples/RemoteWorker/RemoteWorker.csproj
+++ b/Engine/Source/Programs/Horde/Samples/RemoteWorker/RemoteWorker.csproj
@@ -2,7 +2,7 @@
 
   <PropertyGroup>
     <OutputType>Exe</OutputType>
-    <TargetFramework>net6.0</TargetFramework>
+    <TargetFramework>net8.0</TargetFramework>
     <ImplicitUsings>enable</ImplicitUsings>
     <Nullable>enable</Nullable>
   </PropertyGroup>
