Horde: Add command line tool to PATH after being installed through toolbox.

#rnx

diff --git a/Engine/Source/Programs/Horde/HordeCmd/HordeCmd.csproj b/Engine/Source/Programs/Horde/HordeCmd/HordeCmd.csproj
--- a/Engine/Source/Programs/Horde/HordeCmd/HordeCmd.csproj
+++ b/Engine/Source/Programs/Horde/HordeCmd/HordeCmd.csproj
@@ -40,6 +40,9 @@
       <CopyToOutputDirectory>Always</CopyToOutputDirectory>
       <CopyToPublishDirectory>Always</CopyToPublishDirectory>
     </Content>
+    <Content Include="Toolbox.json">
+      <CopyToOutputDirectory>Always</CopyToOutputDirectory>
+    </Content>
     <Content Include="UgsTool.json">
       <CopyToOutputDirectory>Always</CopyToOutputDirectory>
     </Content>
diff --git a/Engine/Source/Programs/Horde/HordeCmd/Toolbox.json b/Engine/Source/Programs/Horde/HordeCmd/Toolbox.json
new file mode 100644
--- /dev/null
+++ b/Engine/Source/Programs/Horde/HordeCmd/Toolbox.json
@@ -0,0 +1,14 @@
+{
+    "installCommand": {
+        "fileName": "horde.exe",
+        "arguments": [
+            "install"
+        ]
+    },
+    "uninstallCommand": {
+        "fileName": "horde.exe",
+        "arguments": [
+            "uninstall"
+        ]
+    }
+}
\ No newline at end of file
