Horde: Fix exception generating platform metadata for tools, due to CommandLineArguments class not supporting addition of items to arrays.

#rnx

diff --git a/Engine/Source/Programs/Horde/Plugins/Tools/HordeServer.Tools/Commands/Generate/ToolDataCommand.cs b/Engine/Source/Programs/Horde/Plugins/Tools/HordeServer.Tools/Commands/Generate/ToolDataCommand.cs
--- a/Engine/Source/Programs/Horde/Plugins/Tools/HordeServer.Tools/Commands/Generate/ToolDataCommand.cs
+++ b/Engine/Source/Programs/Horde/Plugins/Tools/HordeServer.Tools/Commands/Generate/ToolDataCommand.cs
@@ -40,7 +40,7 @@ namespace HordeServer.Commands.Generate
 
 		[CommandLine(ListSeparator = ';')]
 		[Description("Semicolon delimited list of tool platforms")]
-		public string[]? Platforms { get; set; }
+		public List<string> Platforms { get; set; } = new List<string>();
 
 		[CommandLine]
 		[Description("Version string for the tool")]
@@ -136,7 +136,7 @@ namespace HordeServer.Commands.Generate
 					bundledTool[nameof(BundledToolConfig.Version)] = Version;
 				}
 
-				if (Platforms != null && Platforms.Length > 0)
+				if (Platforms.Count > 0)
 				{
 					JsonArray platforms = new JsonArray();
 					foreach (string platform in Platforms) 
