Horde: Add a separate build of the Horde Agent MSI for the toolbox, which includes metadata for running the installer.

diff --git a/Engine/Source/Programs/Horde/BuildHorde.xml b/Engine/Source/Programs/Horde/BuildHorde.xml
--- a/Engine/Source/Programs/Horde/BuildHorde.xml
+++ b/Engine/Source/Programs/Horde/BuildHorde.xml
@@ -129,6 +129,7 @@
 	<Property Name="ServerToolsStagingDir_HordeAgent" Value="$(ServerToolsStagingDir)/HordeAgent"/>
 	<Property Name="ServerToolsStagingDir_HordeAgentWin64" Value="$(ServerToolsStagingDir)/HordeAgentWin64"/>
 	<Property Name="ServerToolsStagingDir_HordeAgentMsi" Value="$(ServerToolsStagingDir)/HordeAgentMsi"/>
+	<Property Name="ServerToolsStagingDir_HordeAgentMsiToolbox" Value="$(ServerToolsStagingDir)/HordeAgentMsiToolbox"/>
 	<Property Name="ServerToolsStagingDir_UnrealToolbox" Value="$(ServerToolsStagingDir)/UnrealToolbox"/>
 	<Property Name="ServerToolsStagingDir_UnrealToolboxMsi" Value="$(ServerToolsStagingDir)/UnrealToolboxMsi"/>
 	<Property Name="ServerToolsStagingDir_P4VUtils" Value="$(ServerToolsStagingDir)/P4VUtils"/>
@@ -292,9 +293,15 @@
 			<Tag Files="$(ServerToolsStagingDir_HordeAgentMsi)/..." With="#HordeAgentMsiFiles"/>
 
 		</Node>
+		
+		<Node Name="Build Horde Agent Installer (Toolbox)" Requires="Build Horde Agent Installer" Produces="#HordeAgentMsiToolboxFiles">
+			<Copy Files="UnrealHordeAgent.msi" From="$(ServerToolsStagingDir_HordeAgentMsi)" To="$(ServerToolsStagingDir_HordeAgentMsiToolbox)"/>
+			<Copy Files="Toolbox.json" From="$(RootDir)/Engine/Source/Programs/Horde/HordeAgent/Toolbox.json" To="$(ServerToolsStagingDir_HordeAgentMsiToolbox)"/>
+			<Tag Files="$(ServerToolsStagingDir_HordeAgentMsiToolbox)/..." With="#HordeAgentMsiToolboxFiles"/>
+		</Node>
 
 		<ForEach Name="ServerArch" Values="$(ServerArchList)">
-			<Node Name="Stage Horde Server Installer ($(ServerArch))" Requires="#DashboardOutputFiles;#HordeCmdFiles;#HordeAgentFiles;#HordeAgentMsiFiles;#UnrealToolboxFiles;#UnrealToolboxMsiFiles;#P4VUtilsFiles;#UnrealGameSyncFiles;#UnrealGameSyncCmdFiles;#UnrealGameSyncMsiFiles;" Produces="#ServerWithToolsFiles_$(ServerArch)">
+			<Node Name="Stage Horde Server Installer ($(ServerArch))" Requires="#DashboardOutputFiles;#HordeCmdFiles;#HordeAgentFiles;#HordeAgentMsiFiles;#HordeAgentMsiToolboxFiles;#UnrealToolboxFiles;#UnrealToolboxMsiFiles;#P4VUtilsFiles;#UnrealGameSyncFiles;#UnrealGameSyncCmdFiles;#UnrealGameSyncMsiFiles;" Produces="#ServerWithToolsFiles_$(ServerArch)">
 		
 				<!-- Stage the server -->
 				<Delete Files="$(StagingDir)/Server/$(ServerArch)/..."/>
@@ -304,17 +311,18 @@
 				<Delete Files="$(StagingDir)/Server/$(ServerArch)/appsettings.Local.json"/>
 
 				<!-- Add the bundled tools into the installer -->
-				
+
 				<!-- General Tools -->
 				<DotNet Arguments="run --project &quot;$(RootDir)/Engine/Source/Programs/Horde/HordeServer/HordeServer.csproj&quot; -- generate tooldata -id=horde-agent -name=&quot;Horde Agent (Cross-Platform)&quot; -group=&quot;horde-agent&quot; -platforms=any -description=&quot;Cross-platform Horde Agent distribution. Allows hosts to take on work from the Horde Server.&quot; -version=&quot;$(InformationalVersion)&quot; -inputdir=&quot;$(ServerToolsStagingDir_HordeAgent)&quot; -serverdir=&quot;$(StagingDir)/Server/$(ServerArch)&quot; -public" DotNetPath="$(DotNetPath)" />				
-				<DotNet Arguments="run --project &quot;$(RootDir)/Engine/Source/Programs/Horde/HordeServer/HordeServer.csproj&quot; -- generate tooldata -id=horde-agent-win64-installer -name=&quot;Horde Agent (Windows Installer)&quot; -group=&quot;horde-agent&quot; -platforms=win-x64 -description=&quot;Installer for the Horde Agent on Windows, with toolbox application.&quot; -version=&quot;$(InformationalVersion)&quot; -inputdir=&quot;$(ServerToolsStagingDir_HordeAgentMsi)&quot; -serverdir=&quot;$(StagingDir)/Server/$(ServerArch)&quot; -public" DotNetPath="$(DotNetPath)" />				
+				<DotNet Arguments="run --project &quot;$(RootDir)/Engine/Source/Programs/Horde/HordeServer/HordeServer.csproj&quot; -- generate tooldata -id=horde-agent-msi -name=&quot;Horde Agent (Windows Installer)&quot; -group=&quot;horde-agent&quot; -platforms=win-x64 -description=&quot;Installer for the Horde Agent on Windows.&quot; -version=&quot;$(InformationalVersion)&quot; -inputdir=&quot;$(ServerToolsStagingDir_HordeAgentMsi)&quot; -serverdir=&quot;$(StagingDir)/Server/$(ServerArch)&quot; -public" DotNetPath="$(DotNetPath)" />				
 				<DotNet Arguments="run --project &quot;$(RootDir)/Engine/Source/Programs/Horde/HordeServer/HordeServer.csproj&quot; -- generate tooldata -id=horde-toolbox-msi -name=&quot;Horde Toolbox (Windows Installer)&quot; -description=&quot;Windows installer for the Horde Toolbox.&quot; -version=&quot;$(InformationalVersion)&quot; -inputdir=&quot;$(ServerToolsStagingDir_UnrealToolboxMsi)&quot; -serverdir=&quot;$(StagingDir)/Server/$(ServerArch)&quot; -public -showinugs" DotNetPath="$(DotNetPath)" />				
 				<DotNet Arguments="run --project &quot;$(RootDir)/Engine/Source/Programs/Horde/HordeServer/HordeServer.csproj&quot; -- generate tooldata -id=ugs-cmd -name=&quot;UnrealGameSync (Command-Line)&quot; -group=&quot;ugs&quot; -platforms=any -description=&quot;Command-line version of UnrealGameSync. Supports Windows, MacOS and Linux.&quot; -version=&quot;$(InformationalVersion)&quot; -inputdir=&quot;$(ServerToolsStagingDir_UnrealGameSyncCmd)&quot; -serverdir=&quot;$(StagingDir)/Server/$(ServerArch)&quot; -public -showintoolbox" DotNetPath="$(DotNetPath)" />
 				<DotNet Arguments="run --project &quot;$(RootDir)/Engine/Source/Programs/Horde/HordeServer/HordeServer.csproj&quot; -- generate tooldata -id=ugs-msi -name=&quot;UnrealGameSync (Windows Installer)&quot; -group=&quot;ugs&quot; -platforms=win-x64 -description=&quot;Installer for UnrealGameSync.&quot; -version=&quot;$(InformationalVersion)&quot; -inputdir=&quot;$(ServerToolsStagingDir_UnrealGameSyncMsi)&quot; -serverdir=&quot;$(StagingDir)/Server/$(ServerArch)&quot; -public -showintoolbox" DotNetPath="$(DotNetPath)" />
-				
+
 				<!-- Internal Category -->
 				<DotNet Arguments="run --project &quot;$(RootDir)/Engine/Source/Programs/Horde/HordeServer/HordeServer.csproj&quot; -- generate tooldata -id=horde-cmd -name=&quot;Horde (Command-Line)&quot; -Category=Internal -description=&quot;Update package for the Horde command-line tool.&quot; -version=&quot;$(InformationalVersion)&quot; -inputdir=&quot;$(ServerToolsStagingDir_HordeCmd)&quot; -serverdir=&quot;$(StagingDir)/Server/$(ServerArch)&quot; -public -showinugs -showintoolbox" DotNetPath="$(DotNetPath)" />
 				<DotNet Arguments="run --project &quot;$(RootDir)/Engine/Source/Programs/Horde/HordeServer/HordeServer.csproj&quot; -- generate tooldata -id=horde-agent-win64 -name=&quot;Horde Agent (Windows)&quot; -category=Internal -description=&quot;Update package for older installations of Horde Agent through the MSI installer.&quot; -version=&quot;$(InformationalVersion)&quot; -inputdir=&quot;$(ServerToolsStagingDir_HordeAgentWin64)&quot; -serverdir=&quot;$(StagingDir)/Server/$(ServerArch)&quot; -public -showindashboard=false" DotNetPath="$(DotNetPath)" />
+				<DotNet Arguments="run --project &quot;$(RootDir)/Engine/Source/Programs/Horde/HordeServer/HordeServer.csproj&quot; -- generate tooldata -id=horde-agent-msi-toolbox -name=&quot;Horde Agent&quot; -description=&quot;Installer for the Horde Agent on Windows, for Unreal Toolbox.&quot; -version=&quot;$(InformationalVersion)&quot; -inputdir=&quot;$(ServerToolsStagingDir_HordeAgentMsi)&quot; -serverdir=&quot;$(StagingDir)/Server/$(ServerArch)&quot; -public -showintoolbox" DotNetPath="$(DotNetPath)" />
 				<DotNet Arguments="run --project &quot;$(RootDir)/Engine/Source/Programs/Horde/HordeServer/HordeServer.csproj&quot; -- generate tooldata -id=horde-toolbox -name=&quot;Horde Toolbox&quot; -Category=Internal -description=&quot;Update package for the Horde Toolbox.&quot; -version=&quot;$(InformationalVersion)&quot; -inputdir=&quot;$(ServerToolsStagingDir_UnrealToolbox)&quot; -serverdir=&quot;$(StagingDir)/Server/$(ServerArch)&quot; -public -showinugs" DotNetPath="$(DotNetPath)" />
 				<DotNet Arguments="run --project &quot;$(RootDir)/Engine/Source/Programs/Horde/HordeServer/HordeServer.csproj&quot; -- generate tooldata -id=p4v-utils -name=&quot;P4V Extensions&quot; -category=Internal -description=&quot;Update package for the P4VUtils tool.&quot; -version=&quot;$(InformationalVersion)&quot; -inputdir=&quot;$(ServerToolsStagingDir_P4VUtils)&quot; -serverdir=&quot;$(StagingDir)/Server/$(ServerArch)&quot; -public -showinugs -showintoolbox" DotNetPath="$(DotNetPath)" />
 				<DotNet Arguments="run --project &quot;$(RootDir)/Engine/Source/Programs/Horde/HordeServer/HordeServer.csproj&quot; -- generate tooldata -id=ugs-win -name=&quot;UnrealGameSync (Windows Update Package)&quot; -category=Internal -description=&quot;Update package for UGS on Windows.&quot; -version=&quot;$(InformationalVersion)&quot; -inputdir=&quot;$(ServerToolsStagingDir_UnrealGameSync)&quot; -serverdir=&quot;$(StagingDir)/Server/$(ServerArch)&quot; -public -showindashboard=false" DotNetPath="$(DotNetPath)" />
diff --git a/Engine/Source/Programs/Horde/HordeAgent/Toolbox.json b/Engine/Source/Programs/Horde/HordeAgent/Toolbox.json
new file mode 100644
--- /dev/null
+++ b/Engine/Source/Programs/Horde/HordeAgent/Toolbox.json
@@ -0,0 +1,15 @@
+{
+	"manualInstall": true,
+	"installCommand": {
+		"fileName": "msiexec.exe",
+		"arguments": [
+			"UnrealHordeAgent.msi"
+		]
+	},
+	"uninstallCommand": {
+		"fileName": "msiexec.exe",
+		"arguments": [
+			"UnrealHordeAgent.msi"
+		]
+	}
+}
\ No newline at end of file
