Horde: Expose .NET runtime version in agent properties

diff --git a/Engine/Source/Programs/Horde/HordeAgent/Services/CapabilitiesService.cs b/Engine/Source/Programs/Horde/HordeAgent/Services/CapabilitiesService.cs
--- a/Engine/Source/Programs/Horde/HordeAgent/Services/CapabilitiesService.cs
+++ b/Engine/Source/Programs/Horde/HordeAgent/Services/CapabilitiesService.cs
@@ -468,6 +468,7 @@ namespace HordeAgent.Services
 			// Whether the agent is packaged as a self-contained .NET app
 			// Used during the transition period over from multi-platform, non-self-contained agent packages.
 			capabilities.Properties.Add($"SelfContained={AgentApp.IsSelfContained}");
+			capabilities.Properties.Add($"DotNetFramework={RuntimeInformation.FrameworkDescription}");
 
 			// Add any additional properties from the config file
 			capabilities.Properties.AddRange(_settings.Properties.Select(kvp => $"{kvp.Key}={kvp.Value}"));
