Horde: Update release notes
#skipci
#rnx

diff --git a/Engine/Source/Programs/Horde/Docs/ReleaseNotes.md b/Engine/Source/Programs/Horde/Docs/ReleaseNotes.md
--- a/Engine/Source/Programs/Horde/Docs/ReleaseNotes.md
+++ b/Engine/Source/Programs/Horde/Docs/ReleaseNotes.md
@@ -1,5 +1,93 @@
 # Release Notes
 
+## 2024-09-19
+
+* Expose .NET runtime version in agent properties (36427479)
+* Add a separate build of the Horde Agent MSI for the toolbox, which includes metadata for running the installer. (36425974)
+* Remove tray app and idle behavior setting from agent installer (36410206)
+* Add support for group and platforms to bundled tools, adjust the build accordingly (36407524)
+* Rename HordeToolbox again. Now: "UnrealToolbox!" (36401901)
+* Add a dedicated installer for the tray app, and include it in the server installer. (36394095)
+* Fix infinite loop issue when querying for jobs (36392029)
+* Fix bug where physical cores are counted twice in capability detection (36387918)
+* Replace OpenTracing with OpenTelemetry in JobDriver (36386025)
+* Document the location of the server config file. (36384366)
+* Ensure all streams register default artifact types.  We recently introduced the requirement that artifact types are explicitly configured (CL 35531333). The job driver expects certain artifact types to be present, and jobs will fail if these types are not registered.  This issue affects new installations and older installations upgrading to 5.5. (36383858)
+* Remove OpenTelemetry.Instrumentation.AspNetCore from agent (36383393)
+* Add an option for automatically updating tools in the tray app. (36356951)
+* Fixed a bug where issues would have their FixCommitId set to 0 rather than cleared after re-assignment to a new owner.  Note we set fixChange to 0 in the onAssign function in IssueViewV2.tsx.  I now check whether fixChange is == 0 and if so set FixCommitId to be empty, this causes the FIX CL to be cleared. (36345723)
+* Add information about processes with a file locked when it cannot be overwritten when extracting data from a bundle. (36341246)
+* Add a dataflow graph with multiple deletion workers to improve GC performance. (36340210)
+* Upgrade Horde Agent to NET 8. (36340066)
+* Add an exponential pause between iterations over the GC tick queue. This should reduce the time that we spend idle during blob deletion. (36321827)
+* Fix changes not being enumerated in descending order by PerforceService, causing the minimum changelist number for ICommitCollection.SubscribeAsync() not being updated. This prevents issue fixed tags from being processed. (36321198)
+* Add missing file. (36318745)
+* Replace OpenTracing with OpenTelemetry in agent (36307242)
+* Move OpenTelemetry settings to EpicGames.Horde for sharing with agent (36303013)
+* Added an additional sanitise case for the PerforceMetadataLogger when trying to get file annotations.  We could potentially handle this in LogEventParser AddLine but I wanted to limit the impact of these changes. (36299319)
+* Add a proxy server to the Horde tray app, allowing users to connect to the Horde server via an unauthenticated connection bound to localhost. (36292016)
+* Add a flag for showing a tool in the Horde tray app. (36284430)
+* Add settings menu to tray app, as well as functionality for downloading tools from Horde. (36265467)
+* Handle .horde.json file not existing when configuring server url on Mac/Linux. (36264830)
+* Changed code traversing workspace on disk to use EnumerateFileSystemInfos instead of directories and files separately. This reduces number of kernel calls and save 20% time on a machine with attached ssd (36149765)
+* Add timeouts for available port checking in tests (36117480)
+* Add server-defined properties for an agent
+  - Allows setting properties that will overwrite and merge with properties reported by agent itself
+  - Refactor parameters for agent creation into an options object
+  - When a user/agent can create new agents outside enrollment process, it's marked as trusted for the time being (grandfathering existing registration in JobRpcService) (36114930)
+* Ignore server-defined properties when sent by agent (36108891)
+* Adding Horde installer custom actions, data directory selector, and fix server to be able to bootstrap using custom data directory (36082966)
+* Add tooltips for step times (36043646)
+* Additional logging for terminating sessions. (36041396)
+* Set the max thread count for managed workspace operations to one less than the number of reported CPUs on the machine. (36038061)
+* Add better error for failing to start a process from the Horde Agent. (36026676)
+* Advertise a new OSFamilyCompatibility property from agents, indicating OSes that the system can emulate (for Linux agents with WINE to indicate Windows compatibility). (36018324)
+* Allow specifying a list of properties required for agents to execute compute leases. (36012503)
+* Re-enable dedupe for the tools bundled with Horde. The permissions affecting access to this data are now waived for bundled tools. (36012474)
+* Fix installed server not correctly identifying code changes correctly. Server was incorrectly setting a flag indicating that all files for a change had been enumerated, preventing it scanning the entire set. (35977472)
+* Add a symbol store plugin. Symbol stores use aliases in the storage system to map symbol store paths onto content streams from existing artifacts, allowing reuse of data already available in artifacts.
+  - Symbols can be tagged with the appropriate metadata by setting the Symbols=true attribute on the CreateArtifact task. Referencing the namespace that the symbols will be uploaded to from the symbol store config will allow accessing them through the api/v1/symbols route.
+  - Hashing for symbols is compatible with symstore.exe, but is handled by a custom implementation in SymStore.cs. (35971584)
+* Fix linq expression for generation of alias index. (35967893)
+* Pause the addition of new blobs for GC once the queue is longer than 50,000 entries. (35967249)
+* Allow graceful draining of leases when yielding to local user activity
+  - The termination signal file used for spot interruptions is now written to let a workload know about the lease being drained.
+  - This requires workload to scan and respect the file for this to be effective, which UBA does for example. (35962694)
+* Add agent version to compute resource class, used for exposing version in compute resource API. (35954995)
+* Log Horde server and agent version for a UBA session (35954941)
+* Fix name of aliases over blobs collection. (35950154)
+* Remove code to read imports from uploaded blobs. These should now be set at upload time. (35949781)
+* Support loading artifacts out of log which are not in the step artifacts (35931864)
+* Improve issue button rendering to not obscure the structured logging eyeball, also change where we render the structured log (35931043)
+* Include artifacts created through the CreateArtifact BuildGraph task in the list of artifacts for the job. (35924757)
+* Catch compute cancellation exceptions and avoid flagging them as errors (35923687)
+* Guard against socket errors during closing of compute socket, also ensure CloseAsync can't be invoked twice. (35903197)
+* Set the UE_HORDE_STREAMID environment variable containing the current stream id when running under Horde. (35900707)
+* Document the look up of P4TRUST file (35892976)
+* Change tray app to use Avalonia rather than Winforms. (35883986)
+* Dashboard side of bundled tool management (35877843)
+* Add clear conform to agent context menu (35877280)
+* Adding bundled information to tool responses (35876724)
+* Fix broken ServiceAccountAuth test (35876245)
+* Add name field to service accounts and use that in audit logs, previously, name of service accounts resolved to "unknown". (35873276)
+* Keep track of the agents which are currently able to execute sessions.  Each server now caches the state of a session based on the last update time, allowing it to perform most operations with minimal context fetches in the common case where an agent is always being updated on the same server. (35872405)
+* Fixed issue in recent UHT changes that prevented errors begin generated from the header file object. (35863940)
+* Add agent settings for CPU count and multiplier.  Only provided as hinting to workloads, which can chose to respect these. The initial use-case is letting UBA limit number of CPUs in use, much how maxcpu/mulcpu works for UBT.  Later on these values could configure job objects (on Windows) for proper OS-enforced CPU limiting. (35861968)
+* Document use of ssl: prefix for connecting to Perforce servers (35859151)
+* Replace NuGet package with more explicit one to reduce dependency size (35822082)
+* Fix description and creation time not being deserialized as part of artifact responses correctly. (35811531)
+* Store session state in Redis rather than MongoDB. (35780092)
+* Add HordeHTTPClient GetUgsMetadataAsync method (35775820)
+* Fixed symbol table lookup to properly include the header file when walking up the outer chain. (35773746)
+* Agent enrollment improvements (35770720)
+* Add REST API endpoint for listing ACL permissions for current user (35770032)
+* Fix invalid dep injection of ServerSettings in JwtHandler (35769300)
+* Show permissions for ACL scopes on account page, useful for debugging permissions in Horde. (35768929)
+* Improve OIDC configuration experience
+  - Add debug mode for better explaining why a JWT bearer token is rejected
+  - Validate required settings are set for auth mode OIDC/Okta
+  - Improve docs for OIDC settings (35768735)
+
 ## 2024-08-22
 
 * Lock fluent to fix upstream regression (35716215)
