Horde: Prevent cancellation exceptions during GC from triggering an error.

#rnx

diff --git a/Engine/Source/Programs/Horde/Plugins/Storage/HordeServer.Storage/Storage/StorageService.cs b/Engine/Source/Programs/Horde/Plugins/Storage/HordeServer.Storage/Storage/StorageService.cs
--- a/Engine/Source/Programs/Horde/Plugins/Storage/HordeServer.Storage/Storage/StorageService.cs
+++ b/Engine/Source/Programs/Horde/Plugins/Storage/HordeServer.Storage/Storage/StorageService.cs
@@ -1061,6 +1061,10 @@ namespace HordeServer.Storage
 								{
 									await TickGcForNamespaceAsync(state.Namespaces[namespaceId], gcState.LastImportBlobInfoId, utcNow, cancellationToken);
 								}
+								catch (OperationCanceledException ex)
+								{
+									_logger.LogInformation(ex, "Cancelled GC pass for {NamespaceId}", namespaceId);
+								}
 								catch (Exception ex)
 								{
 									_logger.LogError(ex, "Exception while running garbage collection: {Message}", ex.Message);
@@ -1122,7 +1126,7 @@ namespace HordeServer.Storage
 				double minScore = Double.NegativeInfinity;
 				while (!modified)
 				{
-					long length = await _redisService.GetDatabase().SortedSetLengthAsync(checkSet);
+					long length = await _redisService.GetDatabase().SortedSetLengthAsync(checkSet, min: minScore, exclude: Exclude.Start);
 					int batchSize = (int)Math.Min(length, 1024);
 					_logger.LogInformation("Garbage collection queue for namespace {NamespaceId} ({QueueName}) has {Length} entries; taking {Count}", namespaceInfo.Id, checkSet.Inner, length, batchSize);
 
