Horde: Fix log spam/busy loop when running GC sweep.

- Removed the capacity limit on the input buffer block.
- Increased the batch size for items added to the queue.
- Run the pipeline to completion for each batch of items.

diff --git a/Engine/Source/Programs/Horde/Plugins/Storage/HordeServer.Storage/Storage/StorageService.cs b/Engine/Source/Programs/Horde/Plugins/Storage/HordeServer.Storage/Storage/StorageService.cs
--- a/Engine/Source/Programs/Horde/Plugins/Storage/HordeServer.Storage/Storage/StorageService.cs
+++ b/Engine/Source/Programs/Horde/Plugins/Storage/HordeServer.Storage/Storage/StorageService.cs
@@ -1110,24 +1110,12 @@ namespace HordeServer.Storage
 				_logger.LogInformation("Running garbage collection for namespace {NamespaceId}...", namespaceInfo.Id);
 
 				GcSweepState sweepState = new GcSweepState(GetGcTimestamp(utcNow));
-				RedisSortedSetKey<RedisValue> checkSet = GetGcCheckSet(namespaceInfo.Id);
-
-				// Create a dataflow pipeline for deleting blobs
-				BufferBlock<ObjectId> sourceBlock = new BufferBlock<ObjectId>(
-					new DataflowBlockOptions { BoundedCapacity = 1000, CancellationToken = cancellationToken });
-
-				ActionBlock<ObjectId> deleteBlock = new ActionBlock<ObjectId>(blobInfoId => RemoveBlobIfUnreferencedAsync(namespaceInfo, blobInfoId, sweepState, storageConfig, cancellationToken),
-					new ExecutionDataflowBlockOptions { MaxDegreeOfParallelism = 4, CancellationToken = cancellationToken });
 
-				sourceBlock.LinkTo(deleteBlock,
-					new DataflowLinkOptions { PropagateCompletion = true });
-
-				// Push items into the source block
-				double minScore = Double.NegativeInfinity;
+				RedisSortedSetKey<RedisValue> checkSet = GetGcCheckSet(namespaceInfo.Id);
 				while (!modified)
 				{
-					long length = await _redisService.GetDatabase().SortedSetLengthAsync(checkSet, min: minScore, exclude: Exclude.Start);
-					int batchSize = (int)Math.Min(length, 1024);
+					long length = await _redisService.GetDatabase().SortedSetLengthAsync(checkSet);
+					int batchSize = (int)Math.Min(length, 4096);
 					_logger.LogInformation("Garbage collection queue for namespace {NamespaceId} ({QueueName}) has {Length} entries; taking {Count}", namespaceInfo.Id, checkSet.Inner, length, batchSize);
 
 					if (length == 0)
@@ -1135,21 +1123,31 @@ namespace HordeServer.Storage
 						break;
 					}
 
-					SortedSetEntry<RedisValue>[] entries = await _redisService.GetDatabase().SortedSetRangeByScoreWithScoresAsync(checkSet, start: minScore, exclude: Exclude.Start, take: batchSize);
-					foreach (SortedSetEntry<RedisValue> entry in entries)
+					// Create a dataflow pipeline for deleting blobs
+					BufferBlock<ObjectId> sourceBlock = new BufferBlock<ObjectId>(
+						new DataflowBlockOptions { CancellationToken = cancellationToken });
+
+					ActionBlock<ObjectId> deleteBlock = new ActionBlock<ObjectId>(blobInfoId => RemoveBlobIfUnreferencedAsync(namespaceInfo, blobInfoId, sweepState, storageConfig, cancellationToken),
+						new ExecutionDataflowBlockOptions { MaxDegreeOfParallelism = 4, CancellationToken = cancellationToken });
+
+					sourceBlock.LinkTo(deleteBlock,
+						new DataflowLinkOptions { PropagateCompletion = true });
+
+					// Push items into the pipeline
+					RedisValue[] entries = await _redisService.GetDatabase().SortedSetRangeByRankAsync(checkSet, 0, batchSize);
+					foreach (RedisValue entry in entries)
 					{
-						ObjectId blobInfoId = new ObjectId(((byte[]?)entry.ElementValue)!);
+						ObjectId blobInfoId = new ObjectId(((byte[]?)entry)!);
 						if (blobInfoId <= lastImportBlobInfoId)
 						{
 							await sourceBlock.SendAsync(blobInfoId, cancellationToken);
 						}
-						minScore = entry.Score;
 					}
-				}
 
-				// Mark the dataflow as complete and wait for it to flush
-				sourceBlock.Complete();
-				await deleteBlock.Completion;
+					// Mark the dataflow as complete and wait for it to flush
+					sourceBlock.Complete();
+					await deleteBlock.Completion;
+				}
 
 				// Update the 
 				await _gcState.UpdateAsync(state => state.FindOrAddNamespace(namespaceInfo.Id).LastTime = utcNow, cancellationToken);
