Horde: Fix flash when initializing dashboard theme
#skipci

diff --git a/Engine/Source/Programs/Horde/HordeDashboard/src/App.tsx b/Engine/Source/Programs/Horde/HordeDashboard/src/App.tsx
--- a/Engine/Source/Programs/Horde/HordeDashboard/src/App.tsx
+++ b/Engine/Source/Programs/Horde/HordeDashboard/src/App.tsx
@@ -85,6 +85,9 @@ const Main: React.FC = () => {
          ErrorHandler.set({ title: "Error initializing site", reason: reason }, true);
       });
 
+      // avoid flash when initializing dashboard
+      document.body.setAttribute('style', `background: ${dashboard.darktheme ? "#0F0F0F" : "#FAF9F9"}`)   
+
       return (<ThemeProvider applyTo='element' theme={dashboard.darktheme ? darkTheme : lightTheme}>
          <div style={{ position: 'absolute', left: '50%', top: '50%', transform: 'translate(-50%, -50%)' }}>
             <Stack horizontalAlign="center" styles={{ root: { padding: 20, minWidth: 200, minHeight: 100 } }}>
