Horde: Fix Toolbox.json for Agent in build script.

#rnx

diff --git a/Engine/Source/Programs/Horde/BuildHorde.xml b/Engine/Source/Programs/Horde/BuildHorde.xml
--- a/Engine/Source/Programs/Horde/BuildHorde.xml
+++ b/Engine/Source/Programs/Horde/BuildHorde.xml
@@ -296,7 +296,7 @@
 		
 		<Node Name="Build Horde Agent Installer (Toolbox)" Requires="Build Horde Agent Installer" Produces="#HordeAgentMsiToolboxFiles">
 			<Copy Files="UnrealHordeAgent.msi" From="$(ServerToolsStagingDir_HordeAgentMsi)" To="$(ServerToolsStagingDir_HordeAgentMsiToolbox)"/>
-			<Copy Files="Toolbox.json" From="$(RootDir)/Engine/Source/Programs/Horde/HordeAgent/Toolbox.json" To="$(ServerToolsStagingDir_HordeAgentMsiToolbox)"/>
+			<Copy Files="Toolbox.json" From="$(RootDir)/Engine/Source/Programs/Horde/HordeAgent" To="$(ServerToolsStagingDir_HordeAgentMsiToolbox)"/>
 			<Tag Files="$(ServerToolsStagingDir_HordeAgentMsiToolbox)/..." With="#HordeAgentMsiToolboxFiles"/>
 		</Node>
 
diff --git a/Engine/Source/Programs/Horde/HordeAgent/Toolbox.json b/Engine/Source/Programs/Horde/HordeAgent/Toolbox.json
--- a/Engine/Source/Programs/Horde/HordeAgent/Toolbox.json
+++ b/Engine/Source/Programs/Horde/HordeAgent/Toolbox.json
@@ -3,12 +3,14 @@
 	"installCommand": {
 		"fileName": "msiexec.exe",
 		"arguments": [
+			"/i",
 			"UnrealHordeAgent.msi"
 		]
 	},
 	"uninstallCommand": {
 		"fileName": "msiexec.exe",
 		"arguments": [
+			"/x",
 			"UnrealHordeAgent.msi"
 		]
 	}
