Horde: Add an "invalidateCache" query parameter to the agent search method to invalidate any cached results.

#rnx

diff --git a/Engine/Source/Programs/Horde/Plugins/Compute/HordeServer.Compute/Agents/AgentService.cs b/Engine/Source/Programs/Horde/Plugins/Compute/HordeServer.Compute/Agents/AgentService.cs
--- a/Engine/Source/Programs/Horde/Plugins/Compute/HordeServer.Compute/Agents/AgentService.cs
+++ b/Engine/Source/Programs/Horde/Plugins/Compute/HordeServer.Compute/Agents/AgentService.cs
@@ -827,7 +827,7 @@ namespace HordeServer.Agents
 		/// Called by ticker to avoid blocking reads. Updates in the background, in favor of slightly more stale agents
 		/// </summary>
 		/// <param name="cancellationToken">Cancellation token</param>
-		private async Task RefreshCachedAgentsAsync(CancellationToken cancellationToken = default)
+		public async Task RefreshCachedAgentsAsync(CancellationToken cancellationToken = default)
 		{
 			using TelemetrySpan span = _tracer.StartActiveSpan($"{nameof(AgentService)}.{nameof(RefreshCachedAgentsAsync)}");
 
diff --git a/Engine/Source/Programs/Horde/Plugins/Compute/HordeServer.Compute/Agents/AgentsController.cs b/Engine/Source/Programs/Horde/Plugins/Compute/HordeServer.Compute/Agents/AgentsController.cs
--- a/Engine/Source/Programs/Horde/Plugins/Compute/HordeServer.Compute/Agents/AgentsController.cs
+++ b/Engine/Source/Programs/Horde/Plugins/Compute/HordeServer.Compute/Agents/AgentsController.cs
@@ -57,6 +57,7 @@ namespace HordeServer.Agents
 		/// <param name="poolId">The pool containing the agent</param>
 		/// <param name="condition">Arbitrary condition to evaluate against the agents</param>
 		/// <param name="includeDeleted">Whether to include agents marked as deleted</param>
+		/// <param name="invalidateCache">Whether to refresh the cached agents</param>
 		/// <param name="index">First result to return</param>
 		/// <param name="count">Number of results to return</param>
 		/// <param name="modifiedAfter">If set, only returns agents modified after this time</param>
@@ -65,7 +66,7 @@ namespace HordeServer.Agents
 		[HttpGet]
 		[Route("/api/v1/agents")]
 		[ProducesResponseType(typeof(List<GetAgentResponse>), 200)]
-		public async Task<ActionResult<List<object>>> FindAgentsAsync([FromQuery] PoolId? poolId = null, [FromQuery] Condition? condition = null, [FromQuery] bool includeDeleted = false, [FromQuery] int? index = null, [FromQuery] int? count = null, [FromQuery] DateTimeOffset? modifiedAfter = null, [FromQuery] PropertyFilter? filter = null)
+		public async Task<ActionResult<List<object>>> FindAgentsAsync([FromQuery] PoolId? poolId = null, [FromQuery] Condition? condition = null, [FromQuery] bool includeDeleted = false, [FromQuery] bool invalidateCache = false, [FromQuery] int? index = null, [FromQuery] int? count = null, [FromQuery] DateTimeOffset? modifiedAfter = null, [FromQuery] PropertyFilter? filter = null)
 		{
 			using TelemetrySpan span = _tracer.StartActiveSpan($"{nameof(AgentsController)}.{nameof(FindAgentsAsync)}");
 
@@ -74,6 +75,11 @@ namespace HordeServer.Agents
 				return Forbid(AgentAclAction.ListAgents);
 			}
 
+			if (invalidateCache)
+			{
+				await _agentService.RefreshCachedAgentsAsync(HttpContext.RequestAborted);
+			}
+
 			IEnumerable<IAgent> agentsEnumerable = (await _agentService.GetCachedAgentsAsync(HttpContext.RequestAborted))
 				.Where(x => poolId == null || x.Pools.Contains(poolId.Value))
 				.Where(x => modifiedAfter == null || modifiedAfter.Value.UtcDateTime >= x.UpdateTime);
