Horde: Refactor GC pipeline to reduce redundantly fetching the same data from Redis.

#rnx

diff --git a/Engine/Source/Programs/Horde/Plugins/Storage/HordeServer.Storage/Storage/StorageService.cs b/Engine/Source/Programs/Horde/Plugins/Storage/HordeServer.Storage/Storage/StorageService.cs
--- a/Engine/Source/Programs/Horde/Plugins/Storage/HordeServer.Storage/Storage/StorageService.cs
+++ b/Engine/Source/Programs/Horde/Plugins/Storage/HordeServer.Storage/Storage/StorageService.cs
@@ -400,6 +400,20 @@ namespace HordeServer.Storage
 		static readonly FieldDefinition<BlobInfo, string> s_blobAliasField
 			= new StringFieldDefinition<BlobInfo, string>($"{GetFieldName<BlobInfo>(x => x.Aliases)}.{GetFieldName<AliasInfo>(x => x.Name)}");
 
+		class ObjectIdRedisConverter : IRedisConverter<ObjectId>
+		{
+			public ObjectId FromRedisValue(RedisValue value)
+				=> value.IsNullOrEmpty ? ObjectId.Empty : new ObjectId((byte[]?)value);
+
+			public RedisValue ToRedisValue(ObjectId value)
+				=> value.ToByteArray();
+		}
+
+		static StorageService()
+		{
+			RedisSerializer.RegisterConverter<ObjectId, ObjectIdRedisConverter>();
+		}
+
 		/// <summary>
 		/// Constructor
 		/// </summary>
@@ -991,7 +1005,7 @@ namespace HordeServer.Storage
 		static RedisKey GetRedisKey(string suffix) => $"storage:{suffix}";
 		static RedisKey GetRedisKey(NamespaceId namespaceId, string suffix) => GetRedisKey($"{namespaceId}:{suffix}");
 
-		static RedisSortedSetKey<RedisValue> GetGcCheckSet(NamespaceId namespaceId) => new RedisSortedSetKey<RedisValue>(GetRedisKey(namespaceId, "check"));
+		static RedisSortedSetKey<ObjectId> GetGcCheckSet(NamespaceId namespaceId) => new RedisSortedSetKey<ObjectId>(GetRedisKey(namespaceId, "check"));
 
 		/// <summary>
 		/// Find the next namespace to run GC on
@@ -1098,8 +1112,32 @@ namespace HordeServer.Storage
 
 		async Task TickGcForNamespaceAsync(NamespaceInfo namespaceInfo, ObjectId lastImportBlobInfoId, DateTime utcNow, CancellationToken cancellationToken)
 		{
-			bool modified = false;
-			using IDisposable? notification = _storageConfig.OnChange((_, _) => modified = true);
+			// Runs the garbage collector over a particular namespace.
+			// 
+			// Horde assigns a unique id to each blob in the DB when uploaded, and we maintain a reverse index of all things that reference a blob. At any
+			// moment in time, we can query this index to determine if a blob is referenced. (TODO: May want a grace period where we consider a ref's old AND
+			// new value to handle races for updates?)
+			//
+			// Since we're not a CAS and never recycle ids, we can guarantee that any blobs that do not have any references cannot be reintroduced into the
+			// live set, and don't have to worry about races with unreferenced blobs becoming referenced again.
+			//
+			// To prevent newly added blobs being GC'd before their references are present, we have a maximum object id that we consider for GC based on
+			// creation timestamp. After the grace period elapses, all blobs are added to the reachability check set.
+			//
+			// For reachability analysis, we maintain a set of blob ids that need checking in Redis. Whenever a ref to a blob is removed, we add it to the set
+			// with a score derived from the current time. If a second request to check a blob changes, we update the time. We leave the blob id in the set 
+			// until we've done our reachability checks on it, and only remove it IFF the score matches the original value - preventing the second check
+			// request being removed from the set.
+			//
+			// Querying the reachability set in score order ensures that heavily referenced blobs are continually pushed to the end, reducing the number of times
+			// we need to do redundant checks on it.
+			// 
+			// There is a potential for the enumeration to skip over blobs; either because the check set drains before new items are added, or because
+			// new scores are assigned before items are actually pushed to Redis. This isn't really a problem; we'll catch them the next
+			// iteration and will process them eventually.
+
+			using CancellationTokenSource cancellationSource = CancellationTokenSource.CreateLinkedTokenSource(cancellationToken);
+			using IDisposable? notification = _storageConfig.OnChange((_, _) => cancellationSource.Cancel());
 
 			StorageConfig storageConfig = _storageConfig.CurrentValue;
 			if (storageConfig.EnableGc || storageConfig.EnableGcVerification)
@@ -1109,52 +1147,62 @@ namespace HordeServer.Storage
 				Stopwatch timer = Stopwatch.StartNew();
 				_logger.LogInformation("Running garbage collection for namespace {NamespaceId}...", namespaceInfo.Id);
 
+				// Get the current GC timestamp, in minutes since the Unix Epoch
 				GcSweepState sweepState = new GcSweepState(GetGcTimestamp(utcNow));
 
-				RedisSortedSetKey<RedisValue> checkSet = GetGcCheckSet(namespaceInfo.Id);
-				while (!modified)
-				{
-					long length = await _redisService.GetDatabase().SortedSetLengthAsync(checkSet);
-					int batchSize = (int)Math.Min(length, 4096);
-					_logger.LogInformation("Garbage collection queue for namespace {NamespaceId} ({QueueName}) has {Length} entries; taking {Count}", namespaceInfo.Id, checkSet.Inner, length, batchSize);
+				Channel<SortedSetEntry<ObjectId>> channel = Channel.CreateBounded<SortedSetEntry<ObjectId>>(new BoundedChannelOptions(1000));
 
-					if (length == 0)
-					{
-						break;
-					}
+				await using AsyncPipeline pipeline = new AsyncPipeline(cancellationSource.Token);
+				pipeline.AddTask(ctx => FindBlobsForReachabilityCheckAsync(namespaceInfo, lastImportBlobInfoId, channel.Writer, ctx));
+				pipeline.AddTasks(4, channel.Reader, (entry, ctx) => CheckReachabilityAsync(namespaceInfo, entry, sweepState, storageConfig, ctx));
+				await pipeline.WaitForCompletionAsync();
+
+				// Update the GC timestamp
+				await _gcState.UpdateAsync(state => state.FindOrAddNamespace(namespaceInfo.Id).LastTime = utcNow, cancellationToken);
+				_logger.LogInformation("Finished garbage collection for namespace {NamespaceId} in {TimeSecs}s ({NumItems} removed)", namespaceInfo.Id, timer.Elapsed.TotalSeconds, sweepState.NumRemovedItems);
+			}
+		}
 
-					// Create a pipeline for deleting blobs
-					Channel<ObjectId> channel = Channel.CreateUnbounded<ObjectId>();
+		// Reads batches of blobs from the check set and queues them for reachability checking
+		async Task FindBlobsForReachabilityCheckAsync(NamespaceInfo namespaceInfo, ObjectId lastImportBlobInfoId, ChannelWriter<SortedSetEntry<ObjectId>> writer, CancellationToken cancellationToken)
+		{
+			RedisSortedSetKey<ObjectId> checkSet = GetGcCheckSet(namespaceInfo.Id);
+
+			double minScore = Double.NegativeInfinity;
+			for (; ; )
+			{
+				long length = await _redisService.GetDatabase().SortedSetLengthAsync(checkSet, min: minScore, exclude: Exclude.Stop);
+				int batchSize = (int)Math.Min(length, 4096);
+				_logger.LogInformation("Garbage collection queue for namespace {NamespaceId} ({QueueName}) has {Length} entries; taking {Count}", namespaceInfo.Id, checkSet.Inner, length, batchSize);
 
-					await using AsyncPipeline pipeline = new AsyncPipeline(cancellationToken);
-					pipeline.AddTasks<ObjectId>(4, channel.Reader, (blobInfoId, ctx) => RemoveBlobIfUnreferencedAsync(namespaceInfo, blobInfoId, sweepState, storageConfig, ctx));
+				if (length == 0)
+				{
+					writer.Complete();
+					break;
+				}
 
-					RedisValue[] entries = await _redisService.GetDatabase().SortedSetRangeByRankAsync(checkSet, 0, batchSize);
-					foreach(RedisValue entry in entries)
+				SortedSetEntry<ObjectId>[] entries = await _redisService.GetDatabase().SortedSetRangeByScoreWithScoresAsync(checkSet, start: minScore, exclude: Exclude.Start, take: batchSize);
+				foreach (SortedSetEntry<ObjectId> entry in entries)
+				{
+					ObjectId blobInfoId = entry.Element;
+					if (blobInfoId <= lastImportBlobInfoId)
 					{
-						ObjectId blobInfoId = new ObjectId(((byte[]?)entry)!);
-						if (blobInfoId <= lastImportBlobInfoId)
-						{
-							await channel.Writer.WriteAsync(blobInfoId, cancellationToken);
-						}
+						await writer.WriteAsync(new SortedSetEntry<ObjectId>(blobInfoId, entry.Score), cancellationToken);
 					}
-
-					channel.Writer.TryComplete();
-					await pipeline.WaitForCompletionAsync();
+					minScore = entry.Score;
 				}
-
-				// Update the GC timestamp
-				await _gcState.UpdateAsync(state => state.FindOrAddNamespace(namespaceInfo.Id).LastTime = utcNow, cancellationToken);
-				_logger.LogInformation("Finished garbage collection for namespace {NamespaceId} in {TimeSecs}s ({NumItems} removed)", namespaceInfo.Id, timer.Elapsed.TotalSeconds, sweepState.NumRemovedItems);
 			}
 		}
 
-		async Task RemoveBlobIfUnreferencedAsync(NamespaceInfo namespaceInfo, ObjectId blobInfoId, GcSweepState state, StorageConfig storageConfig, CancellationToken cancellationToken)
+		// Checks whether an individual blob can be removed
+		async Task CheckReachabilityAsync(NamespaceInfo namespaceInfo, SortedSetEntry<ObjectId> entry, GcSweepState state, StorageConfig storageConfig, CancellationToken cancellationToken)
 		{
+			ObjectId blobInfoId = new ObjectId(((byte[]?)entry.ElementValue)!);
+
 			using TelemetrySpan span = _tracer.StartActiveSpan($"{nameof(StorageService)}.{nameof(TickGcForNamespaceAsync)}");
 			span.SetAttribute("BlobId", blobInfoId.ToString());
 
-			RedisSortedSetKey<RedisValue> checkSet = GetGcCheckSet(namespaceInfo.Id);
+			RedisSortedSetKey<ObjectId> checkSet = GetGcCheckSet(namespaceInfo.Id);
 			if (!await IsBlobReferencedAsync(blobInfoId, cancellationToken))
 			{
 				BlobInfo? info;
@@ -1171,7 +1219,7 @@ namespace HordeServer.Storage
 				{
 					if (info.Imports != null)
 					{
-						SortedSetEntry<RedisValue>[] entries = info.Imports.Select(x => new SortedSetEntry<RedisValue>(x.ToByteArray(), state.GetNextScore())).ToArray();
+						SortedSetEntry<ObjectId>[] entries = info.Imports.Select(x => new SortedSetEntry<ObjectId>(x, state.GetNextScore())).ToArray();
 						_ = _redisService.GetDatabase().SortedSetAddAsync(checkSet, entries, flags: CommandFlags.FireAndForget);
 					}
 
@@ -1187,7 +1235,11 @@ namespace HordeServer.Storage
 				}
 			}
 
-			_ = _redisService.GetDatabase().SortedSetRemoveAsync(checkSet, blobInfoId.ToByteArray(), CommandFlags.FireAndForget);
+			// Remove the item from the check set iff the score is the same
+			ITransaction transaction = _redisService.GetDatabase().CreateTransaction();
+			transaction.AddCondition(checkSet.SortedSetEqual(entry.Element, entry.Score));
+			_ = transaction.SortedSetRemoveAsync(checkSet, entry.Element);
+			_ = transaction.ExecuteAsync(flags: CommandFlags.FireAndForget);
 		}
 
 		static void SyncNamespaceList(GcState state, List<NamespaceConfig> namespaces)
@@ -1210,7 +1262,7 @@ namespace HordeServer.Storage
 		void AddGcCheckRecord(NamespaceId namespaceId, ObjectId id)
 		{
 			double score = GetGcTimestamp();
-			_ = _redisService.GetDatabase().SortedSetAddAsync(GetGcCheckSet(namespaceId), id.ToByteArray(), score, flags: CommandFlags.FireAndForget);
+			_ = _redisService.GetDatabase().SortedSetAddAsync(GetGcCheckSet(namespaceId), id, score, flags: CommandFlags.FireAndForget);
 		}
 
 		#endregion
