Horde: Fix compile errors from package upgrades.

#rnx

diff --git a/Engine/Source/Programs/Horde/Plugins/Ddc/HordeServer.Ddc.Tests/Ddc/Functional/ContentId/CompressedBlobTests.cs b/Engine/Source/Programs/Horde/Plugins/Ddc/HordeServer.Ddc.Tests/Ddc/Functional/ContentId/CompressedBlobTests.cs
--- a/Engine/Source/Programs/Horde/Plugins/Ddc/HordeServer.Ddc.Tests/Ddc/Functional/ContentId/CompressedBlobTests.cs
+++ b/Engine/Source/Programs/Horde/Plugins/Ddc/HordeServer.Ddc.Tests/Ddc/Functional/ContentId/CompressedBlobTests.cs
@@ -162,7 +162,7 @@ namespace HordeServer.Tests.Ddc.FunctionalTests.CompressedBlobs
 						blocksToCompress.Add(block);
 					}
 					Hash blake3Hash = hasher.Finalize();
-					byte[] hash = blake3Hash.AsSpanUnsafe().Slice(0, 20).ToArray();
+					byte[] hash = blake3Hash.AsSpan().Slice(0, 20).ToArray();
 					uncompressedContentHash = new IoHash(hash);
 
 					await using FileStream fs = tempCompressedFile.OpenWrite();
diff --git a/Engine/Source/Programs/Horde/Plugins/Ddc/HordeServer.Ddc/Ddc/BlobId.cs b/Engine/Source/Programs/Horde/Plugins/Ddc/HordeServer.Ddc/Ddc/BlobId.cs
--- a/Engine/Source/Programs/Horde/Plugins/Ddc/HordeServer.Ddc/Ddc/BlobId.cs
+++ b/Engine/Source/Programs/Horde/Plugins/Ddc/HordeServer.Ddc/Ddc/BlobId.cs
@@ -94,7 +94,7 @@ namespace HordeServer.Ddc
 			}
 
 			// we only keep the first 20 bytes of the Blake3 hash
-			Span<byte> hash = blake3Hash.AsSpanUnsafe().Slice(0, 20);
+			Span<byte> hash = blake3Hash.AsSpan().Slice(0, 20);
 			return new BlobId(hash.ToArray());
 		}
 
@@ -115,7 +115,7 @@ namespace HordeServer.Ddc
 			}
 
 			// we only keep the first 20 bytes of the Blake3 hash
-			Span<byte> hash = blake3Hash.AsSpanUnsafe().Slice(0, 20);
+			Span<byte> hash = blake3Hash.AsSpan().Slice(0, 20);
 			return new BlobId(hash.ToArray());
 		}
 
@@ -142,7 +142,7 @@ namespace HordeServer.Ddc
 				Hash blake3Hash = hasher.Finalize();
 
 				// we only keep the first 20 bytes of the Blake3 hash
-				byte[] hash = blake3Hash.AsSpanUnsafe().Slice(0, 20).ToArray();
+				byte[] hash = blake3Hash.AsSpan().Slice(0, 20).ToArray();
 				return new BlobId(hash);
 			}
 			finally
diff --git a/Engine/Source/Programs/Horde/Plugins/Ddc/HordeServer.Ddc/Ddc/CompressedBufferUtils.cs b/Engine/Source/Programs/Horde/Plugins/Ddc/HordeServer.Ddc/Ddc/CompressedBufferUtils.cs
--- a/Engine/Source/Programs/Horde/Plugins/Ddc/HordeServer.Ddc/Ddc/CompressedBufferUtils.cs
+++ b/Engine/Source/Programs/Horde/Plugins/Ddc/HordeServer.Ddc/Ddc/CompressedBufferUtils.cs
@@ -354,7 +354,7 @@ namespace HordeServer.Ddc
 			}
 
 			Hash blake3Hash = hasher.Finalize();
-			byte[] hashData = blake3Hash.AsSpanUnsafe().Slice(0, 20).ToArray();
+			byte[] hashData = blake3Hash.AsSpan().Slice(0, 20).ToArray();
 			IoHash hash = new IoHash(hashData);
 
 			CompressedBufferHeader header = new CompressedBufferHeader
diff --git a/Engine/Source/Programs/Horde/Plugins/Ddc/HordeServer.Ddc/Ddc/ContentHash.cs b/Engine/Source/Programs/Horde/Plugins/Ddc/HordeServer.Ddc/Ddc/ContentHash.cs
--- a/Engine/Source/Programs/Horde/Plugins/Ddc/HordeServer.Ddc/Ddc/ContentHash.cs
+++ b/Engine/Source/Programs/Horde/Plugins/Ddc/HordeServer.Ddc/Ddc/ContentHash.cs
@@ -104,7 +104,7 @@ namespace HordeServer.Ddc
 			Hash blake3Hash = hasher.Finalize();
 
 			// we only keep the first 20 bytes of the Blake3 hash
-			Span<byte> hash = blake3Hash.AsSpanUnsafe().Slice(0, HashLength);
+			Span<byte> hash = blake3Hash.AsSpan().Slice(0, HashLength);
 
 			return new ContentHash(hash.ToArray());
 		}
