Horde: Exclude the first item from the GC check set when deciding whether to exit.

#rnx

diff --git a/Engine/Source/Programs/Horde/Plugins/Storage/HordeServer.Storage/Storage/StorageService.cs b/Engine/Source/Programs/Horde/Plugins/Storage/HordeServer.Storage/Storage/StorageService.cs
--- a/Engine/Source/Programs/Horde/Plugins/Storage/HordeServer.Storage/Storage/StorageService.cs
+++ b/Engine/Source/Programs/Horde/Plugins/Storage/HordeServer.Storage/Storage/StorageService.cs
@@ -1171,7 +1171,7 @@ namespace HordeServer.Storage
 			double minScore = Double.NegativeInfinity;
 			for (; ; )
 			{
-				long length = await _redisService.GetDatabase().SortedSetLengthAsync(checkSet, min: minScore, exclude: Exclude.Stop);
+				long length = await _redisService.GetDatabase().SortedSetLengthAsync(checkSet, min: minScore, exclude: Exclude.Start);
 				int batchSize = (int)Math.Min(length, 4096);
 				_logger.LogInformation("Garbage collection queue for namespace {NamespaceId} ({QueueName}) has {Length} entries; taking {Count}", namespaceInfo.Id, checkSet.Inner, length, batchSize);
 
