Horde: Allow log cancellation exceptions to propagate up, rather than trying to write data to the response stream.

diff --git a/Engine/Source/Programs/Horde/Plugins/Compute/HordeServer.Compute/Logs/LogRpcService.cs b/Engine/Source/Programs/Horde/Plugins/Compute/HordeServer.Compute/Logs/LogRpcService.cs
--- a/Engine/Source/Programs/Horde/Plugins/Compute/HordeServer.Compute/Logs/LogRpcService.cs
+++ b/Engine/Source/Programs/Horde/Plugins/Compute/HordeServer.Compute/Logs/LogRpcService.cs
@@ -99,7 +99,7 @@ namespace HordeServer.Logs
 						_logger.LogDebug("Waiting for tail next on log {LogId}", logId);
 						Task<int> waitTask = _logTailService.WaitForTailNextAsync(logId, cancellationSource.Token);
 
-						await Task.WhenAny(waitTask, moveNextTask);
+						Task completeTask = await Task.WhenAny(waitTask, moveNextTask);
 						await cancellationSource.CancelAsync();
 
 						try
@@ -110,6 +110,8 @@ namespace HordeServer.Logs
 						{
 							_logger.LogWarning(ex, "Exception while waiting for tail next");
 						}
+
+						await completeTask; // Allow moveNextTask to throw if it was the thing that finished
 					}
 				}
 
