Horde: Do not propagate non-cancellation exceptions that are trapped by the GRPC interceptor without wrapping and logging them.

diff --git a/Engine/Source/Programs/Horde/HordeServer/Startup.cs b/Engine/Source/Programs/Horde/HordeServer/Startup.cs
--- a/Engine/Source/Programs/Horde/HordeServer/Startup.cs
+++ b/Engine/Source/Programs/Horde/HordeServer/Startup.cs
@@ -153,6 +153,11 @@ namespace HordeServer
 				{
 					await callFunc();
 				}
+				catch (OperationCanceledException ex)
+				{
+					_logger.LogDebug(ex, "Rpc call to {Method} was cancelled", context.Method);
+					throw;
+				}
 				catch (StructuredRpcException ex)
 				{
 #pragma warning disable CA2254 // Template should be a static expression
@@ -162,16 +167,8 @@ namespace HordeServer
 				}
 				catch (Exception ex)
 				{
-					if (context.CancellationToken.IsCancellationRequested)
-					{
-						_logger.LogInformation(ex, "Call to method {Method} was cancelled", context.Method);
-						throw;
-					}
-					else
-					{
-						_logger.LogError(ex, "Exception in call to {Method}", context.Method);
-						throw new RpcException(new Status(StatusCode.Internal, $"An exception was thrown on the server: {ex}"));
-					}
+					_logger.LogError(ex, "Exception in call to {Method}", context.Method);
+					throw new RpcException(new Status(StatusCode.Internal, $"An exception was thrown on the server: {ex}"));
 				}
 			}
 		}
diff --git a/Engine/Source/Programs/Horde/Plugins/Compute/HordeServer.Compute/Logs/LogRpcService.cs b/Engine/Source/Programs/Horde/Plugins/Compute/HordeServer.Compute/Logs/LogRpcService.cs
--- a/Engine/Source/Programs/Horde/Plugins/Compute/HordeServer.Compute/Logs/LogRpcService.cs
+++ b/Engine/Source/Programs/Horde/Plugins/Compute/HordeServer.Compute/Logs/LogRpcService.cs
@@ -111,7 +111,7 @@ namespace HordeServer.Logs
 							_logger.LogWarning(ex, "Exception while waiting for tail next");
 						}
 
-						await completeTask; // Allow moveNextTask to throw if it was the thing that finished
+						await completeTask; // Allow moveNextTask to throw if it finished due to cancellation, so we don't write to 
 					}
 				}
 
