Horde: Run GC a few times during tests since timing of deletion is not deterministic.

#rnx

diff --git a/Engine/Source/Programs/Horde/Plugins/Storage/HordeServer.Storage.Tests/Storage/GcServiceTests.cs b/Engine/Source/Programs/Horde/Plugins/Storage/HordeServer.Storage.Tests/Storage/GcServiceTests.cs
--- a/Engine/Source/Programs/Horde/Plugins/Storage/HordeServer.Storage.Tests/Storage/GcServiceTests.cs
+++ b/Engine/Source/Programs/Horde/Plugins/Storage/HordeServer.Storage.Tests/Storage/GcServiceTests.cs
@@ -50,7 +50,10 @@ namespace HordeServer.Tests.Storage
 
 			HashSet<BlobLocator> nodes = await FindNodesAsync(store, roots);
 
-			await Clock.AdvanceAsync(TimeSpan.FromDays(1.0));
+			for (int idx = 0; idx < 3; idx++)
+			{
+				await Clock.AdvanceAsync(TimeSpan.FromDays(1.0));
+			}
 
 			MemoryObjectStore backend = (MemoryObjectStore)ServiceProvider.GetRequiredService<IObjectStoreFactory>().CreateObjectStore(storageConfig.Backends[0]);
 
