Horde: Keep running the GC loop until there are no new items available.

#rnx

diff --git a/Engine/Source/Programs/Horde/Plugins/Storage/HordeServer.Storage/Storage/StorageService.cs b/Engine/Source/Programs/Horde/Plugins/Storage/HordeServer.Storage/Storage/StorageService.cs
--- a/Engine/Source/Programs/Horde/Plugins/Storage/HordeServer.Storage/Storage/StorageService.cs
+++ b/Engine/Source/Programs/Horde/Plugins/Storage/HordeServer.Storage/Storage/StorageService.cs
@@ -1150,12 +1150,16 @@ namespace HordeServer.Storage
 				// Get the current GC timestamp, in minutes since the Unix Epoch
 				GcSweepState sweepState = new GcSweepState(GetGcTimestamp(utcNow));
 
-				Channel<SortedSetEntry<ObjectId>> channel = Channel.CreateBounded<SortedSetEntry<ObjectId>>(new BoundedChannelOptions(1000));
+				RedisSortedSetKey<ObjectId> checkSet = GetGcCheckSet(namespaceInfo.Id);
+				while (await _redisService.GetDatabase().SortedSetLengthAsync(checkSet, max: sweepState.GetNextScore()) > 0)
+				{
+					Channel<SortedSetEntry<ObjectId>> channel = Channel.CreateBounded<SortedSetEntry<ObjectId>>(new BoundedChannelOptions(20));
 
-				await using AsyncPipeline pipeline = new AsyncPipeline(cancellationSource.Token);
-				pipeline.AddTask(ctx => FindBlobsForReachabilityCheckAsync(namespaceInfo, lastImportBlobInfoId, channel.Writer, ctx));
-				pipeline.AddTasks(4, channel.Reader, (entry, ctx) => CheckReachabilityAsync(namespaceInfo, entry, sweepState, storageConfig, ctx));
-				await pipeline.WaitForCompletionAsync();
+					await using AsyncPipeline pipeline = new AsyncPipeline(cancellationSource.Token);
+					pipeline.AddTask(ctx => FindBlobsForReachabilityCheckAsync(namespaceInfo, lastImportBlobInfoId, channel.Writer, ctx));
+					pipeline.AddTasks(8, channel.Reader, (entry, ctx) => CheckReachabilityAsync(namespaceInfo, entry, sweepState, storageConfig, ctx));
+					await pipeline.WaitForCompletionAsync();
+				}
 
 				// Update the GC timestamp
 				await _gcState.UpdateAsync(state => state.FindOrAddNamespace(namespaceInfo.Id).LastTime = utcNow, cancellationToken);
@@ -1172,7 +1176,7 @@ namespace HordeServer.Storage
 			for (; ; )
 			{
 				long length = await _redisService.GetDatabase().SortedSetLengthAsync(checkSet, min: minScore, exclude: Exclude.Start);
-				int batchSize = (int)Math.Min(length, 4096);
+				int batchSize = (int)Math.Min(length, 50);
 				_logger.LogInformation("Garbage collection queue for namespace {NamespaceId} ({QueueName}) has {Length} entries; taking {Count}", namespaceInfo.Id, checkSet.Inner, length, batchSize);
 
 				if (length == 0)
