Horde: Add Download with UGS to artifact download options
#skipci

diff --git a/Engine/Source/Programs/Horde/HordeDashboard/src/components/artifacts/ArtifactsModal.tsx b/Engine/Source/Programs/Horde/HordeDashboard/src/components/artifacts/ArtifactsModal.tsx
--- a/Engine/Source/Programs/Horde/HordeDashboard/src/components/artifacts/ArtifactsModal.tsx
+++ b/Engine/Source/Programs/Horde/HordeDashboard/src/components/artifacts/ArtifactsModal.tsx
@@ -503,6 +503,19 @@ const DownloadButton: React.FC<{ handler: ArtifactsHandler, openArtifactInfo: ()
       directionalHint: DirectionalHint.bottomRightEdge
    };
 
+   if (handler.artifact?.id) {
+      const id = handler.artifact.id;
+      downloadProps.items.unshift(
+         {
+            key: 'download_ugs',
+            text: 'Download with UGS',
+            onClick: () => {
+               window.location.assign(`/api/v2/artifacts/${id}/download?format=ugs`);
+            }
+         }
+      )
+   }
+
 
    return <Stack id="callout_target_artifactinfo" horizontal verticalAlign="center" tokens={{ childrenGap: 8 }}>
       <PrimaryButton split menuProps={downloadProps} styles={{ root: { fontFamily: 'Horde Open Sans SemiBold !important' } }} disabled={!selection.filesSelected && !selection.directoriesSelected} onClick={async () => {
