Horde: Use serilog rather than Microsoft logger for tests.

#rnx

diff --git a/Engine/Source/Programs/Horde/HordeServer.Tests.Shared/ServerTestSetup.cs b/Engine/Source/Programs/Horde/HordeServer.Tests.Shared/ServerTestSetup.cs
--- a/Engine/Source/Programs/Horde/HordeServer.Tests.Shared/ServerTestSetup.cs
+++ b/Engine/Source/Programs/Horde/HordeServer.Tests.Shared/ServerTestSetup.cs
@@ -18,6 +18,10 @@ using Microsoft.Extensions.Hosting;
 using Microsoft.Extensions.Logging;
 using Microsoft.Extensions.Options;
 using OpenTelemetry.Trace;
+using Serilog;
+using Serilog.Exceptions;
+using Serilog.Exceptions.Core;
+using Serilog.Exceptions.Grpc.Destructurers;
 
 namespace HordeServer.Tests
 {
@@ -55,6 +59,17 @@ namespace HordeServer.Tests
 			PatchDatadogWriter();
 		}
 
+		static ServerTestSetup()
+		{
+			Serilog.Log.Logger = new LoggerConfiguration()
+				.Enrich.FromLogContext()
+				.Enrich.WithExceptionDetails(new DestructuringOptionsBuilder()
+					.WithDefaultDestructurers()
+					.WithDestructurers(new[] { new RpcExceptionDestructurer() }))
+				.WriteTo.Console(outputTemplate: "[{Timestamp:HH:mm:ss} {Level:w3}] {Indent}{Message:l}{NewLine}{Exception}")
+				.CreateLogger();
+		}
+
 		protected void SetConfig(GlobalConfig globalConfig)
 		{
 			globalConfig.PostLoad(ServerSettings, _pluginCollection.LoadedPlugins, DefaultAclModifiers);
@@ -102,7 +117,7 @@ namespace HordeServer.Tests
 
 			services.AddSingleton<IAccountCollection, AccountCollection>();
 
-			services.AddLogging(builder => { builder.AddConsole().SetMinimumLevel(LogLevel.Debug); });
+			services.AddLogging(builder => builder.AddSerilog());
 			services.AddSingleton<IMemoryCache>(sp => new MemoryCache(new MemoryCacheOptions { }));
 
 			services.AddSingleton<OpenTelemetry.Trace.Tracer>(sp => TracerProvider.Default.GetTracer("TestTracer"));
