Horde: Remove the need for multiple passes over the GC check set in order to clear the queue. Worker tasks now signal when the Redis set is modified, and the ingestion task can wait on this before completing the input channel.

#rnx

diff --git a/Engine/Source/Programs/Horde/Plugins/Storage/HordeServer.Storage.Tests/Storage/GcServiceTests.cs b/Engine/Source/Programs/Horde/Plugins/Storage/HordeServer.Storage.Tests/Storage/GcServiceTests.cs
--- a/Engine/Source/Programs/Horde/Plugins/Storage/HordeServer.Storage.Tests/Storage/GcServiceTests.cs
+++ b/Engine/Source/Programs/Horde/Plugins/Storage/HordeServer.Storage.Tests/Storage/GcServiceTests.cs
@@ -49,11 +49,7 @@ namespace HordeServer.Tests.Storage
 			}
 
 			HashSet<BlobLocator> nodes = await FindNodesAsync(store, roots);
-
-			for (int idx = 0; idx < 3; idx++)
-			{
-				await Clock.AdvanceAsync(TimeSpan.FromDays(1.0));
-			}
+			await Clock.AdvanceAsync(TimeSpan.FromDays(1.0));
 
 			MemoryObjectStore backend = (MemoryObjectStore)ServiceProvider.GetRequiredService<IObjectStoreFactory>().CreateObjectStore(storageConfig.Backends[0]);
 
diff --git a/Engine/Source/Programs/Horde/Plugins/Storage/HordeServer.Storage/Storage/StorageService.cs b/Engine/Source/Programs/Horde/Plugins/Storage/HordeServer.Storage/Storage/StorageService.cs
--- a/Engine/Source/Programs/Horde/Plugins/Storage/HordeServer.Storage/Storage/StorageService.cs
+++ b/Engine/Source/Programs/Horde/Plugins/Storage/HordeServer.Storage/Storage/StorageService.cs
@@ -1150,16 +1150,13 @@ namespace HordeServer.Storage
 				// Get the current GC timestamp, in minutes since the Unix Epoch
 				GcSweepState sweepState = new GcSweepState(GetGcTimestamp(utcNow));
 
-				RedisSortedSetKey<ObjectId> checkSet = GetGcCheckSet(namespaceInfo.Id);
-				while (await _redisService.GetDatabase().SortedSetLengthAsync(checkSet, max: sweepState.GetNextScore()) > 0)
-				{
-					Channel<SortedSetEntry<ObjectId>> channel = Channel.CreateBounded<SortedSetEntry<ObjectId>>(new BoundedChannelOptions(20));
+				AsyncEvent queueChangeEvent = new AsyncEvent();
+				Channel<SortedSetEntry<ObjectId>> channel = Channel.CreateBounded<SortedSetEntry<ObjectId>>(new BoundedChannelOptions(128));
 
-					await using AsyncPipeline pipeline = new AsyncPipeline(cancellationSource.Token);
-					pipeline.AddTask(ctx => FindBlobsForReachabilityCheckAsync(namespaceInfo, lastImportBlobInfoId, channel.Writer, ctx));
-					pipeline.AddTasks(8, channel.Reader, (entry, ctx) => CheckReachabilityAsync(namespaceInfo, entry, sweepState, storageConfig, ctx));
-					await pipeline.WaitForCompletionAsync();
-				}
+				await using AsyncPipeline pipeline = new AsyncPipeline(cancellationSource.Token);
+				pipeline.AddTask(ctx => FindBlobsForReachabilityCheckAsync(namespaceInfo, lastImportBlobInfoId, channel.Writer, queueChangeEvent, ctx));
+				pipeline.AddTasks(8, channel.Reader, (entry, ctx) => CheckReachabilityAsync(namespaceInfo, entry, sweepState, storageConfig, queueChangeEvent, ctx));
+				await pipeline.WaitForCompletionAsync();
 
 				// Update the GC timestamp
 				await _gcState.UpdateAsync(state => state.FindOrAddNamespace(namespaceInfo.Id).LastTime = utcNow, cancellationToken);
@@ -1168,38 +1165,53 @@ namespace HordeServer.Storage
 		}
 
 		// Reads batches of blobs from the check set and queues them for reachability checking
-		async Task FindBlobsForReachabilityCheckAsync(NamespaceInfo namespaceInfo, ObjectId lastImportBlobInfoId, ChannelWriter<SortedSetEntry<ObjectId>> writer, CancellationToken cancellationToken)
+		async Task FindBlobsForReachabilityCheckAsync(NamespaceInfo namespaceInfo, ObjectId lastImportBlobInfoId, ChannelWriter<SortedSetEntry<ObjectId>> writer, AsyncEvent queueChangeEvent, CancellationToken cancellationToken)
 		{
-			RedisSortedSetKey<ObjectId> checkSet = GetGcCheckSet(namespaceInfo.Id);
+			// Keep a set of items we previously added to the queue, so we don't queue them again.
+			HashSet<SortedSetEntry<ObjectId>> queuedItems = new HashSet<SortedSetEntry<ObjectId>>();
 
-			double minScore = Double.NegativeInfinity;
+			RedisSortedSetKey<ObjectId> checkSet = GetGcCheckSet(namespaceInfo.Id);
 			for (; ; )
 			{
-				long length = await _redisService.GetDatabase().SortedSetLengthAsync(checkSet, min: minScore, exclude: Exclude.Start);
-				int batchSize = (int)Math.Min(length, 50);
-				_logger.LogInformation("Garbage collection queue for namespace {NamespaceId} ({QueueName}) has {Length} entries; taking {Count}", namespaceInfo.Id, checkSet.Inner, length, batchSize);
+				// Capture the current state of the queue change event. We can use this to speculatively update the queue while being able to check for changes elsewhere
+				queueChangeEvent.Reset();
+				Task queueChangeTask = queueChangeEvent.Task;
 
-				if (length == 0)
+				// Get the current queue stats
+				long length = await _redisService.GetDatabase().SortedSetLengthAsync(checkSet);
+				_logger.LogInformation("Garbage collection queue for namespace {NamespaceId} ({QueueName}) has {Length} entries", namespaceInfo.Id, checkSet.Inner, length);
+
+				// Get the current state of the queue, and filter out any blobs under the threshold for this sweep.
+				SortedSetEntry<ObjectId>[] entries = await _redisService.GetDatabase().SortedSetRangeByRankWithScoresAsync(checkSet, 0, 1024);
+				entries = entries.Where(x => x.Element <= lastImportBlobInfoId).ToArray();
+
+				// If it's empty, we're done.
+				if (entries.Length == 0)
 				{
 					writer.Complete();
 					break;
 				}
 
-				SortedSetEntry<ObjectId>[] entries = await _redisService.GetDatabase().SortedSetRangeByScoreWithScoresAsync(checkSet, start: minScore, exclude: Exclude.Start, take: batchSize);
+				// Write the entries to the channel
+				HashSet<SortedSetEntry<ObjectId>> nextQueuedItems = new HashSet<SortedSetEntry<ObjectId>>(entries.Length);
 				foreach (SortedSetEntry<ObjectId> entry in entries)
 				{
 					ObjectId blobInfoId = entry.Element;
-					if (blobInfoId <= lastImportBlobInfoId)
+					if (!queuedItems.Contains(entry))
 					{
-						await writer.WriteAsync(new SortedSetEntry<ObjectId>(blobInfoId, entry.Score), cancellationToken);
+						await writer.WriteAsync(entry, cancellationToken);
 					}
-					minScore = entry.Score;
+					nextQueuedItems.Add(entry);
 				}
+				queuedItems = nextQueuedItems;
+
+				// Wait for something to be processed before running again
+				await queueChangeTask;
 			}
 		}
 
 		// Checks whether an individual blob can be removed
-		async Task CheckReachabilityAsync(NamespaceInfo namespaceInfo, SortedSetEntry<ObjectId> entry, GcSweepState state, StorageConfig storageConfig, CancellationToken cancellationToken)
+		async Task CheckReachabilityAsync(NamespaceInfo namespaceInfo, SortedSetEntry<ObjectId> entry, GcSweepState state, StorageConfig storageConfig, AsyncEvent queueChangeEvent, CancellationToken cancellationToken)
 		{
 			ObjectId blobInfoId = new ObjectId(((byte[]?)entry.ElementValue)!);
 
@@ -1224,7 +1236,7 @@ namespace HordeServer.Storage
 					if (info.Imports != null)
 					{
 						SortedSetEntry<ObjectId>[] entries = info.Imports.Select(x => new SortedSetEntry<ObjectId>(x, state.GetNextScore())).ToArray();
-						_ = _redisService.GetDatabase().SortedSetAddAsync(checkSet, entries, flags: CommandFlags.FireAndForget);
+						await _redisService.GetDatabase().SortedSetAddAsync(checkSet, entries);
 					}
 
 					ObjectKey objectKey = GetObjectKey(new BlobLocator(info.Path));
@@ -1240,10 +1252,20 @@ namespace HordeServer.Storage
 			}
 
 			// Remove the item from the check set iff the score is the same
-			ITransaction transaction = _redisService.GetDatabase().CreateTransaction();
-			transaction.AddCondition(checkSet.SortedSetEqual(entry.Element, entry.Score));
-			_ = transaction.SortedSetRemoveAsync(checkSet, entry.Element);
-			_ = transaction.ExecuteAsync(flags: CommandFlags.FireAndForget);
+			for (; ; )
+			{
+				ITransaction transaction = _redisService.GetDatabase().CreateTransaction();
+				ConditionResult scoreEqualCondition = transaction.AddCondition(checkSet.SortedSetEqual(entry.Element, entry.Score));
+				_ = transaction.SortedSetRemoveAsync(checkSet, entry.Element);
+
+				if (await transaction.ExecuteAsync() || !scoreEqualCondition.WasSatisfied)
+				{
+					break;
+				}
+			}
+
+			// Flag that the queue has been updated
+			queueChangeEvent.Latch();
 		}
 
 		static void SyncNamespaceList(GcState state, List<NamespaceConfig> namespaces)
