Horde: Support downloading .uartifact files through the Horde command line tool.

#rnx

diff --git a/Engine/Source/Programs/Horde/HordeCmd/Commands/Artifacts/ArtifactDownload.cs b/Engine/Source/Programs/Horde/HordeCmd/Commands/Artifacts/ArtifactDownload.cs
--- a/Engine/Source/Programs/Horde/HordeCmd/Commands/Artifacts/ArtifactDownload.cs
+++ b/Engine/Source/Programs/Horde/HordeCmd/Commands/Artifacts/ArtifactDownload.cs
@@ -18,6 +18,10 @@ namespace Horde.Commands.Artifacts
 		[Description("Unique identifier for the artifact")]
 		public ArtifactId Id { get; set; }
 
+		[CommandLine("-File=")]
+		[Description("Path to a .uartifact file to download")]
+		public FileReference? File { get; set; }
+
 		[CommandLine("-OutputDir=", Required = true)]
 		[Description("Directory to write extracted files.")]
 		public DirectoryReference OutputDir { get; set; } = null!;
@@ -39,14 +43,44 @@ namespace Horde.Commands.Artifacts
 
 		public override async Task<int> ExecuteAsync(ILogger logger)
 		{
-			IArtifact? artifact = await _hordeClient.Artifacts.GetAsync(Id);
-			if (artifact == null)
+			IStorageNamespace store;
+			IBlobRef<DirectoryNode> handle;
+
+			if (File != null && Id == default)
 			{
-				logger.LogError("Artifact {Id} not found", Id);
-				return 1;
+				ArtifactDescriptor descriptor = await ArtifactDescriptor.ReadAsync(File, default);
+
+				Uri serverUrl = new Uri(descriptor.BaseUrl.GetLeftPart(UriPartial.Authority));
+				if (serverUrl != _hordeClient.ServerUrl)
+				{
+					logger.LogError("Artifact references server {ServerUrl}, not the active server {ActiveServerUrl}. Please log in to the correct server.", serverUrl, _hordeClient.ServerUrl);
+					return 1;
+				}
+
+				logger.LogInformation("Downloading artifact {Url}", descriptor.BaseUrl);
+
+				store = _hordeClient.GetStorageNamespace(descriptor.BaseUrl.AbsolutePath);
+				handle = await store.ReadRefAsync<DirectoryNode>(descriptor.RefName);
 			}
+			else if (File == null && Id != default)
+			{
+				IArtifact? artifact = await _hordeClient.Artifacts.GetAsync(Id);
+				if (artifact == null)
+				{
+					logger.LogError("Artifact {Id} not found", Id);
+					return 1;
+				}
 
-			logger.LogInformation("Downloading artifact {Id}: {Description}", Id, artifact.Description);
+				logger.LogInformation("Downloading artifact {Id}: {Description}", Id, artifact.Description);
+
+				store = _hordeClient.GetStorageNamespace(artifact.Id);
+				handle = await store.ReadRefAsync<DirectoryNode>(new RefName("default"));
+			}
+			else
+			{
+				logger.LogError("Either -Id=... or -File=... must be specified.");
+				return 1;
+			}
 
 			if (CleanOutput)
 			{
@@ -55,11 +89,7 @@ namespace Horde.Commands.Artifacts
 			}
 
 			Stopwatch timer = Stopwatch.StartNew();
-
-			IStorageNamespace store = _hordeClient.GetStorageNamespace(artifact.Id);
-			IBlobRef<DirectoryNode> handle = await store.ReadRefAsync<DirectoryNode>(new RefName("default"));
 			await handle.ExtractAsync(OutputDir.ToDirectoryInfo(), new ExtractStatsLogger(logger), logger, CancellationToken.None);
-
 			logger.LogInformation("Elapsed: {Time}s", timer.Elapsed.TotalSeconds);
 
 			if (Stats)
