Horde: Log GC queue stats at most every 30 seconds.

#rnx

diff --git a/Engine/Source/Programs/Horde/Plugins/Storage/HordeServer.Storage/Storage/StorageService.cs b/Engine/Source/Programs/Horde/Plugins/Storage/HordeServer.Storage/Storage/StorageService.cs
--- a/Engine/Source/Programs/Horde/Plugins/Storage/HordeServer.Storage/Storage/StorageService.cs
+++ b/Engine/Source/Programs/Horde/Plugins/Storage/HordeServer.Storage/Storage/StorageService.cs
@@ -1170,6 +1170,8 @@ namespace HordeServer.Storage
 			// Keep a set of items we previously added to the queue, so we don't queue them again.
 			HashSet<SortedSetEntry<ObjectId>> queuedItems = new HashSet<SortedSetEntry<ObjectId>>();
 
+			Stopwatch? timer = null;
+
 			RedisSortedSetKey<ObjectId> checkSet = GetGcCheckSet(namespaceInfo.Id);
 			for (; ; )
 			{
@@ -1177,9 +1179,13 @@ namespace HordeServer.Storage
 				queueChangeEvent.Reset();
 				Task queueChangeTask = queueChangeEvent.Task;
 
-				// Get the current queue stats
-				long length = await _redisService.GetDatabase().SortedSetLengthAsync(checkSet);
-				_logger.LogInformation("Garbage collection queue for namespace {NamespaceId} ({QueueName}) has {Length} entries", namespaceInfo.Id, checkSet.Inner, length);
+				// Print the current queue stats
+				if (timer == null || timer.Elapsed > TimeSpan.FromSeconds(30))
+				{
+					long length = await _redisService.GetDatabase().SortedSetLengthAsync(checkSet);
+					_logger.LogInformation("Garbage collection queue for namespace {NamespaceId} ({QueueName}) has {Length} entries", namespaceInfo.Id, checkSet.Inner, length);
+					timer = Stopwatch.StartNew();
+				}
 
 				// Get the current state of the queue, and filter out any blobs under the threshold for this sweep.
 				SortedSetEntry<ObjectId>[] entries = await _redisService.GetDatabase().SortedSetRangeByRankWithScoresAsync(checkSet, 0, 1024);
@@ -1188,6 +1194,7 @@ namespace HordeServer.Storage
 				// If it's empty, we're done.
 				if (entries.Length == 0)
 				{
+					_logger.LogInformation("Garbage collection complete for namespace {NamespaceId}", namespaceInfo.Id);
 					writer.Complete();
 					break;
 				}
