Horde: Fix GC sweep ending early due to all object ids being outside the grace period for new blobs being added. We now avoid adding any blobs to the GC set that cannot be deleted.

diff --git a/Engine/Source/Programs/Horde/Plugins/Storage/HordeServer.Storage/Storage/StorageService.cs b/Engine/Source/Programs/Horde/Plugins/Storage/HordeServer.Storage/Storage/StorageService.cs
--- a/Engine/Source/Programs/Horde/Plugins/Storage/HordeServer.Storage/Storage/StorageService.cs
+++ b/Engine/Source/Programs/Horde/Plugins/Storage/HordeServer.Storage/Storage/StorageService.cs
@@ -728,7 +728,7 @@ namespace HordeServer.Storage
 				_logger.LogDebug("Adding {NumBlobs} blobs for GC consideration ({FirstId} to {LastId})", current.Count, current[0].Id, current[^1].Id);
 				foreach (BlobInfo blobInfo in current)
 				{
-					AddGcCheckRecord(blobInfo.NamespaceId, blobInfo.Id);
+					AddGcCheckRecord(blobInfo.NamespaceId, blobInfo.Id, latestInfoId);
 				}
 
 				// Update the last imported blob id
@@ -875,6 +875,8 @@ namespace HordeServer.Storage
 
 			using TelemetrySpan span = _tracer.StartActiveSpan($"{nameof(StorageService)}.{nameof(TickRefsAsync)}");
 
+			GcState gcState = await _gcState.GetAsync(cancellationToken);
+
 			FilterDefinition<RefInfo> queryFilter = Builders<RefInfo>.Filter.Exists(x => x.ExpiresAtUtc) & Builders<RefInfo>.Filter.Lt(x => x.ExpiresAtUtc, utcNow);
 			using (IAsyncCursor<RefInfo> cursor = await _refCollection.Find(queryFilter).ToCursorAsync(cancellationToken))
 			{
@@ -888,7 +890,7 @@ namespace HordeServer.Storage
 						_logger.LogInformation("Expired ref {NamespaceId}:{RefName}", refInfo.NamespaceId, refInfo.Name);
 						FilterDefinition<RefInfo> filter = Builders<RefInfo>.Filter.Expr(x => x.Id == refInfo.Id && x.ExpiresAtUtc == refInfo.ExpiresAtUtc);
 						requests.Add(new DeleteOneModel<RefInfo>(filter));
-						AddGcCheckRecord(refInfo.NamespaceId, refInfo.TargetBlobId);
+						AddGcCheckRecord(refInfo.NamespaceId, refInfo.TargetBlobId, gcState.LastImportBlobInfoId);
 						AddRefToCache(refInfo.NamespaceId, refInfo.Name, default);
 					}
 
@@ -924,7 +926,8 @@ namespace HordeServer.Storage
 			if (oldRefInfo != null)
 			{
 				_logger.LogInformation("Deleted ref {NamespaceId}:{RefName}", namespaceId, name);
-				AddGcCheckRecord(namespaceId, oldRefInfo.TargetBlobId);
+				GcState gcState = await _gcState.GetAsync(cancellationToken);
+				AddGcCheckRecord(namespaceId, oldRefInfo.TargetBlobId, gcState.LastImportBlobInfoId);
 				return true;
 			}
 
@@ -988,7 +991,8 @@ namespace HordeServer.Storage
 			RefInfo? oldRefInfo = await _refCollection.FindOneAndReplaceAsync<RefInfo>(x => x.NamespaceId == namespaceId && x.Name == name, newRefInfo, new FindOneAndReplaceOptions<RefInfo> { IsUpsert = true }, cancellationToken);
 			if (oldRefInfo != null)
 			{
-				AddGcCheckRecord(namespaceId, oldRefInfo.TargetBlobId);
+				GcState gcState = await _gcState.GetAsync(cancellationToken);
+				AddGcCheckRecord(namespaceId, oldRefInfo.TargetBlobId, gcState.LastImportBlobInfoId);
 			}
 
 			_logger.LogInformation("Updated ref {NamespaceId}:{RefName}", namespaceId, name);
@@ -1154,8 +1158,8 @@ namespace HordeServer.Storage
 				Channel<SortedSetEntry<ObjectId>> channel = Channel.CreateBounded<SortedSetEntry<ObjectId>>(new BoundedChannelOptions(128));
 
 				await using AsyncPipeline pipeline = new AsyncPipeline(cancellationSource.Token);
-				pipeline.AddTask(ctx => FindBlobsForReachabilityCheckAsync(namespaceInfo, lastImportBlobInfoId, channel.Writer, queueChangeEvent, ctx));
-				pipeline.AddTasks(8, channel.Reader, (entry, ctx) => CheckReachabilityAsync(namespaceInfo, entry, sweepState, storageConfig, queueChangeEvent, ctx));
+				pipeline.AddTask(ctx => FindBlobsForReachabilityCheckAsync(namespaceInfo, channel.Writer, queueChangeEvent, ctx));
+				pipeline.AddTasks(8, channel.Reader, (entry, ctx) => CheckReachabilityAsync(namespaceInfo, entry, lastImportBlobInfoId, sweepState, storageConfig, queueChangeEvent, ctx));
 				await pipeline.WaitForCompletionAsync();
 
 				// Update the GC timestamp
@@ -1165,7 +1169,7 @@ namespace HordeServer.Storage
 		}
 
 		// Reads batches of blobs from the check set and queues them for reachability checking
-		async Task FindBlobsForReachabilityCheckAsync(NamespaceInfo namespaceInfo, ObjectId lastImportBlobInfoId, ChannelWriter<SortedSetEntry<ObjectId>> writer, AsyncEvent queueChangeEvent, CancellationToken cancellationToken)
+		async Task FindBlobsForReachabilityCheckAsync(NamespaceInfo namespaceInfo, ChannelWriter<SortedSetEntry<ObjectId>> writer, AsyncEvent queueChangeEvent, CancellationToken cancellationToken)
 		{
 			// Keep a set of items we previously added to the queue, so we don't queue them again.
 			HashSet<SortedSetEntry<ObjectId>> queuedItems = new HashSet<SortedSetEntry<ObjectId>>();
@@ -1187,11 +1191,8 @@ namespace HordeServer.Storage
 					timer = Stopwatch.StartNew();
 				}
 
-				// Get the current state of the queue, and filter out any blobs under the threshold for this sweep.
+				// Get the current state of the queue. If it's empty, we're done.
 				SortedSetEntry<ObjectId>[] entries = await _redisService.GetDatabase().SortedSetRangeByRankWithScoresAsync(checkSet, 0, 1024);
-				entries = entries.Where(x => x.Element <= lastImportBlobInfoId).ToArray();
-
-				// If it's empty, we're done.
 				if (entries.Length == 0)
 				{
 					_logger.LogInformation("Garbage collection complete for namespace {NamespaceId}", namespaceInfo.Id);
@@ -1218,7 +1219,7 @@ namespace HordeServer.Storage
 		}
 
 		// Checks whether an individual blob can be removed
-		async Task CheckReachabilityAsync(NamespaceInfo namespaceInfo, SortedSetEntry<ObjectId> entry, GcSweepState state, StorageConfig storageConfig, AsyncEvent queueChangeEvent, CancellationToken cancellationToken)
+		async Task CheckReachabilityAsync(NamespaceInfo namespaceInfo, SortedSetEntry<ObjectId> entry, ObjectId lastImportBlobInfoId, GcSweepState state, StorageConfig storageConfig, AsyncEvent queueChangeEvent, CancellationToken cancellationToken)
 		{
 			ObjectId blobInfoId = new ObjectId(((byte[]?)entry.ElementValue)!);
 
@@ -1242,7 +1243,7 @@ namespace HordeServer.Storage
 				{
 					if (info.Imports != null)
 					{
-						SortedSetEntry<ObjectId>[] entries = info.Imports.Select(x => new SortedSetEntry<ObjectId>(x, state.GetNextScore())).ToArray();
+						SortedSetEntry<ObjectId>[] entries = info.Imports.Where(x => x <= lastImportBlobInfoId).Select(x => new SortedSetEntry<ObjectId>(x, state.GetNextScore())).ToArray();
 						await _redisService.GetDatabase().SortedSetAddAsync(checkSet, entries);
 					}
 
@@ -1292,10 +1293,13 @@ namespace HordeServer.Storage
 			state.Namespaces.SortBy(x => x.Id.Text.Text);
 		}
 
-		void AddGcCheckRecord(NamespaceId namespaceId, ObjectId id)
+		void AddGcCheckRecord(NamespaceId namespaceId, ObjectId id, ObjectId lastImportBlobId)
 		{
-			double score = GetGcTimestamp();
-			_ = _redisService.GetDatabase().SortedSetAddAsync(GetGcCheckSet(namespaceId), id, score, flags: CommandFlags.FireAndForget);
+			if (id < lastImportBlobId)
+			{
+				double score = GetGcTimestamp();
+				_ = _redisService.GetDatabase().SortedSetAddAsync(GetGcCheckSet(namespaceId), id, score, flags: CommandFlags.FireAndForget);
+			}
 		}
 
 		#endregion
