Remove unused references in Shared
Pin System.Text.Json and System.Text.RegularExpression version

#rnx

diff --git a/Engine/Source/Programs/AutomationTool/AutomationUtils/AutomationUtils.Automation.csproj b/Engine/Source/Programs/AutomationTool/AutomationUtils/AutomationUtils.Automation.csproj
--- a/Engine/Source/Programs/AutomationTool/AutomationUtils/AutomationUtils.Automation.csproj
+++ b/Engine/Source/Programs/AutomationTool/AutomationUtils/AutomationUtils.Automation.csproj
@@ -72,5 +72,6 @@
 
     <PackageReference Include="Google.Apis.Drive.v3" Version="1.33.1.1225" PrivateAssets="all" />
     <PackageReference Include="Google.Apis.Sheets.v4" Version="1.33.1.1229" PrivateAssets="all" />
+    <PackageReference Include="System.Text.RegularExpressions" Version="4.3.1" />
   </ItemGroup>
 </Project>
\ No newline at end of file
diff --git a/Engine/Source/Programs/Horde/Drivers/JobDriver/JobDriver.csproj b/Engine/Source/Programs/Horde/Drivers/JobDriver/JobDriver.csproj
--- a/Engine/Source/Programs/Horde/Drivers/JobDriver/JobDriver.csproj
+++ b/Engine/Source/Programs/Horde/Drivers/JobDriver/JobDriver.csproj
@@ -92,10 +92,11 @@
     <PackageReference Include="Serilog.Sinks.File" Version="5.0.0" />
     <PackageReference Include="System.Diagnostics.PerformanceCounter" Version="8.0.0" />
     <PackageReference Include="System.ServiceProcess.ServiceController" Version="8.0.0" />
+    <PackageReference Include="System.Text.Json" Version="8.0.4" />
   </ItemGroup>
   
   <ItemGroup>
-    <Protobuf Include="..\..\Plugins\Build\HordeServer.Build\Protos\horde\job_rpc.proto" ProtoRoot="..\..\Plugins\Build\HordeServer.Build\Protos" GrpcServices="Client" Access="Public" AdditionalImportDirs="..\..\..\Shared\EpicGames.Horde\Protos;..\..\Plugins\Compute\HordeServer.Compute\Protos"/>
+    <Protobuf Include="..\..\Plugins\Build\HordeServer.Build\Protos\horde\job_rpc.proto" ProtoRoot="..\..\Plugins\Build\HordeServer.Build\Protos" GrpcServices="Client" Access="Public" AdditionalImportDirs="..\..\..\Shared\EpicGames.Horde\Protos;..\..\Plugins\Compute\HordeServer.Compute\Protos" />
     <Protobuf Include="..\..\Plugins\Build\HordeServer.Build\Protos\horde\job_rpc_messages.proto" ProtoRoot="..\..\Plugins\Build\HordeServer.Build\Protos" Access="Public" AdditionalImportDirs="..\..\Plugins\Compute\HordeServer.Compute\Protos" />
     <Protobuf Include="..\..\Plugins\Build\HordeServer.Build\Protos\horde\common\label_change.proto" ProtoRoot="..\..\Plugins\Build\HordeServer.Build\Protos" Access="Public" />
     <Protobuf Include="..\..\Plugins\Compute\HordeServer.Compute\Protos\horde\messages\agent_workspace.proto" ProtoRoot="..\..\Plugins\Compute\HordeServer.Compute\Protos" Access="Public" />
diff --git a/Engine/Source/Programs/Horde/HordeAgent/HordeAgent.csproj b/Engine/Source/Programs/Horde/HordeAgent/HordeAgent.csproj
--- a/Engine/Source/Programs/Horde/HordeAgent/HordeAgent.csproj
+++ b/Engine/Source/Programs/Horde/HordeAgent/HordeAgent.csproj
@@ -73,6 +73,7 @@
     <PackageReference Include="Serilog.Sinks.File" Version="5.0.0" />
     <PackageReference Include="System.Diagnostics.PerformanceCounter" Version="8.0.0" />
     <PackageReference Include="System.ServiceProcess.ServiceController" Version="8.0.0" />
+    <PackageReference Include="System.Text.Json" Version="8.0.4" />
   </ItemGroup>
 
   <ItemGroup>
@@ -150,21 +151,21 @@
 
   <ItemGroup>
     <Protobuf Include="..\Plugins\Compute\HordeServer.Compute\Protos\horde\horde_rpc.proto" ProtoRoot="..\Plugins\Compute\HordeServer.Compute\Protos" GrpcServices="Client" Access="Public" AdditionalImportDirs="..\Plugins\Compute\HordeServer.Compute\Protos" />
-    <Protobuf Include="..\Plugins\Compute\HordeServer.Compute\Protos\horde\horde_rpc_messages.proto" ProtoRoot="..\Plugins\Compute\HordeServer.Compute\Protos" Access="Public" AdditionalImportDirs="..\Plugins\Compute\HordeServer.Compute\Protos"/>
-    <Protobuf Include="..\Plugins\Compute\HordeServer.Compute\Protos\horde\enrollment_rpc.proto" ProtoRoot="..\Plugins\Compute\HordeServer.Compute\Protos" GrpcServices="Client" Access="Public" AdditionalImportDirs="..\Plugins\Compute\HordeServer.Compute\Protos"/>
-    <Protobuf Include="..\Plugins\Compute\HordeServer.Compute\Protos\horde\enrollment_rpc_messages.proto" ProtoRoot="..\Plugins\Compute\HordeServer.Compute\Protos" Access="Public" AdditionalImportDirs="..\Plugins\Compute\HordeServer.Compute\Protos"/>
-    <Protobuf Include="..\Plugins\Compute\HordeServer.Compute\Protos\horde\common\agent_status.proto" ProtoRoot="..\Plugins\Compute\HordeServer.Compute\Protos" Access="Public" AdditionalImportDirs="..\Plugins\Compute\HordeServer.Compute\Protos"/>
-    <Protobuf Include="..\Plugins\Compute\HordeServer.Compute\Protos\horde\common\lease_outcome.proto" ProtoRoot="..\Plugins\Compute\HordeServer.Compute\Protos" Access="Public" AdditionalImportDirs="..\Plugins\Compute\HordeServer.Compute\Protos"/>
-    <Protobuf Include="..\Plugins\Compute\HordeServer.Compute\Protos\horde\common\lease_state.proto" ProtoRoot="..\Plugins\Compute\HordeServer.Compute\Protos" Access="Public" AdditionalImportDirs="..\Plugins\Compute\HordeServer.Compute\Protos"/>
-    <Protobuf Include="..\Plugins\Compute\HordeServer.Compute\Protos\horde\messages\agent_capabilities.proto" ProtoRoot="..\Plugins\Compute\HordeServer.Compute\Protos" Access="Public" AdditionalImportDirs="..\Plugins\Compute\HordeServer.Compute\Protos"/>
-    <Protobuf Include="..\Plugins\Compute\HordeServer.Compute\Protos\horde\messages\agent_requirements.proto" ProtoRoot="..\Plugins\Compute\HordeServer.Compute\Protos" Access="Public" AdditionalImportDirs="..\Plugins\Compute\HordeServer.Compute\Protos"/>
-    <Protobuf Include="..\Plugins\Compute\HordeServer.Compute\Protos\horde\messages\lease.proto" ProtoRoot="..\Plugins\Compute\HordeServer.Compute\Protos" Access="Public" AdditionalImportDirs="..\Plugins\Compute\HordeServer.Compute\Protos"/>
-    <Protobuf Include="..\Plugins\Compute\HordeServer.Compute\Protos\horde\messages\aws_lambda.proto" ProtoRoot="..\Plugins\Compute\HordeServer.Compute\Protos" Access="Public" AdditionalImportDirs="..\Plugins\Compute\HordeServer.Compute\Protos"/>
-    <Protobuf Include="..\Plugins\Compute\HordeServer.Compute\Protos\horde\tasks\compute_task.proto" ProtoRoot="..\Plugins\Compute\HordeServer.Compute\Protos" Access="Public" AdditionalImportDirs="..\Plugins\Compute\HordeServer.Compute\Protos"/>
-    <Protobuf Include="..\Plugins\Compute\HordeServer.Compute\Protos\horde\tasks\restart_task.proto" ProtoRoot="..\Plugins\Compute\HordeServer.Compute\Protos" Access="Public" AdditionalImportDirs="..\Plugins\Compute\HordeServer.Compute\Protos"/>
-    <Protobuf Include="..\Plugins\Compute\HordeServer.Compute\Protos\horde\tasks\shutdown_task.proto" ProtoRoot="..\Plugins\Compute\HordeServer.Compute\Protos" Access="Public" AdditionalImportDirs="..\Plugins\Compute\HordeServer.Compute\Protos"/>
-    <Protobuf Include="..\Plugins\Compute\HordeServer.Compute\Protos\horde\tasks\test_task.proto" ProtoRoot="..\Plugins\Compute\HordeServer.Compute\Protos" Access="Public" AdditionalImportDirs="..\Plugins\Compute\HordeServer.Compute\Protos"/>
-    <Protobuf Include="..\Plugins\Compute\HordeServer.Compute\Protos\horde\tasks\upgrade_task.proto" ProtoRoot="..\Plugins\Compute\HordeServer.Compute\Protos" Access="Public" AdditionalImportDirs="..\Plugins\Compute\HordeServer.Compute\Protos"/>
+    <Protobuf Include="..\Plugins\Compute\HordeServer.Compute\Protos\horde\horde_rpc_messages.proto" ProtoRoot="..\Plugins\Compute\HordeServer.Compute\Protos" Access="Public" AdditionalImportDirs="..\Plugins\Compute\HordeServer.Compute\Protos" />
+    <Protobuf Include="..\Plugins\Compute\HordeServer.Compute\Protos\horde\enrollment_rpc.proto" ProtoRoot="..\Plugins\Compute\HordeServer.Compute\Protos" GrpcServices="Client" Access="Public" AdditionalImportDirs="..\Plugins\Compute\HordeServer.Compute\Protos" />
+    <Protobuf Include="..\Plugins\Compute\HordeServer.Compute\Protos\horde\enrollment_rpc_messages.proto" ProtoRoot="..\Plugins\Compute\HordeServer.Compute\Protos" Access="Public" AdditionalImportDirs="..\Plugins\Compute\HordeServer.Compute\Protos" />
+    <Protobuf Include="..\Plugins\Compute\HordeServer.Compute\Protos\horde\common\agent_status.proto" ProtoRoot="..\Plugins\Compute\HordeServer.Compute\Protos" Access="Public" AdditionalImportDirs="..\Plugins\Compute\HordeServer.Compute\Protos" />
+    <Protobuf Include="..\Plugins\Compute\HordeServer.Compute\Protos\horde\common\lease_outcome.proto" ProtoRoot="..\Plugins\Compute\HordeServer.Compute\Protos" Access="Public" AdditionalImportDirs="..\Plugins\Compute\HordeServer.Compute\Protos" />
+    <Protobuf Include="..\Plugins\Compute\HordeServer.Compute\Protos\horde\common\lease_state.proto" ProtoRoot="..\Plugins\Compute\HordeServer.Compute\Protos" Access="Public" AdditionalImportDirs="..\Plugins\Compute\HordeServer.Compute\Protos" />
+    <Protobuf Include="..\Plugins\Compute\HordeServer.Compute\Protos\horde\messages\agent_capabilities.proto" ProtoRoot="..\Plugins\Compute\HordeServer.Compute\Protos" Access="Public" AdditionalImportDirs="..\Plugins\Compute\HordeServer.Compute\Protos" />
+    <Protobuf Include="..\Plugins\Compute\HordeServer.Compute\Protos\horde\messages\agent_requirements.proto" ProtoRoot="..\Plugins\Compute\HordeServer.Compute\Protos" Access="Public" AdditionalImportDirs="..\Plugins\Compute\HordeServer.Compute\Protos" />
+    <Protobuf Include="..\Plugins\Compute\HordeServer.Compute\Protos\horde\messages\lease.proto" ProtoRoot="..\Plugins\Compute\HordeServer.Compute\Protos" Access="Public" AdditionalImportDirs="..\Plugins\Compute\HordeServer.Compute\Protos" />
+    <Protobuf Include="..\Plugins\Compute\HordeServer.Compute\Protos\horde\messages\aws_lambda.proto" ProtoRoot="..\Plugins\Compute\HordeServer.Compute\Protos" Access="Public" AdditionalImportDirs="..\Plugins\Compute\HordeServer.Compute\Protos" />
+    <Protobuf Include="..\Plugins\Compute\HordeServer.Compute\Protos\horde\tasks\compute_task.proto" ProtoRoot="..\Plugins\Compute\HordeServer.Compute\Protos" Access="Public" AdditionalImportDirs="..\Plugins\Compute\HordeServer.Compute\Protos" />
+    <Protobuf Include="..\Plugins\Compute\HordeServer.Compute\Protos\horde\tasks\restart_task.proto" ProtoRoot="..\Plugins\Compute\HordeServer.Compute\Protos" Access="Public" AdditionalImportDirs="..\Plugins\Compute\HordeServer.Compute\Protos" />
+    <Protobuf Include="..\Plugins\Compute\HordeServer.Compute\Protos\horde\tasks\shutdown_task.proto" ProtoRoot="..\Plugins\Compute\HordeServer.Compute\Protos" Access="Public" AdditionalImportDirs="..\Plugins\Compute\HordeServer.Compute\Protos" />
+    <Protobuf Include="..\Plugins\Compute\HordeServer.Compute\Protos\horde\tasks\test_task.proto" ProtoRoot="..\Plugins\Compute\HordeServer.Compute\Protos" Access="Public" AdditionalImportDirs="..\Plugins\Compute\HordeServer.Compute\Protos" />
+    <Protobuf Include="..\Plugins\Compute\HordeServer.Compute\Protos\horde\tasks\upgrade_task.proto" ProtoRoot="..\Plugins\Compute\HordeServer.Compute\Protos" Access="Public" AdditionalImportDirs="..\Plugins\Compute\HordeServer.Compute\Protos" />
 
     <Protobuf Include="..\Plugins\Build\HordeServer.Build\Protos\horde\job_rpc_messages.proto" ProtoRoot="..\Plugins\Build\HordeServer.Build\Protos" Access="Public" AdditionalImportDirs="..\Plugins\Compute\HordeServer.Compute\Protos" />
     <Protobuf Include="..\Plugins\Build\HordeServer.Build\Protos\horde\common\label_change.proto" ProtoRoot="..\Plugins\Build\HordeServer.Build\Protos" Access="Public" AdditionalImportDirs="..\Plugins\Compute\HordeServer.Compute\Protos" />
diff --git a/Engine/Source/Programs/Horde/HordeCmd/HordeCmd.csproj b/Engine/Source/Programs/Horde/HordeCmd/HordeCmd.csproj
--- a/Engine/Source/Programs/Horde/HordeCmd/HordeCmd.csproj
+++ b/Engine/Source/Programs/Horde/HordeCmd/HordeCmd.csproj
@@ -32,6 +32,7 @@
     <PackageReference Include="Serilog.Settings.Configuration" Version="8.0.0" />
     <PackageReference Include="Serilog.Sinks.Console" Version="5.0.1" />
     <PackageReference Include="Serilog.Sinks.File" Version="5.0.0" />
+    <PackageReference Include="System.Text.Json" Version="8.0.4" />
   </ItemGroup>
 
   <ItemGroup>
diff --git a/Engine/Source/Programs/Horde/HordeServer/HordeServer.csproj b/Engine/Source/Programs/Horde/HordeServer/HordeServer.csproj
--- a/Engine/Source/Programs/Horde/HordeServer/HordeServer.csproj
+++ b/Engine/Source/Programs/Horde/HordeServer/HordeServer.csproj
@@ -79,6 +79,7 @@
     <PackageReference Include="StackExchange.Redis" Version="2.7.4" />
     <PackageReference Include="Swashbuckle.AspNetCore" Version="6.5.0" />
     <PackageReference Include="System.Linq.Async" Version="6.0.1" />
+    <PackageReference Include="System.Text.Json" Version="8.0.4" />
     <PackageReference Include="T-Digest.NET" Version="1.0.3" />
     <PackageReference Include="TimeZoneConverter" Version="6.1.0" />
     <PackageReference Include="ZstdSharp.Port" Version="0.8.1" />
diff --git a/Engine/Source/Programs/Horde/Tools/ReleaseNotes/ReleaseNotes.csproj b/Engine/Source/Programs/Horde/Tools/ReleaseNotes/ReleaseNotes.csproj
--- a/Engine/Source/Programs/Horde/Tools/ReleaseNotes/ReleaseNotes.csproj
+++ b/Engine/Source/Programs/Horde/Tools/ReleaseNotes/ReleaseNotes.csproj
@@ -10,6 +10,7 @@
 
   <ItemGroup>
     <PackageReference Include="Microsoft.Extensions.Logging.Console" Version="8.0.0" />
+    <PackageReference Include="System.Text.Json" Version="8.0.4" />
   </ItemGroup>
 
   <ItemGroup>
diff --git a/Engine/Source/Programs/Shared/EpicGames.AspNet/EpicGames.AspNet.csproj b/Engine/Source/Programs/Shared/EpicGames.AspNet/EpicGames.AspNet.csproj
--- a/Engine/Source/Programs/Shared/EpicGames.AspNet/EpicGames.AspNet.csproj
+++ b/Engine/Source/Programs/Shared/EpicGames.AspNet/EpicGames.AspNet.csproj
@@ -44,10 +44,5 @@
   </ItemGroup>
   <ItemGroup>
     <PackageReference Include="Microsoft.CSharp" Version="4.7.0" />
-    <PackageReference Include="Microsoft.Win32.Registry" Version="5.0.0" />
-    <PackageReference Include="System.Data.DataSetExtensions" Version="4.5.0" />
-
-    
-  
   </ItemGroup>
 </Project>
\ No newline at end of file
diff --git a/Engine/Source/Programs/Shared/EpicGames.Build/EpicGames.Build.csproj b/Engine/Source/Programs/Shared/EpicGames.Build/EpicGames.Build.csproj
--- a/Engine/Source/Programs/Shared/EpicGames.Build/EpicGames.Build.csproj
+++ b/Engine/Source/Programs/Shared/EpicGames.Build/EpicGames.Build.csproj
@@ -46,9 +46,6 @@
   <ItemGroup>
     <PackageReference Include="Microsoft.CSharp" Version="4.7.0" />
     <PackageReference Include="Microsoft.Extensions.FileSystemGlobbing" Version="6.0.0" />
-    <PackageReference Include="Microsoft.Win32.Registry" Version="5.0.0" />
-    <PackageReference Include="System.CodeDom" Version="8.0.0" />
-    <PackageReference Include="System.Data.DataSetExtensions" Version="4.5.0" />
     <PackageReference Include="System.Security.Permissions" Version="4.7.0" />
   </ItemGroup>
 </Project>
diff --git a/Engine/Source/Programs/Shared/EpicGames.BuildGraph.Tests/EpicGames.BuildGraph.Tests.csproj b/Engine/Source/Programs/Shared/EpicGames.BuildGraph.Tests/EpicGames.BuildGraph.Tests.csproj
--- a/Engine/Source/Programs/Shared/EpicGames.BuildGraph.Tests/EpicGames.BuildGraph.Tests.csproj
+++ b/Engine/Source/Programs/Shared/EpicGames.BuildGraph.Tests/EpicGames.BuildGraph.Tests.csproj
@@ -46,7 +46,9 @@
 
   <ItemGroup>
     <PackageReference Include="Microsoft.NET.Test.Sdk" Version="17.11.1" />
-    <PackageReference Include="MSTest.TestAdapter" Version="3.6.0" />
+    <PackageReference Include="MSTest.TestAdapter" Version="3.6.0">
+      <TreatAsUsed>true</TreatAsUsed>
+    </PackageReference>
     <PackageReference Include="MSTest.TestFramework" Version="3.6.0" />
     <PackageReference Include="coverlet.collector" Version="6.0.2">
       <PrivateAssets>all</PrivateAssets>
diff --git a/Engine/Source/Programs/Shared/EpicGames.Core.Tests/EpicGames.Core.Tests.csproj b/Engine/Source/Programs/Shared/EpicGames.Core.Tests/EpicGames.Core.Tests.csproj
--- a/Engine/Source/Programs/Shared/EpicGames.Core.Tests/EpicGames.Core.Tests.csproj
+++ b/Engine/Source/Programs/Shared/EpicGames.Core.Tests/EpicGames.Core.Tests.csproj
@@ -49,7 +49,9 @@
 
   <ItemGroup>
     <PackageReference Include="Microsoft.NET.Test.Sdk" Version="17.11.1" />
-    <PackageReference Include="MSTest.TestAdapter" Version="3.6.0" />
+    <PackageReference Include="MSTest.TestAdapter" Version="3.6.0">
+      <TreatAsUsed>true</TreatAsUsed>
+    </PackageReference>
     <PackageReference Include="MSTest.TestFramework" Version="3.6.0" />
     <PackageReference Include="coverlet.collector" Version="6.0.2">
       <PrivateAssets>all</PrivateAssets>
diff --git a/Engine/Source/Programs/Shared/EpicGames.Core/EpicGames.Core.csproj b/Engine/Source/Programs/Shared/EpicGames.Core/EpicGames.Core.csproj
--- a/Engine/Source/Programs/Shared/EpicGames.Core/EpicGames.Core.csproj
+++ b/Engine/Source/Programs/Shared/EpicGames.Core/EpicGames.Core.csproj
@@ -53,8 +53,6 @@
     <PackageReference Include="JetBrains.Annotations" Version="2022.1.0" />
     <PackageReference Include="Microsoft.Extensions.Configuration" Version="8.0.0" />
     <PackageReference Include="Microsoft.Extensions.Http.Polly" Version="6.0.26" />
-    <PackageReference Include="Microsoft.Extensions.Logging" Version="6.0.0" />
-    <PackageReference Include="Microsoft.Extensions.Logging.Abstractions" Version="6.0.1" />
     <PackageReference Include="Microsoft.Extensions.Logging.Console" Version="6.0.0" />
     <PackageReference Include="Microsoft.Extensions.ObjectPool" Version="6.0.8" />
     <PackageReference Include="OpenTracing" Version="0.12.1" />
diff --git a/Engine/Source/Programs/Shared/EpicGames.Horde.Tests/EpicGames.Horde.Tests.csproj b/Engine/Source/Programs/Shared/EpicGames.Horde.Tests/EpicGames.Horde.Tests.csproj
--- a/Engine/Source/Programs/Shared/EpicGames.Horde.Tests/EpicGames.Horde.Tests.csproj
+++ b/Engine/Source/Programs/Shared/EpicGames.Horde.Tests/EpicGames.Horde.Tests.csproj
@@ -31,8 +31,9 @@
 
   <ItemGroup>
     <PackageReference Include="Microsoft.NET.Test.Sdk" Version="17.11.1" />
-    <PackageReference Include="Microsoft.Extensions.Caching.Memory" Version="6.0.1" />
-    <PackageReference Include="MSTest.TestAdapter" Version="3.6.0" />
+    <PackageReference Include="MSTest.TestAdapter" Version="3.6.0">
+      <TreatAsUsed>true</TreatAsUsed>
+    </PackageReference>
     <PackageReference Include="MSTest.TestFramework" Version="3.6.0" />
     <PackageReference Include="coverlet.collector" Version="6.0.2">
       <PrivateAssets>all</PrivateAssets>
diff --git a/Engine/Source/Programs/Shared/EpicGames.Horde/EpicGames.Horde.csproj b/Engine/Source/Programs/Shared/EpicGames.Horde/EpicGames.Horde.csproj
--- a/Engine/Source/Programs/Shared/EpicGames.Horde/EpicGames.Horde.csproj
+++ b/Engine/Source/Programs/Shared/EpicGames.Horde/EpicGames.Horde.csproj
@@ -51,16 +51,10 @@
     <PackageReference Include="Microsoft.Data.Sqlite" Version="8.0.0" />
     <PackageReference Include="Microsoft.Extensions.Caching.Abstractions" Version="6.0.0" />
     <PackageReference Include="Microsoft.Extensions.Caching.Memory" Version="6.0.1" />
-    <PackageReference Include="Microsoft.Extensions.DependencyInjection" Version="6.0.0" />
-    <PackageReference Include="Microsoft.Extensions.Http" Version="6.0.0" />
-    <PackageReference Include="Microsoft.Extensions.Http.Polly" Version="6.0.26" />
-    <PackageReference Include="Microsoft.Extensions.Logging" Version="6.0.0" />
     <PackageReference Include="Microsoft.VisualStudio.Threading.Analyzers" Version="17.8.14">
       <PrivateAssets>all</PrivateAssets>
       <IncludeAssets>runtime; build; native; contentfiles; analyzers; buildtransitive</IncludeAssets>
     </PackageReference>
-    <PackageReference Include="System.Net.Http" Version="4.3.4" />
-    <PackageReference Include="System.Text.RegularExpressions" Version="4.3.1" />
     <PackageReference Include="ZstdSharp.Port" Version="0.8.1" />
   </ItemGroup>
   <ItemGroup>
diff --git a/Engine/Source/Programs/Shared/EpicGames.IoHash.Tests/EpicGames.IoHash.Tests.csproj b/Engine/Source/Programs/Shared/EpicGames.IoHash.Tests/EpicGames.IoHash.Tests.csproj
--- a/Engine/Source/Programs/Shared/EpicGames.IoHash.Tests/EpicGames.IoHash.Tests.csproj
+++ b/Engine/Source/Programs/Shared/EpicGames.IoHash.Tests/EpicGames.IoHash.Tests.csproj
@@ -18,7 +18,9 @@
 
   <ItemGroup>
     <PackageReference Include="Microsoft.NET.Test.Sdk" Version="17.11.1" />
-    <PackageReference Include="MSTest.TestAdapter" Version="3.6.0" />
+    <PackageReference Include="MSTest.TestAdapter" Version="3.6.0">
+      <TreatAsUsed>true</TreatAsUsed>
+    </PackageReference>
     <PackageReference Include="MSTest.TestFramework" Version="3.6.0" />
     <PackageReference Include="coverlet.collector" Version="6.0.2">
       <PrivateAssets>all</PrivateAssets>
diff --git a/Engine/Source/Programs/Shared/EpicGames.Perforce.Fixture/EpicGames.Perforce.Fixture.csproj b/Engine/Source/Programs/Shared/EpicGames.Perforce.Fixture/EpicGames.Perforce.Fixture.csproj
--- a/Engine/Source/Programs/Shared/EpicGames.Perforce.Fixture/EpicGames.Perforce.Fixture.csproj
+++ b/Engine/Source/Programs/Shared/EpicGames.Perforce.Fixture/EpicGames.Perforce.Fixture.csproj
@@ -34,7 +34,9 @@
   <ItemGroup>
     <PackageReference Include="Microsoft.Extensions.Logging.Console" Version="8.0.0" />
     <PackageReference Include="Microsoft.NET.Test.Sdk" Version="17.11.1" />
-    <PackageReference Include="MSTest.TestAdapter" Version="3.6.0" />
+    <PackageReference Include="MSTest.TestAdapter" Version="3.6.0">
+      <TreatAsUsed>true</TreatAsUsed>
+    </PackageReference>
     <PackageReference Include="MSTest.TestFramework" Version="3.6.0" />
     <PackageReference Include="coverlet.collector" Version="6.0.2">
       <PrivateAssets>all</PrivateAssets>
diff --git a/Engine/Source/Programs/Shared/EpicGames.Perforce.Managed.Tests/EpicGames.Perforce.Managed.Tests.csproj b/Engine/Source/Programs/Shared/EpicGames.Perforce.Managed.Tests/EpicGames.Perforce.Managed.Tests.csproj
--- a/Engine/Source/Programs/Shared/EpicGames.Perforce.Managed.Tests/EpicGames.Perforce.Managed.Tests.csproj
+++ b/Engine/Source/Programs/Shared/EpicGames.Perforce.Managed.Tests/EpicGames.Perforce.Managed.Tests.csproj
@@ -31,9 +31,15 @@
   </ItemGroup>
 
   <ItemGroup>
-    <PackageReference Include="Microsoft.NET.Test.Sdk" Version="17.11.1" />
-    <PackageReference Include="MSTest.TestAdapter" Version="3.6.0" />
-    <PackageReference Include="MSTest.TestFramework" Version="3.6.0" />
+    <PackageReference Include="Microsoft.NET.Test.Sdk" Version="17.11.1">
+      <TreatAsUsed>true</TreatAsUsed>
+    </PackageReference>
+    <PackageReference Include="MSTest.TestAdapter" Version="3.6.0">
+      <TreatAsUsed>true</TreatAsUsed>
+    </PackageReference>
+    <PackageReference Include="MSTest.TestFramework" Version="3.6.0">
+      <TreatAsUsed>true</TreatAsUsed>
+    </PackageReference>
     <PackageReference Include="coverlet.collector" Version="6.0.2">
       <PrivateAssets>all</PrivateAssets>
       <IncludeAssets>runtime; build; native; contentfiles; analyzers; buildtransitive</IncludeAssets>
diff --git a/Engine/Source/Programs/Shared/EpicGames.Perforce.Tests/EpicGames.Perforce.Tests.csproj b/Engine/Source/Programs/Shared/EpicGames.Perforce.Tests/EpicGames.Perforce.Tests.csproj
--- a/Engine/Source/Programs/Shared/EpicGames.Perforce.Tests/EpicGames.Perforce.Tests.csproj
+++ b/Engine/Source/Programs/Shared/EpicGames.Perforce.Tests/EpicGames.Perforce.Tests.csproj
@@ -31,7 +31,9 @@
 
   <ItemGroup>
     <PackageReference Include="Microsoft.NET.Test.Sdk" Version="17.11.1" />
-    <PackageReference Include="MSTest.TestAdapter" Version="3.6.0" />
+    <PackageReference Include="MSTest.TestAdapter" Version="3.6.0">
+      <TreatAsUsed>true</TreatAsUsed>
+    </PackageReference>
     <PackageReference Include="MSTest.TestFramework" Version="3.6.0" />
     <PackageReference Include="coverlet.collector" Version="6.0.2">
       <PrivateAssets>all</PrivateAssets>
diff --git a/Engine/Source/Programs/Shared/EpicGames.Perforce/EpicGames.Perforce.csproj b/Engine/Source/Programs/Shared/EpicGames.Perforce/EpicGames.Perforce.csproj
--- a/Engine/Source/Programs/Shared/EpicGames.Perforce/EpicGames.Perforce.csproj
+++ b/Engine/Source/Programs/Shared/EpicGames.Perforce/EpicGames.Perforce.csproj
@@ -77,7 +77,6 @@
   </ItemGroup>
 
   <ItemGroup>
-    <PackageReference Include="Microsoft.Win32.Registry" Version="5.0.0" />
     <PackageReference Include="System.Linq.Async" Version="6.0.1" />
   </ItemGroup>
 
diff --git a/Engine/Source/Programs/Shared/EpicGames.Redis.Tests/EpicGames.Redis.Tests.csproj b/Engine/Source/Programs/Shared/EpicGames.Redis.Tests/EpicGames.Redis.Tests.csproj
--- a/Engine/Source/Programs/Shared/EpicGames.Redis.Tests/EpicGames.Redis.Tests.csproj
+++ b/Engine/Source/Programs/Shared/EpicGames.Redis.Tests/EpicGames.Redis.Tests.csproj
@@ -32,7 +32,9 @@
 
   <ItemGroup>
     <PackageReference Include="Microsoft.NET.Test.Sdk" Version="17.11.1" />
-    <PackageReference Include="MSTest.TestAdapter" Version="3.6.0" />
+    <PackageReference Include="MSTest.TestAdapter" Version="3.6.0">
+      <TreatAsUsed>true</TreatAsUsed>
+    </PackageReference>
     <PackageReference Include="MSTest.TestFramework" Version="3.6.0" />
     <PackageReference Include="coverlet.collector" Version="6.0.2">
       <PrivateAssets>all</PrivateAssets>
diff --git a/Engine/Source/Programs/Shared/EpicGames.Redis/EpicGames.Redis.csproj b/Engine/Source/Programs/Shared/EpicGames.Redis/EpicGames.Redis.csproj
--- a/Engine/Source/Programs/Shared/EpicGames.Redis/EpicGames.Redis.csproj
+++ b/Engine/Source/Programs/Shared/EpicGames.Redis/EpicGames.Redis.csproj
@@ -35,7 +35,6 @@
     </Compile>
   </ItemGroup>
   <ItemGroup>
-    <PackageReference Include="Microsoft.Extensions.Logging" Version="8.0.0" />
     <PackageReference Include="Google.Protobuf" Version="3.25.1" />
     <PackageReference Include="protobuf-net" Version="3.2.30" />
     <PackageReference Include="StackExchange.Redis" Version="2.7.4" />
diff --git a/Engine/Source/Programs/Shared/EpicGames.Serialization.Tests/EpicGames.Serialization.Tests.csproj b/Engine/Source/Programs/Shared/EpicGames.Serialization.Tests/EpicGames.Serialization.Tests.csproj
--- a/Engine/Source/Programs/Shared/EpicGames.Serialization.Tests/EpicGames.Serialization.Tests.csproj
+++ b/Engine/Source/Programs/Shared/EpicGames.Serialization.Tests/EpicGames.Serialization.Tests.csproj
@@ -34,7 +34,9 @@
 
   <ItemGroup>
     <PackageReference Include="Microsoft.NET.Test.Sdk" Version="17.11.1" />
-    <PackageReference Include="MSTest.TestAdapter" Version="3.6.0" />
+    <PackageReference Include="MSTest.TestAdapter" Version="3.6.0">
+      <TreatAsUsed>true</TreatAsUsed>
+    </PackageReference>
     <PackageReference Include="MSTest.TestFramework" Version="3.6.0" />
     <PackageReference Include="coverlet.collector" Version="6.0.2">
       <PrivateAssets>all</PrivateAssets>
diff --git a/Engine/Source/Programs/Shared/EpicGames.Tracing.Tests/EpicGames.Tracing.Tests.csproj b/Engine/Source/Programs/Shared/EpicGames.Tracing.Tests/EpicGames.Tracing.Tests.csproj
--- a/Engine/Source/Programs/Shared/EpicGames.Tracing.Tests/EpicGames.Tracing.Tests.csproj
+++ b/Engine/Source/Programs/Shared/EpicGames.Tracing.Tests/EpicGames.Tracing.Tests.csproj
@@ -30,7 +30,9 @@
 
   <ItemGroup>
     <PackageReference Include="Microsoft.NET.Test.Sdk" Version="17.11.1" />
-    <PackageReference Include="MSTest.TestAdapter" Version="3.6.0" />
+    <PackageReference Include="MSTest.TestAdapter" Version="3.6.0">
+      <TreatAsUsed>true</TreatAsUsed>
+    </PackageReference>
     <PackageReference Include="MSTest.TestFramework" Version="3.6.0" />
     <PackageReference Include="coverlet.collector" Version="6.0.2">
       <PrivateAssets>all</PrivateAssets>
diff --git a/Engine/Source/Programs/Shared/EpicGames.UHT/EpicGames.UHT.csproj b/Engine/Source/Programs/Shared/EpicGames.UHT/EpicGames.UHT.csproj
--- a/Engine/Source/Programs/Shared/EpicGames.UHT/EpicGames.UHT.csproj
+++ b/Engine/Source/Programs/Shared/EpicGames.UHT/EpicGames.UHT.csproj
@@ -47,6 +47,5 @@
   </ItemGroup>
   <ItemGroup>
     <PackageReference Include="Microsoft.CSharp" Version="4.7.0" />
-    <PackageReference Include="System.Collections.Specialized" Version="4.3.0" />
   </ItemGroup>
 </Project>
\ No newline at end of file
diff --git a/Engine/Source/Programs/UnrealBuildTool.Tests/UnrealBuildTool.Tests.csproj b/Engine/Source/Programs/UnrealBuildTool.Tests/UnrealBuildTool.Tests.csproj
--- a/Engine/Source/Programs/UnrealBuildTool.Tests/UnrealBuildTool.Tests.csproj
+++ b/Engine/Source/Programs/UnrealBuildTool.Tests/UnrealBuildTool.Tests.csproj
@@ -17,7 +17,9 @@
 
   <ItemGroup>
     <PackageReference Include="Microsoft.NET.Test.Sdk" Version="17.11.1" />
-    <PackageReference Include="MSTest.TestAdapter" Version="3.6.0" />
+    <PackageReference Include="MSTest.TestAdapter" Version="3.6.0">
+      <TreatAsUsed>true</TreatAsUsed>
+    </PackageReference>
     <PackageReference Include="MSTest.TestFramework" Version="3.6.0" />
     <PackageReference Include="coverlet.collector" Version="6.0.2">
       <PrivateAssets>all</PrivateAssets>
