Horde: Fix issue with conform threshold for workspaces not propagating

diff --git a/Engine/Source/Programs/Horde/Plugins/Build/HordeServer.Build.Tests/Agents/Pools/PoolUpdateServiceTest.cs b/Engine/Source/Programs/Horde/Plugins/Build/HordeServer.Build.Tests/Agents/Pools/PoolUpdateServiceTest.cs
--- a/Engine/Source/Programs/Horde/Plugins/Build/HordeServer.Build.Tests/Agents/Pools/PoolUpdateServiceTest.cs
+++ b/Engine/Source/Programs/Horde/Plugins/Build/HordeServer.Build.Tests/Agents/Pools/PoolUpdateServiceTest.cs
@@ -3,6 +3,8 @@
 using EpicGames.Horde.Agents;
 using HordeServer.Agents;
 using HordeServer.Agents.Pools;
+using HordeServer.Projects;
+using HordeServer.Streams;
 using Microsoft.Extensions.DependencyInjection;
 using Microsoft.Extensions.Logging.Abstractions;
 using Microsoft.Extensions.Options;
@@ -21,7 +23,7 @@ public class PoolUpdateServiceTest : BuildTestSetup
 	public PoolUpdateServiceTest()
 	{
 		UpdateConfig(x => x.Plugins.GetComputeConfig().Pools.Clear());
-		_pus = new PoolUpdateService(AgentService, PoolCollection, Clock, ServiceProvider.GetRequiredService<IOptionsMonitor<BuildConfig>>(), Tracer, new NullLogger<PoolUpdateService>());
+		_pus = new PoolUpdateService(AgentService, PoolService, PoolCollection, Clock, ServiceProvider.GetRequiredService<IOptionsMonitor<BuildConfig>>(), Tracer, new NullLogger<PoolUpdateService>());
 	}
 
 	[TestInitialize]
@@ -119,10 +121,10 @@ public class PoolUpdateServiceTest : BuildTestSetup
 	[DataRow(true, 200, 50)]
 	[DataRow(true, 200, null)]
 	[DataRow(true, null, 0, 50, 300)]
-	public async Task AutoConformAgentsAsync(bool conformRequested, params int?[] autoConformThresholdsM)
+	public async Task AutoConformAgentsAsync(bool conformRequested, params int?[] autoConformThresholdsMb)
 	{
 		// Arrange
-		IAgent agent = await CreateAutoConformAgentAsync(100, autoConformThresholdsM);
+		IAgent agent = await CreateAutoConformAgentAsync(100, autoConformThresholdsMb);
 
 		// Act
 		await _pus.AutoConformAgentsAsync(CancellationToken.None);
@@ -134,8 +136,23 @@ public class PoolUpdateServiceTest : BuildTestSetup
 
 	private async Task<IAgent> CreateAutoConformAgentAsync(int freeDiskSpaceMb, params int?[] autoConformThresholdsMb)
 	{
-		IEnumerable<AgentWorkspaceInfo> workspaces = autoConformThresholdsMb.Select(x => new AgentWorkspaceInfo(null, null, "someWorkspace", "//Some/Stream", null, false, null, null, x));
-		IAgent? agent = await CreateAgentAsync(_pool, properties: [$"{KnownPropertyNames.DiskFreeSpace}={freeDiskSpaceMb * 1024 * 1024}"], workspaces: workspaces.ToList());
+		Dictionary<string, AgentConfig> agentTypes = new();
+		Dictionary<string, WorkspaceConfig> workspaceTypes = new();
+		for (int i = 0; i < autoConformThresholdsMb.Length; i++)
+		{
+			string workspaceId = "myWorkspace" + i;
+			string agentTypeName = "myAgentType" + i;
+			agentTypes[agentTypeName] = new AgentConfig { Pool = _pool.Id, Workspace = workspaceId };
+			workspaceTypes[workspaceId] = new WorkspaceConfig { Identifier = workspaceId, ConformDiskFreeSpace = autoConformThresholdsMb[i] };
+		}
+		
+		UpdateConfig(x =>
+		{
+			x.Plugins.GetBuildConfig().Projects.Clear();
+			x.Plugins.GetBuildConfig().Projects.Add(new ProjectConfig { Streams = [new StreamConfig { WorkspaceTypes = workspaceTypes, AgentTypes = agentTypes }]});
+		});
+		
+		IAgent? agent = await CreateAgentAsync(_pool, properties: [$"{KnownPropertyNames.DiskFreeSpace}={freeDiskSpaceMb * 1024 * 1024}"]);
 		agent = await agent.TryTerminateSessionAsync(); // Make agent status = stopped
 		Assert.IsNotNull(agent);
 		return agent;
diff --git a/Engine/Source/Programs/Horde/Plugins/Build/HordeServer.Build/Agents/Pools/PoolUpdateService.cs b/Engine/Source/Programs/Horde/Plugins/Build/HordeServer.Build/Agents/Pools/PoolUpdateService.cs
--- a/Engine/Source/Programs/Horde/Plugins/Build/HordeServer.Build/Agents/Pools/PoolUpdateService.cs
+++ b/Engine/Source/Programs/Horde/Plugins/Build/HordeServer.Build/Agents/Pools/PoolUpdateService.cs
@@ -18,6 +18,7 @@ namespace HordeServer.Agents.Pools
 	public sealed class PoolUpdateService : IHostedService, IAsyncDisposable
 	{
 		readonly AgentService _agentService;
+		readonly PoolService _poolService;
 		readonly IPoolCollection _pools;
 		readonly IClock _clock;
 		readonly IOptionsMonitor<BuildConfig> _buildConfig;
@@ -30,9 +31,10 @@ namespace HordeServer.Agents.Pools
 		/// <summary>
 		/// Constructor
 		/// </summary>
-		public PoolUpdateService(AgentService agentService, IPoolCollection pools, IClock clock, IOptionsMonitor<BuildConfig> buildConfig, Tracer tracer, ILogger<PoolUpdateService> logger)
+		public PoolUpdateService(AgentService agentService, PoolService poolService, IPoolCollection pools, IClock clock, IOptionsMonitor<BuildConfig> buildConfig, Tracer tracer, ILogger<PoolUpdateService> logger)
 		{
 			_agentService = agentService;
+			_poolService = poolService;
 			_pools = pools;
 			_clock = clock;
 			_buildConfig = buildConfig;
@@ -138,22 +140,23 @@ namespace HordeServer.Agents.Pools
 			foreach (IAgent agent in agents)
 			{
 				long? freeDiskSpace = agent.GetDiskFreeSpace();
-				long maxConformDiskSpace = 0;
-
+				long conformDiskSpaceNeeded = 0;
+				HashSet<AgentWorkspaceInfo> workspaces = await _poolService.GetWorkspacesAsync(agent, DateTime.UtcNow - TimeSpan.FromHours(1), _buildConfig.CurrentValue, cancellationToken);
+				
 				// Find the largest conform disk space amount needed, if any
-				foreach (AgentWorkspaceInfo workspace in agent.Workspaces)
+				foreach (AgentWorkspaceInfo workspace in workspaces)
 				{
 					if (workspace.ConformDiskFreeSpace is > 0)
 					{
-						maxConformDiskSpace = Math.Max(maxConformDiskSpace, MegabytesToBytes(workspace.ConformDiskFreeSpace.Value));
+						conformDiskSpaceNeeded = Math.Max(conformDiskSpaceNeeded, MegabytesToBytes(workspace.ConformDiskFreeSpace.Value));
 					}
 				}
 
-				if (freeDiskSpace != null && maxConformDiskSpace > 0 && freeDiskSpace < maxConformDiskSpace && agent.ConformAttemptCount is null or 0)
+				if (freeDiskSpace != null && conformDiskSpaceNeeded > 0 && freeDiskSpace < conformDiskSpaceNeeded && agent.ConformAttemptCount is null or 0)
 				{
 					await agent.TryUpdateAsync(new UpdateAgentOptions { RequestFullConform = true }, cancellationToken: cancellationToken);
 					_logger.LogInformation("Auto-conforming {AgentId} as workspace conform disk space needed ({ConformDiskSpace:F1} MB) is less than free disk space ({FreeDiskSpace:F1} MB)",
-						agent.Id.ToString(), maxConformDiskSpace / 1024.0 / 1024.0, freeDiskSpace.Value / 1024.0 / 1024.0);
+						agent.Id.ToString(), conformDiskSpaceNeeded / 1024.0 / 1024.0, freeDiskSpace.Value / 1024.0 / 1024.0);
 				}
 			}
 		}
diff --git a/Engine/Source/Programs/Horde/Plugins/Build/HordeServer.Build/Perforce/PerforceWorkspace.cs b/Engine/Source/Programs/Horde/Plugins/Build/HordeServer.Build/Perforce/PerforceWorkspace.cs
--- a/Engine/Source/Programs/Horde/Plugins/Build/HordeServer.Build/Perforce/PerforceWorkspace.cs
+++ b/Engine/Source/Programs/Horde/Plugins/Build/HordeServer.Build/Perforce/PerforceWorkspace.cs
@@ -184,7 +184,17 @@ namespace HordeServer.Perforce
 
 				// Create the new workspace
 				string cluster = workspaceConfig.Cluster ?? streamConfig.ClusterName;
-				workspace = new AgentWorkspaceInfo(cluster, workspaceConfig.UserName, identifier, workspaceConfig.Stream ?? streamConfig.Name, workspaceConfig.View, workspaceConfig.Incremental ?? false, workspaceConfig.Method, workspaceConfig.MinScratchSpace);
+				workspace = new AgentWorkspaceInfo(
+					cluster,
+					workspaceConfig.UserName,
+					identifier,
+					workspaceConfig.Stream ?? streamConfig.Name,
+					workspaceConfig.View,
+					workspaceConfig.Incremental ?? false,
+					workspaceConfig.Method,
+					workspaceConfig.MinScratchSpace,
+					workspaceConfig.ConformDiskFreeSpace);
+				
 				autoSdkConfig = GetAutoSdkConfig(workspaceConfig, streamConfig);
 
 				return true;
