Horde: Revert change to return commits in ascending order by default, and change the P4 SubscribeAsync method to handle ascending commit numbers explicitly instead.

#rnx

diff --git a/Engine/Source/Programs/Horde/Plugins/Build/HordeServer.Build/Perforce/PerforceService.cs b/Engine/Source/Programs/Horde/Plugins/Build/HordeServer.Build/Perforce/PerforceService.cs
--- a/Engine/Source/Programs/Horde/Plugins/Build/HordeServer.Build/Perforce/PerforceService.cs
+++ b/Engine/Source/Programs/Horde/Plugins/Build/HordeServer.Build/Perforce/PerforceService.cs
@@ -1237,8 +1237,6 @@ namespace HordeServer.Perforce
 						string filter = GetFilter($"//{perforce.Settings.ClientName}/...", minChange, maxChange);
 
 						List<ChangesRecord> changes = await perforce.GetChangesAsync(ChangesOptions.IncludeTimes | ChangesOptions.LongOutput, maxResults ?? -1, ChangeStatus.Submitted, filter, cancellationToken);
-						changes.SortBy(x => x.Number);
-
 						foreach (ChangesRecord change in changes)
 						{
 							ICommit commit = await PerforceService.CreateCommitAsync(StreamConfig, change, info, cancellationToken);
@@ -1264,8 +1262,6 @@ namespace HordeServer.Perforce
 							const int MaxFiles = 1000;
 
 							List<DescribeRecord> describeRecords = await perforce.DescribeAsync(DescribeOptions.None, MaxFiles, changesRecords.Select(x => x.Number).ToArray(), cancellationToken);
-							describeRecords.SortBy(x => x.Number);
-
 							foreach (DescribeRecord describeRecord in describeRecords)
 							{
 								bool hasAllFiles = describeRecord.Files.Count < MaxFiles;
@@ -1309,12 +1305,29 @@ namespace HordeServer.Perforce
 			{
 				for (; ; )
 				{
-					await foreach (ICommit commit in FindAsync(minChange: CommitId.FromPerforceChange(minChange), includeMinChange: false, maxResults: 10, tags: tags, cancellationToken: cancellationToken))
+					const int MaxChanges = 10;
+
+					InfoRecord info;
+					List<ChangesRecord> changes;
+
+					using (PooledConnectionHandle perforce = await PerforceService.ConnectWithStreamClientAsync(StreamConfig, null, cancellationToken))
+					{
+						info = await perforce.GetInfoAsync(cancellationToken);
+						string filter = GetFilter($"//{perforce.Settings.ClientName}/...", minChange + 1, null);
+						changes = await perforce.GetChangesAsync(ChangesOptions.Reverse | ChangesOptions.IncludeTimes | ChangesOptions.LongOutput, MaxChanges, ChangeStatus.Submitted, filter, cancellationToken);
+					}
+
+					foreach (ChangesRecord change in changes)
 					{
+						ICommit commit = await PerforceService.CreateCommitAsync(StreamConfig, change, info, cancellationToken);
 						yield return commit;
-						minChange = commit.Id.GetPerforceChange();
+						minChange = change.Number;
+					}
+
+					if (changes.Count < MaxChanges)
+					{
+						await Task.Delay(TimeSpan.FromSeconds(10.0), cancellationToken);
 					}
-					await Task.Delay(TimeSpan.FromSeconds(10.0), cancellationToken);
 				}
 			}
 		}
