Horde: Fix path for constructing Azure storage backends.

diff --git a/Engine/Source/Programs/Horde/Plugins/Storage/HordeServer.Storage/Storage/ObjectStoreFactory.cs b/Engine/Source/Programs/Horde/Plugins/Storage/HordeServer.Storage/Storage/ObjectStoreFactory.cs
--- a/Engine/Source/Programs/Horde/Plugins/Storage/HordeServer.Storage/Storage/ObjectStoreFactory.cs
+++ b/Engine/Source/Programs/Horde/Plugins/Storage/HordeServer.Storage/Storage/ObjectStoreFactory.cs
@@ -68,6 +68,8 @@ namespace HordeServer.Storage
 					return _serviceProvider.GetRequiredService<FileObjectStoreFactory>().CreateStore(DirectoryReference.Combine(_serverInfo.DataDir, config.BaseDir ?? "Storage"));
 				case StorageBackendType.Aws:
 					return _serviceProvider.GetRequiredService<AwsObjectStoreFactory>().CreateStore(config);
+				case StorageBackendType.Azure:
+					return _serviceProvider.GetRequiredService<AzureObjectStoreFactory>().CreateStore(config);
 				case StorageBackendType.Memory:
 					return new MemoryObjectStore();
 				default:
diff --git a/Engine/Source/Programs/Horde/Plugins/Storage/HordeServer.Storage/Storage/ObjectStores/AzureObjectStore.cs b/Engine/Source/Programs/Horde/Plugins/Storage/HordeServer.Storage/Storage/ObjectStores/AzureObjectStore.cs
--- a/Engine/Source/Programs/Horde/Plugins/Storage/HordeServer.Storage/Storage/ObjectStores/AzureObjectStore.cs
+++ b/Engine/Source/Programs/Horde/Plugins/Storage/HordeServer.Storage/Storage/ObjectStores/AzureObjectStore.cs
@@ -190,4 +190,31 @@ namespace HordeServer.Storage.ObjectStores
 		{
 		}
 	}
+
+	/// <summary>
+	/// Factory for constructing <see cref="AzureObjectStore"/> instances
+	/// </summary>
+	public sealed class AzureObjectStoreFactory
+	{
+		readonly Tracer _tracer;
+		readonly ILogger<AzureObjectStore> _logger;
+
+		/// <summary>
+		/// Constructor
+		/// </summary>
+		public AzureObjectStoreFactory(Tracer tracer, ILogger<AzureObjectStore> logger)
+		{
+			_tracer = tracer;
+			_logger = logger;
+		}
+
+		/// <summary>
+		/// Create a new object store with the given configuration
+		/// </summary>
+		/// <param name="options">Configuration for the store</param>
+		public AzureObjectStore CreateStore(IAzureStorageOptions options)
+		{
+			return new AzureObjectStore(options, _logger, _tracer);
+		}
+	}
 }
diff --git a/Engine/Source/Programs/Horde/Plugins/Storage/HordeServer.Storage/StorageConfig.cs b/Engine/Source/Programs/Horde/Plugins/Storage/HordeServer.Storage/StorageConfig.cs
--- a/Engine/Source/Programs/Horde/Plugins/Storage/HordeServer.Storage/StorageConfig.cs
+++ b/Engine/Source/Programs/Horde/Plugins/Storage/HordeServer.Storage/StorageConfig.cs
@@ -224,7 +224,7 @@ namespace HordeServer.Storage
 	/// Common settings object for different providers
 	/// </summary>
 	[DebuggerDisplay("{Id}")]
-	public class BackendConfig : IAwsStorageOptions
+	public class BackendConfig : IAwsStorageOptions, IAzureStorageOptions
 	{
 		/// <summary>
 		/// The storage backend ID
diff --git a/Engine/Source/Programs/Horde/Plugins/Storage/HordeServer.Storage/StoragePlugin.cs b/Engine/Source/Programs/Horde/Plugins/Storage/HordeServer.Storage/StoragePlugin.cs
--- a/Engine/Source/Programs/Horde/Plugins/Storage/HordeServer.Storage/StoragePlugin.cs
+++ b/Engine/Source/Programs/Horde/Plugins/Storage/HordeServer.Storage/StoragePlugin.cs
@@ -44,6 +44,7 @@ namespace HordeServer.Storage
 
 			services.AddSingleton<IObjectStoreFactory, ObjectStoreFactory>();
 			services.AddSingleton<AwsObjectStoreFactory>();
+			services.AddSingleton<AzureObjectStoreFactory>();
 			services.AddSingleton<FileObjectStoreFactory>();
 
 			services.AddSingleton<BundleCache>();
