Horde: Updating schema documentation

#rnx

diff --git a/Engine/Source/Programs/Horde/Docs/Config/Schema/Globals.md b/Engine/Source/Programs/Horde/Docs/Config/Schema/Globals.md
--- a/Engine/Source/Programs/Horde/Docs/Config/Schema/Globals.md
+++ b/Engine/Source/Programs/Horde/Docs/Config/Schema/Globals.md
@@ -474,8 +474,8 @@ Name | Description
 `awsRole` | `string`<br>ARN of a role to assume
 `awsProfile` | `string`<br>The AWS profile to read credentials form
 `awsRegion` | `string`<br>Region to connect to
-`azureConnectionString` | `string`<br>
-`azureContainerName` | `string`<br>
+`azureConnectionString` | `string`<br>Connection string for Azure
+`azureContainerName` | `string`<br>Name of the container
 `relayServer` | `string`<br>
 `relayToken` | `string`<br>
 
