Horde: Re-add logical cores as an agent property

UBA uses this property for determining number of agents to allocate.

diff --git a/Engine/Source/Programs/Horde/HordeAgent/Services/CapabilitiesService.cs b/Engine/Source/Programs/Horde/HordeAgent/Services/CapabilitiesService.cs
--- a/Engine/Source/Programs/Horde/HordeAgent/Services/CapabilitiesService.cs
+++ b/Engine/Source/Programs/Horde/HordeAgent/Services/CapabilitiesService.cs
@@ -539,6 +539,7 @@ namespace HordeAgent.Services
 			if (numLogicalCores > 0)
 			{
 				capabilities.Resources.Add("LogicalCores", numLogicalCores);
+				capabilities.Properties.Add($"LogicalCores={numLogicalCores}");
 			}
 
 			if (numPhysicalCores > 0)
