Horde: Add filter for agent properties in REST API

diff --git a/Engine/Source/Programs/Horde/HordeServer.Tests.Shared/ServerTestSetup.cs b/Engine/Source/Programs/Horde/HordeServer.Tests.Shared/ServerTestSetup.cs
--- a/Engine/Source/Programs/Horde/HordeServer.Tests.Shared/ServerTestSetup.cs
+++ b/Engine/Source/Programs/Horde/HordeServer.Tests.Shared/ServerTestSetup.cs
@@ -2,6 +2,7 @@
 
 using System.Diagnostics.Metrics;
 using EpicGames.Core;
+using EpicGames.Horde.Tools;
 using HordeCommon;
 using HordeServer.Accounts;
 using HordeServer.Acls;
@@ -10,6 +11,7 @@ using HordeServer.Configuration;
 using HordeServer.Dashboard;
 using HordeServer.Plugins;
 using HordeServer.Server;
+using HordeServer.Tools;
 using HordeServer.Users;
 using Microsoft.Extensions.Caching.Memory;
 using Microsoft.Extensions.Configuration;
@@ -116,6 +118,7 @@ namespace HordeServer.Tests
 			services.AddSingleton<MemoryMappedFileCache>();
 
 			services.AddSingleton<IAccountCollection, AccountCollection>();
+			services.AddSingleton<IToolCollection, ToolCollection>();
 
 			services.AddLogging(builder => builder.AddSerilog());
 			services.AddSingleton<IMemoryCache>(sp => new MemoryCache(new MemoryCacheOptions { }));
@@ -124,7 +127,8 @@ namespace HordeServer.Tests
 			services.AddSingleton(sp => new Meter("TestMeter"));
 
 			services.AddSingleton(typeof(IAuditLogFactory<>), typeof(AuditLogFactory<>));
-
+			
+			services.AddSingleton<GlobalsService>();
 			services.AddSingleton<ConfigService>();
 			services.AddSingleton<IConfigService>(sp => sp.GetRequiredService<ConfigService>());
 			services.AddSingleton<IOptionsFactory<GlobalConfig>>(sp => sp.GetRequiredService<ConfigService>());
diff --git a/Engine/Source/Programs/Horde/Plugins/Compute/HordeServer.Compute.Tests/Agents/AgentsControllerTests.cs b/Engine/Source/Programs/Horde/Plugins/Compute/HordeServer.Compute.Tests/Agents/AgentsControllerTests.cs
new file mode 100644
--- /dev/null
+++ b/Engine/Source/Programs/Horde/Plugins/Compute/HordeServer.Compute.Tests/Agents/AgentsControllerTests.cs
@@ -0,0 +1,115 @@
+﻿// Copyright Epic Games, Inc. All Rights Reserved.
+
+using System.Security.Claims;
+using EpicGames.Horde.Agents;
+using HordeCommon.Rpc.Messages;
+using HordeServer.Agents;
+using HordeServer.Utilities;
+using Microsoft.AspNetCore.Http;
+using Microsoft.AspNetCore.Mvc;
+using Microsoft.Extensions.Logging.Abstractions;
+using Microsoft.Extensions.Options;
+
+namespace HordeServer.Tests.Agents;
+
+[TestClass]
+public class AgentsControllerTests : ComputeTestSetup
+{
+	private class OptionsSnapshotStub<T>(T value) : IOptionsSnapshot<T> where T : class
+	{
+		public T Value { get; } = value;
+		public T Get(string? name)
+		{
+			return Value;
+		}
+	}
+	
+	[TestMethod]
+	public async Task FindAgents_Simple_Async()
+	{
+		// Arrange
+		List<string> properties = ["foo=1", "bar=2"];
+		AgentsController controller = CreateAgentsController();
+		await CreateAgentAsync("myAgent1", properties);
+		
+		// Act
+		ActionResult<List<object>> result = await controller.FindAgentsAsync();
+		
+		// Assert
+		Assert.IsNotNull(result.Value);
+		List<GetAgentResponse> getAgentResponses = result.Value!.Cast<GetAgentResponse>().ToList();
+		Assert.AreEqual(1, getAgentResponses.Count);
+		CollectionAssert.AreEquivalent(properties, getAgentResponses[0].Properties);
+	}
+	
+	[TestMethod]
+	public async Task FindAgents_Filter_Async()
+	{
+		// Arrange
+		List<string> properties = ["foo=1", "bar=2"];
+		AgentsController controller = CreateAgentsController();
+		await CreateAgentAsync("myAgent1", properties);
+		
+		// Act
+		ActionResult<List<object>> result = await controller.FindAgentsAsync(filter: PropertyFilter.Parse("id,name,enabled"));
+		
+		// Assert
+		Assert.IsNotNull(result.Value);
+		List<Dictionary<string, object>> getAgentResponses = result.Value!.Cast<Dictionary<string, object>>().ToList();
+		Assert.AreEqual(1, getAgentResponses.Count);
+		Assert.AreEqual(3, getAgentResponses[0].Count);
+		Assert.AreEqual(new AgentId("MYAGENT1"), getAgentResponses[0]["id"]);
+		Assert.AreEqual("MYAGENT1", getAgentResponses[0]["name"]);
+		Assert.AreEqual(true, getAgentResponses[0]["enabled"]);
+	}
+	
+	[TestMethod]
+	public async Task FindAgents_FilterProps_Async()
+	{
+		// Arrange
+		List<string> properties = ["foo=1", "bar=2", "baz=3"];
+		AgentsController controller = CreateAgentsController();
+		await CreateAgentAsync("myAgent1", properties);
+		
+		// Act
+		ActionResult<List<object>> result = await controller.FindAgentsAsync(filter: PropertyFilter.Parse("id,name,enabled,properties"), propsFilter: "foo,baz");
+		
+		// Assert
+		Assert.IsNotNull(result.Value);
+		List<Dictionary<string, object>> getAgentResponses = result.Value!.Cast<Dictionary<string, object>>().ToList();
+		Assert.AreEqual(1, getAgentResponses.Count);
+		Assert.AreEqual(4, getAgentResponses[0].Count);
+		Assert.AreEqual(new AgentId("MYAGENT1"), getAgentResponses[0]["id"]);
+		Assert.AreEqual("MYAGENT1", getAgentResponses[0]["name"]);
+		Assert.AreEqual(true, getAgentResponses[0]["enabled"]);
+		object o = getAgentResponses[0]["properties"];
+		
+		List<string> actualProps = [..(IEnumerable<string>)o];
+		CollectionAssert.AreEquivalent(new List<string>() { "foo=1", "baz=3"  }, actualProps);
+	}
+
+	private AgentsController CreateAgentsController()
+	{
+		ComputeConfig computeConfig = GlobalConfig.CurrentValue.Plugins.GetComputeConfig();
+		IOptionsSnapshot<ComputeConfig> configOptions = new OptionsSnapshotStub<ComputeConfig>(computeConfig);
+		AgentsController controller = new (AgentService, null!, null!, null!, configOptions, Tracer, NullLogger<AgentsController>.Instance)
+		{
+			ControllerContext = new ControllerContext()
+			{
+				HttpContext = new DefaultHttpContext()
+				{
+					User = new ClaimsPrincipal(new ClaimsIdentity([HordeClaims.AdminClaim.ToClaim()]))
+				}
+			}
+		};
+		return controller;
+	}
+	
+	private async Task<IAgent> CreateAgentAsync(string name, List<string> properties)
+	{
+		IAgent agent = await AgentService.CreateAgentAsync(new CreateAgentOptions(new AgentId(name), false, ""));
+		agent = await AgentService.CreateSessionAsync(agent, new RpcAgentCapabilities(properties), null);
+		return agent;
+	}
+}
+
diff --git a/Engine/Source/Programs/Horde/Plugins/Compute/HordeServer.Compute/Agents/AgentsController.cs b/Engine/Source/Programs/Horde/Plugins/Compute/HordeServer.Compute/Agents/AgentsController.cs
--- a/Engine/Source/Programs/Horde/Plugins/Compute/HordeServer.Compute/Agents/AgentsController.cs
+++ b/Engine/Source/Programs/Horde/Plugins/Compute/HordeServer.Compute/Agents/AgentsController.cs
@@ -61,12 +61,22 @@ namespace HordeServer.Agents
 		/// <param name="index">First result to return</param>
 		/// <param name="count">Number of results to return</param>
 		/// <param name="modifiedAfter">If set, only returns agents modified after this time</param>
-		/// <param name="filter">Filter for the properties to return</param>
+		/// <param name="propsFilter">If set, what keys to include in the list of agent properties. Separated by commas.</param>
+		/// <param name="filter">If set, filter for fields/properties to return</param>
 		/// <returns>List of matching agents</returns>
 		[HttpGet]
 		[Route("/api/v1/agents")]
 		[ProducesResponseType(typeof(List<GetAgentResponse>), 200)]
-		public async Task<ActionResult<List<object>>> FindAgentsAsync([FromQuery] PoolId? poolId = null, [FromQuery] Condition? condition = null, [FromQuery] bool includeDeleted = false, [FromQuery] bool invalidateCache = false, [FromQuery] int? index = null, [FromQuery] int? count = null, [FromQuery] DateTimeOffset? modifiedAfter = null, [FromQuery] PropertyFilter? filter = null)
+		public async Task<ActionResult<List<object>>> FindAgentsAsync(
+			[FromQuery] PoolId? poolId = null,
+			[FromQuery] Condition? condition = null,
+			[FromQuery] bool includeDeleted = false,
+			[FromQuery] bool invalidateCache = false,
+			[FromQuery] int? index = null,
+			[FromQuery] int? count = null,
+			[FromQuery] DateTimeOffset? modifiedAfter = null,
+			[FromQuery] string? propsFilter = null,
+			[FromQuery] PropertyFilter? filter = null)
 		{
 			using TelemetrySpan span = _tracer.StartActiveSpan($"{nameof(AgentsController)}.{nameof(FindAgentsAsync)}");
 
@@ -100,7 +110,7 @@ namespace HordeServer.Agents
 				{
 					if (condition == null || agent.SatisfiesCondition(condition))
 					{
-						responses.Add(await GetAgentResponseAsync(agent, filter));
+						responses.Add(await GetAgentResponseAsync(agent, propsFilter, filter));
 					}
 				}
 			}
@@ -132,13 +142,13 @@ namespace HordeServer.Agents
 				return NotFound(agentId);
 			}
 
-			return await GetAgentResponseAsync(agent, filter);
+			return await GetAgentResponseAsync(agent, null, filter);
 		}
 
 		/// <summary>
 		/// Gets an individual agent response
 		/// </summary>
-		async ValueTask<object> GetAgentResponseAsync(IAgent agent, PropertyFilter? filter = null, CancellationToken cancellationToken = default)
+		async ValueTask<object> GetAgentResponseAsync(IAgent agent, string? propFilter = null, PropertyFilter? filter = null, CancellationToken cancellationToken = default)
 		{
 			double? rate = null;
 			if (_computeConfig.Value.Authorize(ServerAclAction.ViewCosts, User))
@@ -164,15 +174,40 @@ namespace HordeServer.Agents
 				}
 			}
 
-			return CreateGetAgentResponse(agent, leases, rate).ApplyFilter(filter);
+			return CreateGetAgentResponse(agent, leases, rate, propFilter).ApplyFilter(filter);
 		}
 
 		internal static GetAgentLeaseResponse CreateGetAgentLeaseResponse(ILease lease, Dictionary<string, string>? details, double? agentRate)
 		{
 			return new GetAgentLeaseResponse(lease.Id, lease.ParentId, lease.AgentId, agentRate, lease.Name, lease.LogId, lease.StartTime, lease.FinishTime, lease.FinishTime == null, details, lease.Outcome, null);
 		}
+		
+		private static List<string> FilterProperties(List<string> properties, string? filter)
+		{
+			if (filter == null)
+			{
+				return properties;
+			}
+			
+			List<string> filteredProps = new (properties.Count);
+			List<string> includedKeys = filter.Split(",").Select(x => x.ToUpperInvariant()).ToList();
+			foreach (string property in properties)
+			{
+				string[] parts = property.Split("=");
+				string key = parts[0].ToUpperInvariant();
+				foreach (string includedKey in includedKeys)
+				{
+					if (key == includedKey)
+					{
+						filteredProps.Add(property);
+					}
+				}
+			}
+			
+			return filteredProps;
+		}
 
-		static GetAgentResponse CreateGetAgentResponse(IAgent agent, List<GetAgentLeaseResponse> leases, double? rate)
+		static GetAgentResponse CreateGetAgentResponse(IAgent agent, List<GetAgentLeaseResponse> leases, double? rate, string? propFilter)
 		{
 			return new GetAgentResponse(
 				agent.Id,
@@ -193,7 +228,7 @@ namespace HordeServer.Agents
 				agent.ConformAttemptCount,
 				agent.LastConformTime,
 				agent.Version?.ToString() ?? "Unknown",
-				new List<string>(agent.Properties),
+				FilterProperties(new List<string>(agent.Properties), propFilter),
 				new Dictionary<string, int>(agent.Resources),
 				agent.UpdateTime,
 				agent.LastOnlineTime,
