Horde: Add missing file

diff --git a/Engine/Source/Programs/Horde/Plugins/Compute/HordeServer.Compute.Tests/ComputeTestSetup.cs b/Engine/Source/Programs/Horde/Plugins/Compute/HordeServer.Compute.Tests/ComputeTestSetup.cs
--- a/Engine/Source/Programs/Horde/Plugins/Compute/HordeServer.Compute.Tests/ComputeTestSetup.cs
+++ b/Engine/Source/Programs/Horde/Plugins/Compute/HordeServer.Compute.Tests/ComputeTestSetup.cs
@@ -13,6 +13,7 @@ namespace HordeServer.Tests
 		public ComputeService ComputeService => ServiceProvider.GetRequiredService<ComputeService>();
 		public IAgentCollection AgentCollection => ServiceProvider.GetRequiredService<IAgentCollection>();
 		public IAgentScheduler AgentScheduler => ServiceProvider.GetRequiredService<IAgentScheduler>();
+		public AgentService AgentService => ServiceProvider.GetRequiredService<AgentService>();
 		public StorageService StorageService => ServiceProvider.GetRequiredService<StorageService>();
 		public ILogCollection LogCollection => ServiceProvider.GetRequiredService<ILogCollection>();
 
