Horde: Update release notes.

#rnx

diff --git a/Engine/Source/Programs/Horde/Docs/ReleaseNotes.md b/Engine/Source/Programs/Horde/Docs/ReleaseNotes.md
--- a/Engine/Source/Programs/Horde/Docs/ReleaseNotes.md
+++ b/Engine/Source/Programs/Horde/Docs/ReleaseNotes.md
@@ -1,5 +1,14 @@
 # Release Notes
 
+## 2024-09-25
+
+* Add filter for agent properties in REST API (36563629)
+* Re-add logical cores as an agent property. UBA uses this property for determining number of agents to allocate. (36532400)
+* Rename sample stream project to 5.5, and add a ugs-pcb artifact type. (36508675)
+* Fix path for constructing Azure storage backends. (36508415)
+* Move agent sandbox to be independent of agent settings, not in a hidden folder, and to avoid tripping path length warning by default (36505143)
+* Fix issue with conform threshold for workspaces not propagating (36502791)
+
 ## 2024-09-19
 
 * Expose .NET runtime version in agent properties (36427479)
