Horde: Use same logic for shutting down on MacOS as Linux.

diff --git a/Engine/Source/Programs/Horde/HordeAgent/Utility/Shutdown.cs b/Engine/Source/Programs/Horde/HordeAgent/Utility/Shutdown.cs
--- a/Engine/Source/Programs/Horde/HordeAgent/Utility/Shutdown.cs
+++ b/Engine/Source/Programs/Horde/HordeAgent/Utility/Shutdown.cs
@@ -118,7 +118,7 @@ namespace HordeAgent.Utility
 				}
 				return true;
 			}
-			else if (RuntimeInformation.IsOSPlatform(OSPlatform.Linux))
+			else if (RuntimeInformation.IsOSPlatform(OSPlatform.OSX) || RuntimeInformation.IsOSPlatform(OSPlatform.Linux))
 			{
 				string shutdownArgs;
 				if (restartAfterShutdown)
