Horde: Add note to temp storage exception advising users to merge a fix from //UE5/Main.

diff --git a/Engine/Source/Programs/Horde/Drivers/JobDriver/Utility/TempStorage.cs b/Engine/Source/Programs/Horde/Drivers/JobDriver/Utility/TempStorage.cs
--- a/Engine/Source/Programs/Horde/Drivers/JobDriver/Utility/TempStorage.cs
+++ b/Engine/Source/Programs/Horde/Drivers/JobDriver/Utility/TempStorage.cs
@@ -776,7 +776,7 @@ namespace JobDriver.Utility
 				DirectoryEntry? rootDirEntry;
 				if (!node.TryGetDirectoryEntry(blockDirectoryName, out rootDirEntry))
 				{
-					throw new TempStorageException($"Missing block \"{blockName}\" from node \"{nodeName}\"");
+					throw new TempStorageException($"Missing block \"{blockName}\" from node \"{nodeName}\". If running Horde against a 5.4 or earlier version of Unreal Engine, please make sure you have integrated //UE5/Main/Engine/Source/Programs/AutomationTool/BuildGraph/BuildGraph.cs#70.");
 				}
 
 				StorageStats initialStats = storageNamespace.GetStats();
