Horde: Add logging for shutdown command execution on Mac and Linux.

#rnx

diff --git a/Engine/Source/Programs/Horde/HordeAgent/Utility/Shutdown.cs b/Engine/Source/Programs/Horde/HordeAgent/Utility/Shutdown.cs
--- a/Engine/Source/Programs/Horde/HordeAgent/Utility/Shutdown.cs
+++ b/Engine/Source/Programs/Horde/HordeAgent/Utility/Shutdown.cs
@@ -132,11 +132,27 @@ namespace HordeAgent.Utility
 
 				using (Process shutdownProcess = new Process())
 				{
+					DataReceivedEventHandler handler = new DataReceivedEventHandler((_, args) =>
+					{
+						if (!String.IsNullOrEmpty(args.Data))
+						{
+							logger.LogInformation("{Output}", args.Data);
+						}
+					});
+
 					shutdownProcess.StartInfo.FileName = "/bin/sh";
-					shutdownProcess.StartInfo.Arguments = String.Format("-c \"{0}\"", shutdownArgs);
+					shutdownProcess.StartInfo.ArgumentList.Add("-c");
+					shutdownProcess.StartInfo.ArgumentList.Add(shutdownArgs);
 					shutdownProcess.StartInfo.UseShellExecute = false;
 					shutdownProcess.StartInfo.CreateNoWindow = true;
+					shutdownProcess.ErrorDataReceived += handler;
+					shutdownProcess.OutputDataReceived += handler;
+
+					logger.LogInformation("Running {Command} {Arguments}", shutdownProcess.StartInfo.FileName, shutdownProcess.StartInfo.Arguments);
+
 					shutdownProcess.Start();
+					shutdownProcess.BeginOutputReadLine();
+					shutdownProcess.BeginErrorReadLine();
 					shutdownProcess.WaitForExit();
 
 					int exitCode = shutdownProcess.ExitCode;
@@ -145,6 +161,8 @@ namespace HordeAgent.Utility
 						logger.LogError("Shutdown failed ({ExitCode})", exitCode);
 						return false;
 					}
+
+					logger.LogInformation("Exit code {ExitCode}", exitCode);
 				}
 
 				return true;
