Horde: Fix exception logging restart output on Mac/Linux.

#rnx

diff --git a/Engine/Source/Programs/Horde/HordeAgent/Utility/Shutdown.cs b/Engine/Source/Programs/Horde/HordeAgent/Utility/Shutdown.cs
--- a/Engine/Source/Programs/Horde/HordeAgent/Utility/Shutdown.cs
+++ b/Engine/Source/Programs/Horde/HordeAgent/Utility/Shutdown.cs
@@ -145,6 +145,8 @@ namespace HordeAgent.Utility
 					shutdownProcess.StartInfo.ArgumentList.Add(shutdownArgs);
 					shutdownProcess.StartInfo.UseShellExecute = false;
 					shutdownProcess.StartInfo.CreateNoWindow = true;
+					shutdownProcess.StartInfo.RedirectStandardOutput = true;
+					shutdownProcess.StartInfo.RedirectStandardError = true;
 					shutdownProcess.ErrorDataReceived += handler;
 					shutdownProcess.OutputDataReceived += handler;
 
