Horde: Add log output for restart/shutdown lease.

#rnx

diff --git a/Engine/Source/Programs/Horde/HordeAgent/Leases/Handlers/RestartHandler.cs b/Engine/Source/Programs/Horde/HordeAgent/Leases/Handlers/RestartHandler.cs
--- a/Engine/Source/Programs/Horde/HordeAgent/Leases/Handlers/RestartHandler.cs
+++ b/Engine/Source/Programs/Horde/HordeAgent/Leases/Handlers/RestartHandler.cs
@@ -1,6 +1,7 @@
 // Copyright Epic Games, Inc. All Rights Reserved.
 
 using EpicGames.Horde.Agents.Leases;
+using EpicGames.Horde.Logs;
 using HordeAgent.Services;
 using HordeAgent.Utility;
 using HordeCommon.Rpc.Messages;
@@ -17,11 +18,12 @@ namespace HordeAgent.Leases.Handlers
 		{ }
 
 		/// <inheritdoc/>
-		protected override Task<LeaseResult> ExecuteAsync(ISession session, LeaseId leaseId, RestartTask task, Tracer tracer, ILogger logger, CancellationToken cancellationToken)
+		protected override async Task<LeaseResult> ExecuteAsync(ISession session, LeaseId leaseId, RestartTask task, Tracer tracer, ILogger logger, CancellationToken cancellationToken)
 		{
-			logger.LogInformation("Scheduling restart task for agent {AgentId}", session.AgentId);
+			await using IServerLogger serverLogger = session.HordeClient.CreateServerLogger(LogId.Parse(task.LogId)).WithLocalLogger(logger);
+			serverLogger.LogInformation("Scheduling restart task for agent {AgentId} when agent shuts down.", session.AgentId);
 			SessionResult result = new SessionResult((logger, ctx) => Shutdown.ExecuteAsync(true, logger, ctx));
-			return Task.FromResult(new LeaseResult(result));
+			return new LeaseResult(result);
 		}
 	}
 
diff --git a/Engine/Source/Programs/Horde/HordeAgent/Leases/Handlers/ShutdownHandler.cs b/Engine/Source/Programs/Horde/HordeAgent/Leases/Handlers/ShutdownHandler.cs
--- a/Engine/Source/Programs/Horde/HordeAgent/Leases/Handlers/ShutdownHandler.cs
+++ b/Engine/Source/Programs/Horde/HordeAgent/Leases/Handlers/ShutdownHandler.cs
@@ -1,6 +1,7 @@
 // Copyright Epic Games, Inc. All Rights Reserved.
 
 using EpicGames.Horde.Agents.Leases;
+using EpicGames.Horde.Logs;
 using HordeAgent.Services;
 using HordeAgent.Utility;
 using HordeCommon.Rpc.Messages;
@@ -17,11 +18,12 @@ namespace HordeAgent.Leases.Handlers
 		{ }
 
 		/// <inheritdoc/>
-		protected override Task<LeaseResult> ExecuteAsync(ISession session, LeaseId leaseId, ShutdownTask task, Tracer tracer, ILogger logger, CancellationToken cancellationToken)
+		protected override async Task<LeaseResult> ExecuteAsync(ISession session, LeaseId leaseId, ShutdownTask task, Tracer tracer, ILogger logger, CancellationToken cancellationToken)
 		{
-			logger.LogInformation("Scheduling shutdown task for agent {AgentId}", session.AgentId);
+			await using IServerLogger serverLogger = session.HordeClient.CreateServerLogger(LogId.Parse(task.LogId)).WithLocalLogger(logger);
+			serverLogger.LogInformation("Scheduling shutdown task for agent {AgentId} when agent shuts down.", session.AgentId);
 			SessionResult result = new SessionResult((logger, ctx) => Shutdown.ExecuteAsync(false, logger, ctx));
-			return Task.FromResult(new LeaseResult(result));
+			return new LeaseResult(result);
 		}
 	}
 
