Horde: Return properly OAuth error responses.

#rnx

diff --git a/Engine/Source/Programs/Horde/HordeServer/Accounts/OAuthController.cs b/Engine/Source/Programs/Horde/HordeServer/Accounts/OAuthController.cs
--- a/Engine/Source/Programs/Horde/HordeServer/Accounts/OAuthController.cs
+++ b/Engine/Source/Programs/Horde/HordeServer/Accounts/OAuthController.cs
@@ -251,7 +251,8 @@ namespace HordeServer.Accounts
 				string? purpose = GetClaimOrDefault(payload, PurposeClaim);
 				if (purpose != AuthorizationCodePurpose)
 				{
-					return BadRequest($"Expected authorization token, not {purpose}");
+					OAuthGetTokenResponse error = new OAuthGetTokenResponse { Error = "invalid_request", ErrorDescription = $"Expected '{AuthorizationCodePurpose}' purpose, not '{purpose}'" };
+					return BadRequest(error);
 				}
 
 				// Perform the PKCE challenge
@@ -263,7 +264,8 @@ namespace HordeServer.Accounts
 					string verifier = ComputePkceVerifier(method, request.PkceCodeVerifier ?? String.Empty);
 					if (!String.Equals(pkceChallenge, verifier, StringComparison.OrdinalIgnoreCase))
 					{
-						return Unauthorized("PKCE verification failure");
+						OAuthGetTokenResponse error = new OAuthGetTokenResponse { Error = "invalid_request", ErrorDescription = $"PKCE verification failure" };
+						return Unauthorized(error);
 					}
 				}
 			}
@@ -275,31 +277,36 @@ namespace HordeServer.Accounts
 				string? purpose = GetClaimOrDefault(payload, PurposeClaim);
 				if (purpose != RefreshTokenPurpose)
 				{
-					return BadRequest($"Expected refresh token, not {purpose}");
+					OAuthGetTokenResponse error = new OAuthGetTokenResponse { Error = "invalid_request", ErrorDescription = $"Expected '{RefreshTokenPurpose}' purpose, not '{purpose}'" };
+					return BadRequest(error);
 				}
 			}
 			else
 			{
-				return BadRequest($"Unsupported grant type: '{request.GrantType}'");
+				OAuthGetTokenResponse error = new OAuthGetTokenResponse { Error = "unsupported_grant_type", ErrorDescription = $"Unsupported grant type: '{request.GrantType}'" };
+				return BadRequest(error);
 			}
 
 			// Get the matching account and check the session key is still valid
 			AccountId accountId;
 			if (!TryParseAccountIdFromSubject(payload, out accountId))
 			{
-				return Unauthorized("Missing account-id subject");
+				OAuthGetTokenResponse error = new OAuthGetTokenResponse { Error = "invalid_request", ErrorDescription = $"Missing account-id subject" };
+				return Unauthorized(error);
 			}
 
 			IAccount? account = await _accountCollection.GetAsync(accountId, cancellationToken);
 			if (account == null)
 			{
-				return Unauthorized($"Invalid account-id ({accountId})");
+				OAuthGetTokenResponse error = new OAuthGetTokenResponse { Error = "invalid_request", ErrorDescription = $"Invalid account-id ({accountId})" };
+				return Unauthorized(error);
 			}
 
 			string? session = GetClaimOrDefault(payload, SessionClaim);
 			if (session != account.SessionKey)
 			{
-				return Unauthorized($"Invalid session key ('{session}')");
+				OAuthGetTokenResponse error = new OAuthGetTokenResponse { Error = "invalid_request", ErrorDescription = $"Invalid session key ('{session}')" };
+				return Unauthorized(error);
 			}
 
 			// Create the tokens and response object
diff --git a/Engine/Source/Programs/Horde/HordeServer/Accounts/OAuthMessages.cs b/Engine/Source/Programs/Horde/HordeServer/Accounts/OAuthMessages.cs
--- a/Engine/Source/Programs/Horde/HordeServer/Accounts/OAuthMessages.cs
+++ b/Engine/Source/Programs/Horde/HordeServer/Accounts/OAuthMessages.cs
@@ -148,5 +148,17 @@ namespace HordeServer.Accounts
 		/// </summary>
 		[JsonPropertyName("id_token")]
 		public string? IdToken { get; set; }
+
+		/// <summary>
+		/// Error code
+		/// </summary>
+		[JsonPropertyName("error")]
+		public string? Error { get; set; }
+
+		/// <summary>
+		/// The error description text
+		/// </summary>
+		[JsonPropertyName("error_description")]
+		public string? ErrorDescription { get; set; }
 	}
 }
