Horde: Add an audit log for template schedules, under api/v1/streams/x/templates/y/history. Shows all attempts to start jobs, skipped intervals due to schedules being paused, jobs that failed to start due to missing gates, etc...

#rnx

diff --git a/Engine/Source/Programs/Horde/HordeServer.Shared/Auditing/IAuditLog.cs b/Engine/Source/Programs/Horde/HordeServer.Shared/Auditing/IAuditLog.cs
--- a/Engine/Source/Programs/Horde/HordeServer.Shared/Auditing/IAuditLog.cs
+++ b/Engine/Source/Programs/Horde/HordeServer.Shared/Auditing/IAuditLog.cs
@@ -126,7 +126,7 @@ namespace HordeServer.Auditing
 		/// <param name="count">Number of records to return</param>
 		/// <param name="cancellationToken">Cancellation token for the operation</param>
 		/// <returns>Information about the requested agent</returns>
-		public static async Task FindAsync<T>(this IAuditLogChannel<T> channel, PipeWriter bodyWriter, DateTime? minTime = null, DateTime? maxTime = null, int index = 0, int count = 50, CancellationToken cancellationToken = default)
+		public static async Task FindAsync(this IAuditLogChannel channel, PipeWriter bodyWriter, DateTime? minTime = null, DateTime? maxTime = null, int index = 0, int count = 50, CancellationToken cancellationToken = default)
 		{
 			string prefix = "{\n\t\"entries\":\n\t[";
 			await bodyWriter.WriteAsync(Encoding.UTF8.GetBytes(prefix), cancellationToken);
diff --git a/Engine/Source/Programs/Horde/Plugins/Build/HordeServer.Build/Jobs/Schedules/ScheduleService.cs b/Engine/Source/Programs/Horde/Plugins/Build/HordeServer.Build/Jobs/Schedules/ScheduleService.cs
--- a/Engine/Source/Programs/Horde/Plugins/Build/HordeServer.Build/Jobs/Schedules/ScheduleService.cs
+++ b/Engine/Source/Programs/Horde/Plugins/Build/HordeServer.Build/Jobs/Schedules/ScheduleService.cs
@@ -10,6 +10,7 @@ using EpicGames.Horde.Streams;
 using EpicGames.Redis;
 using EpicGames.Redis.Utility;
 using EpicGames.Serialization;
+using HordeServer.Auditing;
 using HordeServer.Commits;
 using HordeServer.Jobs.Graphs;
 using HordeServer.Jobs.Templates;
@@ -24,6 +25,16 @@ using StackExchange.Redis;
 
 namespace HordeServer.Jobs.Schedules
 {
+	/// <summary>
+	/// Identifier for writing log messages to a particular template
+	/// </summary>
+	public record class ScheduleId(StreamId StreamId, TemplateId TemplateId)
+	{
+		/// <inheritdoc/>
+		public override string ToString()
+			=> $"sched:{StreamId}:{TemplateId}";
+	}
+
 	/// <summary>
 	/// Manipulates schedule instances
 	/// </summary>
@@ -52,6 +63,7 @@ namespace HordeServer.Jobs.Schedules
 			public static double GetScoreFromTime(DateTime time) => (time.ToUniversalTime() - DateTime.UnixEpoch).TotalSeconds;
 		}
 
+		readonly IAuditLog<ScheduleId> _auditLog;
 		readonly IGraphCollection _graphs;
 		readonly ICommitService _commitService;
 		readonly IJobCollection _jobCollection;
@@ -72,8 +84,9 @@ namespace HordeServer.Jobs.Schedules
 		/// <summary>
 		/// Constructor
 		/// </summary>
-		public ScheduleService(IRedisService redis, IGraphCollection graphs, ICommitService commitService, IJobCollection jobCollection, JobService jobService, IDowntimeService downtimeService, IStreamCollection streamCollection, ITemplateCollection templateCollection, IMongoService mongoService, IClock clock, IOptionsMonitor<BuildConfig> buildConfig, Tracer tracer, ILogger<ScheduleService> logger)
+		public ScheduleService(IAuditLog<ScheduleId> auditLog, IRedisService redis, IGraphCollection graphs, ICommitService commitService, IJobCollection jobCollection, JobService jobService, IDowntimeService downtimeService, IStreamCollection streamCollection, ITemplateCollection templateCollection, IMongoService mongoService, IClock clock, IOptionsMonitor<BuildConfig> buildConfig, Tracer tracer, ILogger<ScheduleService> logger)
 		{
+			_auditLog = auditLog;
 			_graphs = graphs;
 			_commitService = commitService;
 			_jobCollection = jobCollection;
@@ -229,6 +242,14 @@ namespace HordeServer.Jobs.Schedules
 			await _redis.GetDatabase().SortedSetAddAsync(s_queueKey, queueItems.ToArray());
 		}
 
+		/// <summary>
+		/// Gets the audit log for a particular template
+		/// </summary>
+		/// <param name="streamId">Stream for the schedule</param>
+		/// <param name="templateId"></param>
+		public IAuditLogChannel GetAuditLog(StreamId streamId, TemplateId templateId)
+			=> _auditLog[new ScheduleId(streamId, templateId)];
+
 		/// <summary>
 		/// Trigger a schedule to run
 		/// </summary>
@@ -252,8 +273,15 @@ namespace HordeServer.Jobs.Schedules
 			}
 
 			ITemplateSchedule? schedule = templateRef.Schedule;
-			if (schedule == null || !schedule.Config.Enabled)
+			if (schedule == null)
+			{
+				return false;
+			}
+
+			ForwardingLogger scheduleLogger = new ForwardingLogger(_logger, GetAuditLog(streamId, templateId));
+			if (!schedule.Config.Enabled)
 			{
+				scheduleLogger.LogDebug("Schedule is disabled. No builds started.");
 				return false;
 			}
 
@@ -280,14 +308,14 @@ namespace HordeServer.Jobs.Schedules
 			// If the stream is paused, bail out
 			if (stream.IsPaused(utcNow))
 			{
-				_logger.LogDebug("Skipping schedule update for stream {StreamId}. It has been paused until {PausedUntil} with comment '{PauseComment}'.", stream.Id, stream.PausedUntil, stream.PauseComment);
+				scheduleLogger.LogInformation("Skipping schedule update for stream {StreamId}. It has been paused until {PausedUntil} with comment '{PauseComment}'.", stream.Id, stream.PausedUntil, stream.PauseComment);
 				return false;
 			}
 
 			// Trigger this schedule
 			try
 			{
-				await TriggerAsync(stream, templateId, templateRef, schedule, schedule.ActiveJobs.Count - removeJobIds.Count, utcNow, cancellationToken);
+				await TriggerAsync(stream, templateId, templateRef, schedule, schedule.ActiveJobs.Count - removeJobIds.Count, utcNow, scheduleLogger, cancellationToken);
 			}
 			catch (Exception ex)
 			{
@@ -311,12 +339,12 @@ namespace HordeServer.Jobs.Schedules
 		/// <param name="utcNow">The current time</param>
 		/// <param name="cancellationToken"></param>
 		/// <returns>Async task</returns>
-		private async Task TriggerAsync(IStream stream, TemplateId templateId, ITemplateRef templateRef, ITemplateSchedule schedule, int numActiveJobs, DateTime utcNow, CancellationToken cancellationToken)
+		private async Task TriggerAsync(IStream stream, TemplateId templateId, ITemplateRef templateRef, ITemplateSchedule schedule, int numActiveJobs, DateTime utcNow, ILogger scheduleLogger, CancellationToken cancellationToken)
 		{
 			// Check we're not already at the maximum number of allowed jobs
 			if (schedule.Config.MaxActive != 0 && numActiveJobs >= schedule.Config.MaxActive)
 			{
-				_logger.LogInformation("Skipping trigger of {StreamId} template {TemplateId} - already have maximum number of jobs running ({NumJobs})", stream.Id, templateId, schedule.Config.MaxActive);
+				scheduleLogger.LogInformation("Skipping trigger of {StreamId} template {TemplateId} - already have maximum number of jobs running ({NumJobs})", stream.Id, templateId, schedule.Config.MaxActive);
 				foreach (JobId jobId in schedule.ActiveJobs)
 				{
 					_logger.LogInformation("Active job for {StreamId} template {TemplateId}: {JobId}", stream.Id, templateId, jobId);
@@ -369,7 +397,7 @@ namespace HordeServer.Jobs.Schedules
 
 				if (schedule.Config.Gate != null)
 				{
-					commitId = await GetNextChangeForGateAsync(stream.Id, templateId, schedule.Config.Gate, minCommitId, maxCommitId, cancellationToken);
+					commitId = await GetNextChangeForGateAsync(stream.Id, templateId, schedule.Config.Gate, minCommitId, maxCommitId, scheduleLogger, cancellationToken);
 					commit = (commitId != null)? await commits.FindAsync(commitId, commitId, 1, null, cancellationToken).FirstOrDefaultAsync(cancellationToken) : null; // May be a change in a different stream
 				}
 				else if (await commitEnumerator.MoveNextAsync(cancellationToken))
@@ -421,7 +449,7 @@ namespace HordeServer.Jobs.Schedules
 				// Check we haven't exceeded the time limit
 				if (timer.Elapsed > TimeSpan.FromMinutes(2.0))
 				{
-					_logger.LogError("Querying for changes to trigger for {StreamId} template {TemplateId} has taken {Time}. Aborting.", stream.Id, templateId, timer.Elapsed);
+					scheduleLogger.LogError("Querying for changes to trigger for {StreamId} template {TemplateId} has taken {Time}. Aborting.", stream.Id, templateId, timer.Elapsed);
 					break;
 				}
 
@@ -436,7 +464,7 @@ namespace HordeServer.Jobs.Schedules
 			// Early out if there's nothing to do
 			if (triggerChanges.Count == 0)
 			{
-				_logger.LogInformation("Skipping trigger of {StreamName} template {TemplateId} - no candidate changes after CL {LastTriggerChange}", stream.Id, templateId, schedule.LastTriggerCommitId);
+				scheduleLogger.LogInformation("Skipping trigger of {StreamName} template {TemplateId} - no candidate changes after CL {LastTriggerChange}", stream.Id, templateId, schedule.LastTriggerCommitId);
 				return;
 			}
 
@@ -466,7 +494,7 @@ namespace HordeServer.Jobs.Schedules
 				template.GetArgumentsForParameters(options.Parameters, options.Arguments);
 
 				IJob newJob = await _jobService.CreateJobAsync(null, stream.Config, templateId, template.Hash, graph, template.Name, commitId, codeCommitId, options, cancellationToken);
-				_logger.LogInformation("Started new job for {StreamId} template {TemplateId} at CL {Change} (Code CL {CodeChange}): {JobId}", stream.Id, templateId, commitId, codeCommitId, newJob.Id);
+				scheduleLogger.LogInformation("Started new job for {StreamId} template {TemplateId} at CL {Change} (Code CL {CodeChange}): {JobId}", stream.Id, templateId, commitId, codeCommitId, newJob.Id);
 				await _streamCollection.UpdateScheduleTriggerAsync(stream, templateId, utcNow, commitId, new List<JobId> { newJob.Id }, new List<JobId>(), cancellationToken);
 			}
 		}
@@ -508,7 +536,7 @@ namespace HordeServer.Jobs.Schedules
 		/// <summary>
 		/// Gets the next change to build for a schedule on a gate
 		/// </summary>
-		private async Task<CommitIdWithOrder?> GetNextChangeForGateAsync(StreamId streamId, TemplateId templateRefId, ScheduleGateConfig gate, CommitIdWithOrder? minCommitId, CommitIdWithOrder? maxCommitId, CancellationToken cancellationToken)
+		private async Task<CommitIdWithOrder?> GetNextChangeForGateAsync(StreamId streamId, TemplateId templateRefId, ScheduleGateConfig gate, CommitIdWithOrder? minCommitId, CommitIdWithOrder? maxCommitId, ILogger scheduleLogger, CancellationToken cancellationToken)
 		{
 			for (; ; )
 			{
@@ -533,7 +561,7 @@ namespace HordeServer.Jobs.Schedules
 						{
 							return job.CommitId;
 						}
-						_logger.LogInformation("Skipping trigger of {StreamName} template {TemplateId} - last {OtherTemplateRefId} job ({JobId}) ended with errors", streamId, templateRefId, gate.TemplateId, job.Id);
+						scheduleLogger.LogInformation("Skipping trigger of {StreamName} template {TemplateId} - last {OtherTemplateRefId} job ({JobId}) ended with errors", streamId, templateRefId, gate.TemplateId, job.Id);
 					}
 				}
 
diff --git a/Engine/Source/Programs/Horde/Plugins/Build/HordeServer.Build/Streams/StreamsController.cs b/Engine/Source/Programs/Horde/Plugins/Build/HordeServer.Build/Streams/StreamsController.cs
--- a/Engine/Source/Programs/Horde/Plugins/Build/HordeServer.Build/Streams/StreamsController.cs
+++ b/Engine/Source/Programs/Horde/Plugins/Build/HordeServer.Build/Streams/StreamsController.cs
@@ -7,8 +7,10 @@ using EpicGames.Horde.Jobs.Templates;
 using EpicGames.Horde.Projects;
 using EpicGames.Horde.Streams;
 using EpicGames.Horde.Users;
+using HordeServer.Auditing;
 using HordeServer.Commits;
 using HordeServer.Jobs;
+using HordeServer.Jobs.Schedules;
 using HordeServer.Jobs.Templates;
 using HordeServer.Projects;
 using HordeServer.Server;
@@ -34,6 +36,7 @@ namespace HordeServer.Streams
 		private readonly ITemplateCollection _templateCollection;
 		private readonly IJobStepRefCollection _jobStepRefCollection;
 		private readonly IUserCollection _userCollection;
+		private readonly ScheduleService _scheduleService;
 		private readonly Tracer _tracer;
 		private readonly IOptionsSnapshot<BuildConfig> _buildConfig;
 		private readonly TimeZoneInfo _timeZone;
@@ -41,13 +44,14 @@ namespace HordeServer.Streams
 		/// <summary>
 		/// Constructor
 		/// </summary>
-		public StreamsController(IStreamCollection streamCollection, ICommitService commitService, ITemplateCollection templateCollection, IJobStepRefCollection jobStepRefCollection, IUserCollection userCollection, Tracer tracer, IClock clock, IOptionsSnapshot<BuildConfig> buildConfig)
+		public StreamsController(IStreamCollection streamCollection, ICommitService commitService, ITemplateCollection templateCollection, IJobStepRefCollection jobStepRefCollection, IUserCollection userCollection, ScheduleService scheduleService, Tracer tracer, IClock clock, IOptionsSnapshot<BuildConfig> buildConfig)
 		{
 			_streamCollection = streamCollection;
 			_commitService = commitService;
 			_templateCollection = templateCollection;
 			_jobStepRefCollection = jobStepRefCollection;
 			_userCollection = userCollection;
+			_scheduleService = scheduleService;
 			_tracer = tracer;
 			_buildConfig = buildConfig;
 			_timeZone = clock.TimeZone;
@@ -417,6 +421,47 @@ namespace HordeServer.Streams
 			return new GetTemplateResponse(template).ApplyFilter(filter);
 		}
 
+		/// <summary>
+		/// Retrieve historical information about a specific schedule
+		/// </summary>
+		/// <param name="streamId">Unique id of the stream to query</param>
+		/// <param name="templateId">Unique id of the template to query</param>
+		/// <param name="minTime">Minimum time for records to return</param>
+		/// <param name="maxTime">Maximum time for records to return</param>
+		/// <param name="index">Offset of the first result</param>
+		/// <param name="count">Number of records to return</param>
+		/// <returns>Information about the requested agent</returns>
+		[HttpGet]
+		[Route("/api/v1/streams/{streamId}/templates/{templateId}/history")]
+		public async Task<ActionResult> GetTemplateHistoryAsync(StreamId streamId, TemplateId templateId, [FromQuery] DateTime? minTime = null, [FromQuery] DateTime? maxTime = null, [FromQuery] int index = 0, [FromQuery] int count = 50)
+		{
+			StreamConfig? streamConfig;
+			if (!_buildConfig.Value.TryGetStream(streamId, out streamConfig))
+			{
+				return NotFound(streamId);
+			}
+
+			TemplateRefConfig? templateConfig;
+			if (!streamConfig.TryGetTemplate(templateId, out templateConfig))
+			{
+				return NotFound(streamId, templateId);
+			}
+			if (!templateConfig.Authorize(StreamAclAction.ViewTemplate, User))
+			{
+				return Forbid(StreamAclAction.ViewTemplate, streamId);
+			}
+
+			IAuditLogChannel channel = _scheduleService.GetAuditLog(streamId, templateId);
+
+			Response.ContentType = "application/json";
+			Response.StatusCode = 200;
+			await Response.StartAsync();
+			await channel.FindAsync(HttpContext.Response.BodyWriter, minTime, maxTime, index, count);
+			await Response.CompleteAsync();
+
+			return NoContent();
+		}
+
 		/// <summary>
 		/// Update a stream template ref
 		/// </summary>
diff --git a/Engine/Source/Programs/Horde/Plugins/Compute/HordeServer.Compute/Agents/AgentsController.cs b/Engine/Source/Programs/Horde/Plugins/Compute/HordeServer.Compute/Agents/AgentsController.cs
--- a/Engine/Source/Programs/Horde/Plugins/Compute/HordeServer.Compute/Agents/AgentsController.cs
+++ b/Engine/Source/Programs/Horde/Plugins/Compute/HordeServer.Compute/Agents/AgentsController.cs
@@ -348,12 +348,20 @@ namespace HordeServer.Agents
 		/// <returns>Information about the requested agent</returns>
 		[HttpGet]
 		[Route("/api/v1/agents/{agentId}/history")]
-		public async Task GetAgentHistoryAsync(AgentId agentId, [FromQuery] DateTime? minTime = null, [FromQuery] DateTime? maxTime = null, [FromQuery] int index = 0, [FromQuery] int count = 50)
+		public async Task<ActionResult> GetAgentHistoryAsync(AgentId agentId, [FromQuery] DateTime? minTime = null, [FromQuery] DateTime? maxTime = null, [FromQuery] int index = 0, [FromQuery] int count = 50)
 		{
+			if (!_computeConfig.Value.Authorize(AgentAclAction.ViewAgent, User))
+			{
+				return Forbid(AgentAclAction.ViewAgent, agentId);
+			}
+
 			Response.ContentType = "application/json";
 			Response.StatusCode = 200;
 			await Response.StartAsync();
 			await _agentService.Agents.GetLogger(agentId).FindAsync(Response.BodyWriter, minTime, maxTime, index, count);
+			await Response.CompleteAsync();
+
+			return NoContent();
 		}
 
 		/// <summary>
