Horde: Missing file to register an audit log for schedules.

#rnx

diff --git a/Engine/Source/Programs/Horde/Plugins/Build/HordeServer.Build/BuildPlugin.cs b/Engine/Source/Programs/Horde/Plugins/Build/HordeServer.Build/BuildPlugin.cs
--- a/Engine/Source/Programs/Horde/Plugins/Build/HordeServer.Build/BuildPlugin.cs
+++ b/Engine/Source/Programs/Horde/Plugins/Build/HordeServer.Build/BuildPlugin.cs
@@ -9,6 +9,7 @@ using HordeServer.Agents.Fleet;
 using HordeServer.Agents.Pools;
 using HordeServer.Agents.Utilization;
 using HordeServer.Artifacts;
+using HordeServer.Auditing;
 using HordeServer.Commits;
 using HordeServer.Configuration;
 using HordeServer.Devices;
@@ -103,6 +104,7 @@ namespace HordeServer
 			services.AddSingleton<UnsyncCache>();
 
 			services.AddSingleton<ScheduleService>();
+			services.AddSingleton<IAuditLog<ScheduleId>>(sp => sp.GetRequiredService<IAuditLogFactory<ScheduleId>>().Create("Streams.Log", "ScheduleId"));
 
 			services.AddSingleton<INotificationTriggerCollection, NotificationTriggerCollection>();
 			services.AddSingleton<IDeviceCollection, DeviceCollection>();
