Horde: Fix incorrect URL in Slack notification.

#rnx

diff --git a/Engine/Source/Programs/Horde/Plugins/Build/HordeServer.Build/Notifications/Sinks/SlackNotificationSink.cs b/Engine/Source/Programs/Horde/Plugins/Build/HordeServer.Build/Notifications/Sinks/SlackNotificationSink.cs
--- a/Engine/Source/Programs/Horde/Plugins/Build/HordeServer.Build/Notifications/Sinks/SlackNotificationSink.cs
+++ b/Engine/Source/Programs/Horde/Plugins/Build/HordeServer.Build/Notifications/Sinks/SlackNotificationSink.cs
@@ -412,7 +412,7 @@ namespace HordeServer.Notifications.Sinks
 				}
 
 				JobScheduledNotification notification = notifications[idx];
-				string jobUrl = _serverInfo.DashboardUrl + "/job/" + notification.JobId;
+				string jobUrl = new Uri(_serverInfo.DashboardUrl, $"job/{notification.JobId}").ToString();
 				sb.AppendLine($"Job `{notification.JobName}` with ID <{jobUrl}|{notification.JobId}> in pool `{notification.PoolName}`");
 			}
 
