Horde: Resurrect option to explicitly update build health issues in advanced parameters

diff --git a/Engine/Source/Programs/Horde/HordeDashboard/src/components/build/NewBuildV2.tsx b/Engine/Source/Programs/Horde/HordeDashboard/src/components/build/NewBuildV2.tsx
--- a/Engine/Source/Programs/Horde/HordeDashboard/src/components/build/NewBuildV2.tsx
+++ b/Engine/Source/Programs/Horde/HordeDashboard/src/components/build/NewBuildV2.tsx
@@ -72,6 +72,7 @@ class BuildOptions {
       this.queryShelvedChange = undefined;
       this.autoSubmit = undefined;
       this.advJobName = undefined;
+      this.updateIssues = undefined;
       this.advJobPriority = undefined;
       this.advAdditionalArgs = undefined;
       this.advTargets = undefined;
@@ -114,6 +115,7 @@ class BuildOptions {
       this.projectStore = projectStore;
       this.jobKey = jobKey;
       this.jobDetails = jobDetails;
+      this.updateIssues = jobDetails?.jobData?.updateIssues;
 
       this.showAllTemplates = jobKey === "all" || jobKey === "summary"
 
@@ -155,6 +157,8 @@ class BuildOptions {
    advAdditionalArgs?: string;
    advTargets?: string[];
 
+   updateIssues?: boolean;
+
    parameters: Record<string, string> = {};
    disabledParameters: Set<string> = new Set();
 
@@ -1627,8 +1631,6 @@ const AdvancedPanel: React.FC = observer(() => {
       name: string;
    }
 
-   const updateIssues = options.jobDetails?.jobData?.updateIssues;
-
    const targetItems: TargetPickerItem[] = options.targets.map(t => {
       return {
          key: t,
@@ -1638,7 +1640,7 @@ const AdvancedPanel: React.FC = observer(() => {
 
    return <Stack style={{
       height: height,
-      position: 'relative',      
+      position: 'relative',
       width: parameterWidth,
    }}>
       <ScrollablePane scrollbarVisibility={ScrollbarVisibility.auto}>
@@ -1654,11 +1656,14 @@ const AdvancedPanel: React.FC = observer(() => {
                }
                } />
             </Stack>
-            {updateIssues !== undefined && options.readOnly && <Stack>
-               <Checkbox disabled key={`key_adv_update_issues_${options.currentRenderKey}`}
+            <Stack>
+               <Checkbox disabled={options.readOnly} key={`key_adv_update_issues_${options.currentRenderKey}`}
                   label="Update Build Health Issues"
-                  checked={updateIssues} />
-            </Stack>}
+                  checked={options.updateIssues ?? false} onChange={(ev, checked) => {
+                     options.updateIssues = checked ?? false;
+                     options.setChanged();
+                  }} />
+            </Stack>
             <Stack>
                <TextField key={`key_adv_job_name_${options.currentRenderKey}`} disabled={options.readOnly} spellCheck={false} defaultValue={options.advJobName} label="Job Name" onChange={(ev, newValue) => {
                   options.advJobName = newValue;
@@ -1828,11 +1833,12 @@ const BuildModal: React.FC<{ setUseLegacyDialog: (value: boolean) => void }> = o
          streamId: options.streamId,
          templateId: templateId,
          name: options.advJobName,
-         priority: options.advJobPriority,         
+         priority: options.advJobPriority,
          changeQueries: changeQueries,
          parameters: options.parameters,
          additionalArguments: additionalArgs?.length ? additionalArgs : undefined,
-         targets: options.advTargets
+         targets: options.advTargets,
+         updateIssues: options.updateIssues
       };
 
       if (typeof (options.change) === 'number') {
