Horde: Use the correct token for writing artifacts. Adds an IArtifactBuilder interface returned from IArtifactCollection.CreateAsync() which can be used to upload artifact data with the correct permissions.

diff --git a/Engine/Source/Programs/AutomationTool/BuildGraph/Tasks/CreateArtifactTask.cs b/Engine/Source/Programs/AutomationTool/BuildGraph/Tasks/CreateArtifactTask.cs
--- a/Engine/Source/Programs/AutomationTool/BuildGraph/Tasks/CreateArtifactTask.cs
+++ b/Engine/Source/Programs/AutomationTool/BuildGraph/Tasks/CreateArtifactTask.cs
@@ -187,21 +187,17 @@ namespace AutomationTool.Tasks
 
 			// Create the new artifact
 			IHordeClient hordeClient = CommandUtils.ServiceProvider.GetRequiredService<IHordeClient>();
-			IArtifact artifact = await hordeClient.Artifacts.AddAsync(name, type, _parameters.Description, streamId, commitId, keys, metadata);
+			IArtifactBuilder artifact = await hordeClient.Artifacts.CreateAsync(name, type, _parameters.Description, streamId, commitId, keys, metadata);
 			Logger.LogInformation("Creating artifact {ArtifactId} '{ArtifactName}' ({ArtifactType}) with namespace {NamespaceId}, ref {RefName} ({Link})", artifact.Id, name, type, artifact.NamespaceId, artifact.RefName, $"{hordeClient.ServerUrl}/api/v1/storage/{artifact.NamespaceId}/refs/{artifact.RefName}");
 
 			// Upload the files
-			IStorageNamespace storage = hordeClient.GetStorageNamespace(artifact.NamespaceId);
 			Stopwatch timer = Stopwatch.StartNew();
 
 			IHashedBlobRef<DirectoryNode> rootRef;
-			await using (IBlobWriter blobWriter = storage.CreateBlobWriter(artifact.RefName))
+			await using (IBlobWriter blobWriter = artifact.CreateBlobWriter())
 			{
 				rootRef = await blobWriter.WriteFilesAsync(baseDir, files);
 			}
-			await storage.WriteRefAsync(artifact.RefName, rootRef);
-
-			Logger.LogInformation("Uploaded artifact {ArtifactId} in {Time:n1}s", artifact.Id, timer.Elapsed.TotalSeconds);
 
 			// Tag any uploaded symbols
 			if (_parameters.Symbols)
@@ -223,11 +219,15 @@ namespace AutomationTool.Tasks
 							string fileName = file.GetFileName().ToUpperInvariant();
 							string alias = $"sym:{fileName}/{hash}/{fileName}";
 							Logger.LogInformation("Adding symbol alias: {Alias}", alias);
-							await storage.AddAliasAsync(alias, fileEntry.Target.Handle);
+							await artifact.AddAliasAsync(alias, fileEntry.Target.Handle);
 						}
 					}
 				}
 			}
+
+			// Finalize the artifact
+			await artifact.CompleteAsync(rootRef);
+			Logger.LogInformation("Uploaded artifact {ArtifactId} in {Time:n1}s", artifact.Id, timer.Elapsed.TotalSeconds);
 		}
 
 		static async Task<FileEntry?> FindFileAsync(IBlobRef<DirectoryNode> rootDir, string path, CancellationToken cancellationToken = default)
diff --git a/Engine/Source/Programs/Horde/Plugins/Build/HordeServer.Build.Tests/Artifacts/ArtifactTests.cs b/Engine/Source/Programs/Horde/Plugins/Build/HordeServer.Build.Tests/Artifacts/ArtifactTests.cs
--- a/Engine/Source/Programs/Horde/Plugins/Build/HordeServer.Build.Tests/Artifacts/ArtifactTests.cs
+++ b/Engine/Source/Programs/Horde/Plugins/Build/HordeServer.Build.Tests/Artifacts/ArtifactTests.cs
@@ -36,7 +36,7 @@ namespace HordeServer.Tests.Artifacts
 			StreamId streamId = new StreamId("foo");
 
 			IArtifactCollection artifactCollection = ServiceProvider.GetRequiredService<IArtifactCollection>();
-			IArtifact artifact = await artifactCollection.AddAsync(new ArtifactName("default"), ArtifactType.StepOutput, null, streamId, CommitIdWithOrder.FromPerforceChange(1), new string[] { "test1", "test2" }, Array.Empty<string>());
+			IArtifactBuilder artifact = await artifactCollection.CreateAsync(new ArtifactName("default"), ArtifactType.StepOutput, null, streamId, CommitIdWithOrder.FromPerforceChange(1), new string[] { "test1", "test2" }, Array.Empty<string>());
 
 			{
 				List<IArtifact> artifacts = await artifactCollection.FindAsync(streamId, keys: new[] { "test1" }).ToListAsync();
@@ -70,7 +70,7 @@ namespace HordeServer.Tests.Artifacts
 			StreamId streamId = new StreamId("foo");
 
 			IArtifactCollection artifactCollection = ServiceProvider.GetRequiredService<IArtifactCollection>();
-			IArtifact artifact = await artifactCollection.AddAsync(new ArtifactName("default"), type, null, streamId, CommitIdWithOrder.FromPerforceChange(1), new string[] { "test1", "test2" }, Array.Empty<string>());
+			IArtifactBuilder artifact = await artifactCollection.CreateAsync(new ArtifactName("default"), type, null, streamId, CommitIdWithOrder.FromPerforceChange(1), new string[] { "test1", "test2" }, Array.Empty<string>());
 
 			{
 				List<IArtifact> artifacts = await artifactCollection.FindAsync(streamId, keys: new[] { "test1" }).ToListAsync();
@@ -101,7 +101,7 @@ namespace HordeServer.Tests.Artifacts
 			StreamId streamId = new StreamId("foo");
 
 			IArtifactCollection artifactCollection = ServiceProvider.GetRequiredService<IArtifactCollection>();
-			IArtifact artifact = await artifactCollection.AddAsync(new ArtifactName("default"), type, null, streamId, CommitIdWithOrder.FromPerforceChange(1), new string[] { "test1", "test2" }, Array.Empty<string>());
+			IArtifactBuilder artifact = await artifactCollection.CreateAsync(new ArtifactName("default"), type, null, streamId, CommitIdWithOrder.FromPerforceChange(1), new string[] { "test1", "test2" }, Array.Empty<string>());
 
 			{
 				List<IArtifact> artifacts = await artifactCollection.FindAsync(streamId, keys: new[] { "test1" }).ToListAsync();
@@ -143,11 +143,11 @@ namespace HordeServer.Tests.Artifacts
 			IArtifactCollection artifactCollection = ServiceProvider.GetRequiredService<IArtifactCollection>();
 			for (int idx = 0; idx < 10; idx++)
 			{
-				await artifactCollection.AddAsync(new ArtifactName($"default-{idx}"), type, null, fooStreamId, CommitIdWithOrder.FromPerforceChange(1), new string[] { "test1", "test2" }, Array.Empty<string>());
+				await artifactCollection.CreateAsync(new ArtifactName($"default-{idx}"), type, null, fooStreamId, CommitIdWithOrder.FromPerforceChange(1), new string[] { "test1", "test2" }, Array.Empty<string>());
 			}
 			for (int idx = 0; idx < 10; idx++)
 			{
-				await artifactCollection.AddAsync(new ArtifactName($"default-{idx}"), type, null, barStreamId, CommitIdWithOrder.FromPerforceChange(1), new string[] { "test1", "test2" }, Array.Empty<string>());
+				await artifactCollection.CreateAsync(new ArtifactName($"default-{idx}"), type, null, barStreamId, CommitIdWithOrder.FromPerforceChange(1), new string[] { "test1", "test2" }, Array.Empty<string>());
 			}
 
 			await Clock.AdvanceAsync(TimeSpan.FromDays(2.0));
@@ -198,7 +198,7 @@ namespace HordeServer.Tests.Artifacts
 			// Orphan the artifact type
 			{
 				UpdateConfig(config => config.Plugins.GetBuildConfig().ArtifactTypes.Add(new ArtifactTypeConfig { Type = type, KeepCount = 4 }));
-				IArtifact artifact = await artifactCollection.AddAsync(new ArtifactName($"default"), type, null, fooStreamId, CommitIdWithOrder.FromPerforceChange(1), new string[] { "test1", "test2" }, Array.Empty<string>());
+				IArtifactBuilder artifact = await artifactCollection.CreateAsync(new ArtifactName($"default"), type, null, fooStreamId, CommitIdWithOrder.FromPerforceChange(1), new string[] { "test1", "test2" }, Array.Empty<string>());
 
 				await Clock.AdvanceAsync(TimeSpan.FromDays(30.0));
 				IArtifact? artifact2 = await artifactCollection.GetAsync(artifact.Id);
@@ -216,7 +216,7 @@ namespace HordeServer.Tests.Artifacts
 			// Orphan the stream
 			{
 				UpdateConfig(config => config.Plugins.GetBuildConfig().ArtifactTypes.Add(new ArtifactTypeConfig { Type = type, KeepCount = 4 }));
-				IArtifact artifact = await artifactCollection.AddAsync(new ArtifactName($"default"), type, null, fooStreamId, CommitIdWithOrder.FromPerforceChange(1), new string[] { "test1", "test2" }, Array.Empty<string>());
+				IArtifactBuilder artifact = await artifactCollection.CreateAsync(new ArtifactName($"default"), type, null, fooStreamId, CommitIdWithOrder.FromPerforceChange(1), new string[] { "test1", "test2" }, Array.Empty<string>());
 
 				await Clock.AdvanceAsync(TimeSpan.FromDays(30.0));
 				IArtifact? artifact2 = await artifactCollection.GetAsync(artifact.Id);
diff --git a/Engine/Source/Programs/Horde/Plugins/Build/HordeServer.Build/Artifacts/ArtifactCollection.cs b/Engine/Source/Programs/Horde/Plugins/Build/HordeServer.Build/Artifacts/ArtifactCollection.cs
--- a/Engine/Source/Programs/Horde/Plugins/Build/HordeServer.Build/Artifacts/ArtifactCollection.cs
+++ b/Engine/Source/Programs/Horde/Plugins/Build/HordeServer.Build/Artifacts/ArtifactCollection.cs
@@ -35,7 +35,7 @@ namespace HordeServer.Artifacts
 
 			public ArtifactDocument Document => _document;
 
-			ArtifactId IArtifact.Id => new ArtifactId(BinaryIdUtils.FromObjectId(_document.Id));
+			public ArtifactId Id => new ArtifactId(BinaryIdUtils.FromObjectId(_document.Id));
 			ArtifactName IArtifact.Name => _document.Name;
 			ArtifactType IArtifact.Type => _document.Type;
 			string? IArtifact.Description => _document.Description;
@@ -143,6 +143,43 @@ namespace HordeServer.Artifacts
 			public DateTime Time { get; set; }
 		}
 
+		class ArtifactBuilder : IArtifactBuilder
+		{
+			readonly ArtifactCollection _collection;
+			readonly IArtifact _artifact;
+			readonly IStorageNamespace _storage;
+
+			public ArtifactId Id => _artifact.Id;
+			public ArtifactName Name => _artifact.Name;
+			public ArtifactType Type => _artifact.Type;
+			public string? Description => _artifact.Description;
+			public StreamId StreamId => _artifact.StreamId;
+			public CommitIdWithOrder CommitId => _artifact.CommitId;
+			public IReadOnlyList<string> Keys => _artifact.Keys;
+			public IReadOnlyList<string> Metadata => _artifact.Metadata;
+			public NamespaceId NamespaceId => _artifact.NamespaceId;
+			public RefName RefName => _artifact.RefName;
+
+			public ArtifactBuilder(ArtifactCollection collection, Artifact artifact, IStorageNamespace storage)
+			{
+				_collection = collection;
+				_artifact = artifact;
+				_storage = storage;
+			}
+
+			public Task AddAliasAsync(string name, IBlobRef handle, int rank = 0, ReadOnlyMemory<byte> data = default, CancellationToken cancellationToken = default)
+				=> _storage.AddAliasAsync(name, handle, rank, data, cancellationToken);
+
+			public async Task<IArtifact> CompleteAsync(IHashedBlobRef blobRef, CancellationToken cancellationToken = default)
+			{
+				await _storage.WriteRefAsync(_artifact.RefName, blobRef, cancellationToken: cancellationToken);
+				return _artifact;
+			}
+
+			public IBlobWriter CreateBlobWriter()
+				=> _storage.CreateBlobWriter(_artifact.RefName);
+		}
+
 		readonly IMongoCollection<ArtifactDocument> _artifactCollection;
 		readonly IMongoCollection<ArtifactExpiryDocument> _artifactExpiryCollection;
 		readonly IClock _clock;
@@ -205,7 +242,7 @@ namespace HordeServer.Artifacts
 		public static string GetArtifactPath(StreamId streamId, ArtifactType type) => $"{type}/{streamId}";
 
 		/// <inheritdoc/>
-		public async Task<IArtifact> AddAsync(ArtifactName name, ArtifactType type, string? description, StreamId streamId, CommitId commitId, IEnumerable<string> keys, IEnumerable<string> metadata, CancellationToken cancellationToken = default)
+		public async Task<IArtifactBuilder> CreateAsync(ArtifactName name, ArtifactType type, string? description, StreamId streamId, CommitId commitId, IEnumerable<string> keys, IEnumerable<string> metadata, CancellationToken cancellationToken = default)
 		{
 			if (name.Id.IsEmpty)
 			{
@@ -237,7 +274,8 @@ namespace HordeServer.Artifacts
 			ArtifactDocument artifactDocument = new ArtifactDocument(id, name, type, description, streamId, commitIdWithOrder, keys.Select(x => NormalizeKey(x)), metadata, namespaceId, refName, _clock.UtcNow);
 			await _artifactCollection.InsertOneAsync(artifactDocument, null, cancellationToken);
 
-			return new Artifact(this, artifactDocument);
+			Artifact artifact = new Artifact(this, artifactDocument);
+			return new ArtifactBuilder(this, artifact, _storageService.GetNamespace(namespaceId));
 		}
 
 		async Task AddExpiryRecordAsync(StreamId streamId, ArtifactType type, DateTime utcNow, CancellationToken cancellationToken = default)
diff --git a/Engine/Source/Programs/Horde/Plugins/Build/HordeServer.Build/Artifacts/ArtifactsController.cs b/Engine/Source/Programs/Horde/Plugins/Build/HordeServer.Build/Artifacts/ArtifactsController.cs
--- a/Engine/Source/Programs/Horde/Plugins/Build/HordeServer.Build/Artifacts/ArtifactsController.cs
+++ b/Engine/Source/Programs/Horde/Plugins/Build/HordeServer.Build/Artifacts/ArtifactsController.cs
@@ -145,7 +145,7 @@ namespace HordeServer.Artifacts
 
 		async Task<ActionResult<CreateArtifactResponse>> CreateArtifactInternalAsync(ArtifactName name, ArtifactType type, string? description, StreamId streamId, CommitId commitId, List<string> keys, List<string> metadata, CancellationToken cancellationToken)
 		{
-			IArtifact artifact = await _artifactCollection.AddAsync(name, type, description, streamId, commitId, keys, metadata, cancellationToken);
+			IArtifactBuilder artifact = await _artifactCollection.CreateAsync(name, type, description, streamId, commitId, keys, metadata, cancellationToken);
 			RefName? prevRefName = await GetPrevRefNameForArtifactAsync(artifact, cancellationToken);
 
 			List<AclClaimConfig> claims = new List<AclClaimConfig>();
@@ -156,7 +156,7 @@ namespace HordeServer.Artifacts
 			return new CreateArtifactResponse(artifact.Id, artifact.CommitId, artifact.NamespaceId, artifact.RefName, prevRefName, token);
 		}
 
-		async Task<RefName?> GetPrevRefNameForArtifactAsync(IArtifact artifact, CancellationToken cancellationToken)
+		async Task<RefName?> GetPrevRefNameForArtifactAsync(IArtifactBuilder artifact, CancellationToken cancellationToken)
 		{
 			IStorageBackend storageBackend = _storageService.CreateBackend(artifact.NamespaceId);
 			await foreach (IArtifact prevArtifact in _artifactCollection.FindAsync(artifact.StreamId, maxCommitId: artifact.CommitId, name: artifact.Name, type: artifact.Type, cancellationToken: cancellationToken))
diff --git a/Engine/Source/Programs/Horde/Plugins/Build/HordeServer.Build/Jobs/JobRpcService.cs b/Engine/Source/Programs/Horde/Plugins/Build/HordeServer.Build/Jobs/JobRpcService.cs
--- a/Engine/Source/Programs/Horde/Plugins/Build/HordeServer.Build/Jobs/JobRpcService.cs
+++ b/Engine/Source/Programs/Horde/Plugins/Build/HordeServer.Build/Jobs/JobRpcService.cs
@@ -826,7 +826,7 @@ namespace HordeServer.Jobs
 				description = request.Name;
 			}
 
-			IArtifact artifact = await _artifactCollection.AddAsync(name, type, description, job.StreamId, job.CommitId, keys, request.Metadata, context.CancellationToken);
+			IArtifactBuilder artifact = await _artifactCollection.CreateAsync(name, type, description, job.StreamId, job.CommitId, keys, request.Metadata, context.CancellationToken);
 
 			List<AclClaimConfig> claims = new List<AclClaimConfig>();
 			claims.Add(new AclClaimConfig(HordeClaimTypes.WriteNamespace, $"{artifact.NamespaceId}:{artifact.RefName}"));
diff --git a/Engine/Source/Programs/Shared/EpicGames.Horde/Artifacts/HttpArtifactCollection.cs b/Engine/Source/Programs/Shared/EpicGames.Horde/Artifacts/HttpArtifactCollection.cs
--- a/Engine/Source/Programs/Shared/EpicGames.Horde/Artifacts/HttpArtifactCollection.cs
+++ b/Engine/Source/Programs/Shared/EpicGames.Horde/Artifacts/HttpArtifactCollection.cs
@@ -61,6 +61,48 @@ namespace EpicGames.Horde.Artifacts
 				=> _collection.DeleteAsync(Id, cancellationToken);
 		}
 
+		class ArtifactBuilder : IArtifactBuilder
+		{
+			readonly HttpArtifactCollection _collection;
+			readonly Artifact _artifact;
+			readonly CreateArtifactResponse _response;
+			readonly IStorageNamespace _namespace;
+
+			public ArtifactId Id => _artifact.Id;
+			public ArtifactName Name => _artifact.Name;
+			public ArtifactType Type => _artifact.Type;
+			public string? Description => _artifact.Description;
+			public StreamId StreamId => _artifact.StreamId;
+			public CommitIdWithOrder CommitId => _artifact.CommitId;
+			public IReadOnlyList<string> Keys => _artifact.Keys;
+			public IReadOnlyList<string> Metadata => _artifact.Metadata;
+			public NamespaceId NamespaceId => _artifact.NamespaceId;
+			public RefName RefName => _artifact.RefName;
+
+			public ArtifactBuilder(HttpArtifactCollection collection, Artifact artifact, CreateArtifactResponse response)
+			{
+				_collection = collection;
+				_artifact = artifact;
+				_response = response;
+				_namespace = collection._hordeClient.GetStorageNamespace(response.NamespaceId, response.Token);
+			}
+
+			/// <inheritdoc/>
+			public Task AddAliasAsync(string name, IBlobRef handle, int rank = 0, ReadOnlyMemory<byte> data = default, CancellationToken cancellationToken = default)
+				=> _namespace.AddAliasAsync(name, handle, rank, data, cancellationToken);
+
+			/// <inheritdoc/>
+			public async Task<IArtifact> CompleteAsync(IHashedBlobRef blobRef, CancellationToken cancellationToken = default)
+			{
+				await _namespace.WriteRefAsync(_response.RefName, blobRef);
+				return _artifact;
+			}
+
+			/// <inheritdoc/>
+			public IBlobWriter CreateBlobWriter()
+				=> _namespace.CreateBlobWriter(_response.RefName);
+		}
+
 		readonly IHordeClient _hordeClient;
 
 		public HttpArtifactCollection(IHordeClient hordeClient)
@@ -73,11 +115,12 @@ namespace EpicGames.Horde.Artifacts
 		}
 
 		/// <inheritdoc/>
-		public async Task<IArtifact> AddAsync(ArtifactName name, ArtifactType type, string? description, StreamId streamId, CommitId commitId, IEnumerable<string> keys, IEnumerable<string> metadata, CancellationToken cancellationToken = default)
+		public async Task<IArtifactBuilder> CreateAsync(ArtifactName name, ArtifactType type, string? description, StreamId streamId, CommitId commitId, IEnumerable<string> keys, IEnumerable<string> metadata, CancellationToken cancellationToken = default)
 		{
 			HordeHttpClient hordeHttpClient = _hordeClient.CreateHttpClient();
 			CreateArtifactResponse response = await hordeHttpClient.CreateArtifactAsync(name, type, description, streamId, commitId, keys, metadata, cancellationToken);
-			return new Artifact(this, response.ArtifactId, name, type, description, streamId, CommitIdWithOrder.FromPerforceChange(commitId.GetPerforceChange()), keys.ToList(), metadata.ToList(), response.NamespaceId, response.RefName, DateTime.UtcNow);
+			Artifact artifact = new Artifact(this, response.ArtifactId, name, type, description, streamId, CommitIdWithOrder.FromPerforceChange(commitId.GetPerforceChange()), keys.ToList(), metadata.ToList(), response.NamespaceId, response.RefName, DateTime.UtcNow);
+			return new ArtifactBuilder(this, artifact, response);
 		}
 
 		/// <inheritdoc/>
diff --git a/Engine/Source/Programs/Shared/EpicGames.Horde/Artifacts/IArtifact.cs b/Engine/Source/Programs/Shared/EpicGames.Horde/Artifacts/IArtifact.cs
--- a/Engine/Source/Programs/Shared/EpicGames.Horde/Artifacts/IArtifact.cs
+++ b/Engine/Source/Programs/Shared/EpicGames.Horde/Artifacts/IArtifact.cs
@@ -82,4 +82,63 @@ namespace EpicGames.Horde.Artifacts
 		/// <param name="cancellationToken">Cancellation token for the operation</param>
 		Task DeleteAsync(CancellationToken cancellationToken);
 	}
+
+	/// <summary>
+	/// Used to write data into an artifact
+	/// </summary>
+	public interface IArtifactBuilder
+	{
+		/// <inheritdoc cref="IArtifact.Id"/>
+		ArtifactId Id { get; }
+
+		/// <inheritdoc cref="IArtifact.Name"/>
+		ArtifactName Name { get; }
+
+		/// <inheritdoc cref="IArtifact.Type"/>
+		ArtifactType Type { get; }
+
+		/// <inheritdoc cref="IArtifact.Description"/>
+		string? Description { get; }
+
+		/// <inheritdoc cref="IArtifact.StreamId"/>
+		StreamId StreamId { get; }
+
+		/// <inheritdoc cref="IArtifact.CommitId"/>
+		CommitIdWithOrder CommitId { get; }
+
+		/// <inheritdoc cref="IArtifact.Keys"/>
+		IReadOnlyList<string> Keys { get; }
+
+		/// <inheritdoc cref="IArtifact.Metadata"/>
+		IReadOnlyList<string> Metadata { get; }
+
+		/// <inheritdoc cref="IArtifact.NamespaceId"/>
+		NamespaceId NamespaceId { get; }
+
+		/// <inheritdoc cref="IArtifact.RefName"/>
+		RefName RefName { get; }
+
+		/// <summary>
+		/// Creates a writer for new artifact blobs
+		/// </summary>
+		IBlobWriter CreateBlobWriter();
+
+		/// <summary>
+		/// Adds an alias to a given blob
+		/// </summary>
+		/// <param name="name">Alias for the blob</param>
+		/// <param name="handle">Locator for the blob</param>
+		/// <param name="rank">Rank for this alias. In situations where an alias has multiple mappings, the alias with the highest rank will be returned by default.</param>
+		/// <param name="data">Additional data to be stored inline with the alias</param>
+		/// <param name="cancellationToken">Cancellation token for the operation</param>
+		Task AddAliasAsync(string name, IBlobRef handle, int rank = 0, ReadOnlyMemory<byte> data = default, CancellationToken cancellationToken = default);
+
+		/// <summary>
+		/// Finish writing the artifact data
+		/// </summary>
+		/// <param name="blobRef">Root blob for the artifact</param>
+		/// <param name="cancellationToken">Cancellation token for the operation</param>
+		/// <returns>The complete artifact</returns>
+		Task<IArtifact> CompleteAsync(IHashedBlobRef blobRef, CancellationToken cancellationToken = default);
+	}
 }
diff --git a/Engine/Source/Programs/Shared/EpicGames.Horde/Artifacts/IArtifactCollection.cs b/Engine/Source/Programs/Shared/EpicGames.Horde/Artifacts/IArtifactCollection.cs
--- a/Engine/Source/Programs/Shared/EpicGames.Horde/Artifacts/IArtifactCollection.cs
+++ b/Engine/Source/Programs/Shared/EpicGames.Horde/Artifacts/IArtifactCollection.cs
@@ -4,6 +4,7 @@ using System.Collections.Generic;
 using System.Threading;
 using System.Threading.Tasks;
 using EpicGames.Horde.Commits;
+using EpicGames.Horde.Storage;
 using EpicGames.Horde.Streams;
 
 namespace EpicGames.Horde.Artifacts
@@ -25,7 +26,7 @@ namespace EpicGames.Horde.Artifacts
 		/// <param name="metadata">Metadata for the artifact</param>
 		/// <param name="cancellationToken">Cancellation token for the operation</param>
 		/// <returns>The new log file document</returns>
-		Task<IArtifact> AddAsync(ArtifactName name, ArtifactType type, string? description, StreamId streamId, CommitId commitId, IEnumerable<string> keys, IEnumerable<string> metadata, CancellationToken cancellationToken = default);
+		Task<IArtifactBuilder> CreateAsync(ArtifactName name, ArtifactType type, string? description, StreamId streamId, CommitId commitId, IEnumerable<string> keys, IEnumerable<string> metadata, CancellationToken cancellationToken = default);
 
 		/// <summary>
 		/// Finds artifacts with the given keys.
