Horde: Remove toolbox files from agent installer
#rnx
#skipci

diff --git a/Engine/Source/Programs/Horde/BuildHorde.xml b/Engine/Source/Programs/Horde/BuildHorde.xml
--- a/Engine/Source/Programs/Horde/BuildHorde.xml
+++ b/Engine/Source/Programs/Horde/BuildHorde.xml
@@ -167,10 +167,6 @@
 			<DotNet Arguments="publish &quot;$(RootDir)/Engine/Source/Programs/Horde/Drivers/JobDriver/JobDriver.csproj&quot; --output &quot;$(ServerToolsStagingDir_HordeAgentWin64)/JobDriver&quot; --runtime win-x64 --self-contained $(VersionArguments)"/>
 			<Sign Files="$(ServerToolsStagingDir_HordeAgentWin64)/JobDriver/JobDriver.exe;$(ServerToolsStagingDir_HordeAgentWin64)/JobDriver/JobDriver.dll;$(ServerToolsStagingDir_HordeAgentWin64)/JobDriver/EpicGames.*.dll" If="'$(SignOutputs)' == 'true'"/>
 
-			<!-- Build the toolbox -->
-			<DotNet Arguments="publish &quot;$(RootDir)/Engine/Source/Programs/UnrealToolbox/UnrealToolbox.csproj&quot; --output &quot;$(ServerToolsStagingDir_HordeAgentWin64)/Toolbox&quot; --runtime win-x64 --self-contained $(VersionArguments)"/>
-			<Sign Files="$(ServerToolsStagingDir_HordeAgentWin64)/Toolbox/UnrealToolbox.exe;$(ServerToolsStagingDir_HordeAgentWin64)/Toolbox/EpicGames.*.dll" If="'$(SignOutputs)' == 'true'"/>
-
 			<Delete Files="$(ServerToolsStagingDir_HordeAgentWin64)/appsettings.Local.json"/>
 			<Tag Files="$(ServerToolsStagingDir_HordeAgentWin64)/..." With="#HordeAgentWin64Files"/>
 		</Node>
