Horde: Fix issue with mixed ungrouped and grouped tool downloads
#skipci

diff --git a/Engine/Source/Programs/Horde/HordeDashboard/src/components/ToolView.tsx b/Engine/Source/Programs/Horde/HordeDashboard/src/components/ToolView.tsx
--- a/Engine/Source/Programs/Horde/HordeDashboard/src/components/ToolView.tsx
+++ b/Engine/Source/Programs/Horde/HordeDashboard/src/components/ToolView.tsx
@@ -170,6 +170,8 @@ const ToolPanel: React.FC<{ selectedKey: string }> = observer(({ selectedKey })
 
    const toolHide = new Set<string>();
 
+   let ungrouped = 0;
+
    // emit groups
    for (let i = 0; i < tools.length; i++) {
       const tool = tools[i];
@@ -180,7 +182,7 @@ const ToolPanel: React.FC<{ selectedKey: string }> = observer(({ selectedKey })
             groups[groups.length - 1].count = i - groups[groups.length - 1].startIndex;
          }
 
-         cgroup = tool.group;
+         cgroup = tool.group ?? `Ungrouped ${ungrouped++}`;
 
          if (cgroup) {
             groups.push({ startIndex: i, name: tool.name, key: `group_key_${cgroup}`, count: 0, isCollapsed: true, data: tool });
