Horde: Fix static analysis warnings.

#rnx

diff --git a/Engine/Source/Programs/Horde/Plugins/Build/HordeServer.Build/Artifacts/ArtifactCollection.cs b/Engine/Source/Programs/Horde/Plugins/Build/HordeServer.Build/Artifacts/ArtifactCollection.cs
--- a/Engine/Source/Programs/Horde/Plugins/Build/HordeServer.Build/Artifacts/ArtifactCollection.cs
+++ b/Engine/Source/Programs/Horde/Plugins/Build/HordeServer.Build/Artifacts/ArtifactCollection.cs
@@ -145,7 +145,6 @@ namespace HordeServer.Artifacts
 
 		class ArtifactBuilder : IArtifactBuilder
 		{
-			readonly ArtifactCollection _collection;
 			readonly IArtifact _artifact;
 			readonly IStorageNamespace _storage;
 
@@ -160,9 +159,8 @@ namespace HordeServer.Artifacts
 			public NamespaceId NamespaceId => _artifact.NamespaceId;
 			public RefName RefName => _artifact.RefName;
 
-			public ArtifactBuilder(ArtifactCollection collection, Artifact artifact, IStorageNamespace storage)
+			public ArtifactBuilder(Artifact artifact, IStorageNamespace storage)
 			{
-				_collection = collection;
 				_artifact = artifact;
 				_storage = storage;
 			}
@@ -275,7 +273,7 @@ namespace HordeServer.Artifacts
 			await _artifactCollection.InsertOneAsync(artifactDocument, null, cancellationToken);
 
 			Artifact artifact = new Artifact(this, artifactDocument);
-			return new ArtifactBuilder(this, artifact, _storageService.GetNamespace(namespaceId));
+			return new ArtifactBuilder(artifact, _storageService.GetNamespace(namespaceId));
 		}
 
 		async Task AddExpiryRecordAsync(StreamId streamId, ArtifactType type, DateTime utcNow, CancellationToken cancellationToken = default)
diff --git a/Engine/Source/Programs/Horde/Plugins/Build/HordeServer.Build/Jobs/Schedules/ScheduleService.cs b/Engine/Source/Programs/Horde/Plugins/Build/HordeServer.Build/Jobs/Schedules/ScheduleService.cs
--- a/Engine/Source/Programs/Horde/Plugins/Build/HordeServer.Build/Jobs/Schedules/ScheduleService.cs
+++ b/Engine/Source/Programs/Horde/Plugins/Build/HordeServer.Build/Jobs/Schedules/ScheduleService.cs
@@ -337,6 +337,7 @@ namespace HordeServer.Jobs.Schedules
 		/// <param name="schedule"></param>
 		/// <param name="numActiveJobs"></param>
 		/// <param name="utcNow">The current time</param>
+		/// <param name="scheduleLogger"></param>
 		/// <param name="cancellationToken"></param>
 		/// <returns>Async task</returns>
 		private async Task TriggerAsync(IStream stream, TemplateId templateId, ITemplateRef templateRef, ITemplateSchedule schedule, int numActiveJobs, DateTime utcNow, ILogger scheduleLogger, CancellationToken cancellationToken)
diff --git a/Engine/Source/Programs/Shared/EpicGames.Horde/Artifacts/HttpArtifactCollection.cs b/Engine/Source/Programs/Shared/EpicGames.Horde/Artifacts/HttpArtifactCollection.cs
--- a/Engine/Source/Programs/Shared/EpicGames.Horde/Artifacts/HttpArtifactCollection.cs
+++ b/Engine/Source/Programs/Shared/EpicGames.Horde/Artifacts/HttpArtifactCollection.cs
@@ -63,7 +63,6 @@ namespace EpicGames.Horde.Artifacts
 
 		class ArtifactBuilder : IArtifactBuilder
 		{
-			readonly HttpArtifactCollection _collection;
 			readonly Artifact _artifact;
 			readonly CreateArtifactResponse _response;
 			readonly IStorageNamespace _namespace;
@@ -81,7 +80,6 @@ namespace EpicGames.Horde.Artifacts
 
 			public ArtifactBuilder(HttpArtifactCollection collection, Artifact artifact, CreateArtifactResponse response)
 			{
-				_collection = collection;
 				_artifact = artifact;
 				_response = response;
 				_namespace = collection._hordeClient.GetStorageNamespace(response.NamespaceId, response.Token);
@@ -94,7 +92,7 @@ namespace EpicGames.Horde.Artifacts
 			/// <inheritdoc/>
 			public async Task<IArtifact> CompleteAsync(IHashedBlobRef blobRef, CancellationToken cancellationToken = default)
 			{
-				await _namespace.WriteRefAsync(_response.RefName, blobRef);
+				await _namespace.WriteRefAsync(_response.RefName, blobRef, cancellationToken: cancellationToken);
 				return _artifact;
 			}
 
diff --git a/Engine/Source/Programs/Shared/EpicGames.Horde/Artifacts/IArtifactCollection.cs b/Engine/Source/Programs/Shared/EpicGames.Horde/Artifacts/IArtifactCollection.cs
--- a/Engine/Source/Programs/Shared/EpicGames.Horde/Artifacts/IArtifactCollection.cs
+++ b/Engine/Source/Programs/Shared/EpicGames.Horde/Artifacts/IArtifactCollection.cs
@@ -4,7 +4,6 @@ using System.Collections.Generic;
 using System.Threading;
 using System.Threading.Tasks;
 using EpicGames.Horde.Commits;
-using EpicGames.Horde.Storage;
 using EpicGames.Horde.Streams;
 
 namespace EpicGames.Horde.Artifacts
