Horde: Yet more step trends fixes and tweaks
#skipci

diff --git a/Engine/Source/Programs/Horde/HordeDashboard/src/components/jobDetailsV2/JobDetailStepTrendsV2.tsx b/Engine/Source/Programs/Horde/HordeDashboard/src/components/jobDetailsV2/JobDetailStepTrendsV2.tsx
--- a/Engine/Source/Programs/Horde/HordeDashboard/src/components/jobDetailsV2/JobDetailStepTrendsV2.tsx
+++ b/Engine/Source/Programs/Horde/HordeDashboard/src/components/jobDetailsV2/JobDetailStepTrendsV2.tsx
@@ -83,10 +83,6 @@ class StepTrendsDataView extends JobDataView {
             return false;
          }
 
-         if (r.outcome !== JobStepOutcome.Success && r.outcome !== JobStepOutcome.Warnings) {
-            return false;
-         }
-
          return true;
 
       });
@@ -127,7 +123,11 @@ class StepTrendsDataView extends JobDataView {
             this.maxMinutes = Math.max(minutes, this.maxMinutes);
          }
 
-         alltimes.push(minutes);
+         if (h.outcome === JobStepOutcome.Success || h.outcome === JobStepOutcome.Warnings) {
+            alltimes.push(minutes);
+         } else {
+            this.skipTrendLine.add(h.jobId);
+         }
 
          durations.set(h.jobId, minutes);
       });
@@ -150,7 +150,7 @@ class StepTrendsDataView extends JobDataView {
       this.median = median(alltimes);
       if (this.median) {
 
-         this.median *= 2.5;
+         this.median *= 3.0;
 
          if (this.median < 1) {
             this.median = 1;
@@ -243,6 +243,7 @@ class StepTrendsDataView extends JobDataView {
    order = 8;
 
    durations = new Map<string, number>();
+   skipTrendLine = new Set<string>();
    maxMinutes = 0;
    median?: number;
 
@@ -396,23 +397,31 @@ class StepTrendsRenderer {
          .attr("stroke-width", () => 1)
          .attr("stroke", () => dashboard.darktheme ? "#6D6C6B" : "#4D4C4B")
 
-      const lineI = d3.range(data.length);
+      const lineI: number[] = [];
+      data.forEach((d, idx) => {
+         if (!dataView.skipTrendLine.has(d.jobId)) {
+            lineI.push(idx)
+         }
+      })
 
       const showTrendLine = dataView.durations.size >= 10;
 
-
       const plotTrendY = (i: any, scaleY?: any) => {
          const idx = i as any as number;
 
          let sum = 0;
          let count = 0;
-         const range = 10;
+         const range = 6;
          for (let j = idx - range; j < idx + range; j++) {
 
             if (j < 0 || j >= data.length) {
                continue;
             }
 
+            if (dataView.skipTrendLine.has(data[j].jobId)) {
+               continue;
+            }
+
             const v = dataView.durations.get(data[j].jobId)!;
 
             if (dataView.median && v > dataView.median) {
