Horde: Fix bug where not all fields were included in Equals() check for AgentWorkspaceInfo

diff --git a/Engine/Source/Programs/Horde/Plugins/Compute/HordeServer.Compute/Agents/IAgent.cs b/Engine/Source/Programs/Horde/Plugins/Compute/HordeServer.Compute/Agents/IAgent.cs
--- a/Engine/Source/Programs/Horde/Plugins/Compute/HordeServer.Compute/Agents/IAgent.cs
+++ b/Engine/Source/Programs/Horde/Plugins/Compute/HordeServer.Compute/Agents/IAgent.cs
@@ -758,7 +758,14 @@ namespace HordeServer.Agents
 			{
 				return false;
 			}
-			if (Cluster != other.Cluster || UserName != other.UserName || Identifier != other.Identifier || Stream != other.Stream || Incremental != other.Incremental)
+			if (Cluster != other.Cluster || 
+				UserName != other.UserName ||
+				Identifier != other.Identifier ||
+				Stream != other.Stream || 
+				Incremental != other.Incremental ||
+				Method != other.Method ||
+				MinScratchSpace != other.MinScratchSpace ||
+				ConformDiskFreeSpace != other.ConformDiskFreeSpace)
 			{
 				return false;
 			}
