Horde: Fix incorrect trace name after refactor

diff --git a/Engine/Source/Programs/Horde/Plugins/Compute/HordeServer.Compute/Agents/AgentCollection.cs b/Engine/Source/Programs/Horde/Plugins/Compute/HordeServer.Compute/Agents/AgentCollection.cs
--- a/Engine/Source/Programs/Horde/Plugins/Compute/HordeServer.Compute/Agents/AgentCollection.cs
+++ b/Engine/Source/Programs/Horde/Plugins/Compute/HordeServer.Compute/Agents/AgentCollection.cs
@@ -318,7 +318,7 @@ namespace HordeServer.Agents
 
 		internal async ValueTask TickSharedAsync(CancellationToken stoppingToken)
 		{
-			using TelemetrySpan span = _tracer.StartActiveSpan($"{nameof(AgentService)}.{nameof(TickSharedAsync)}");
+			using TelemetrySpan span = _tracer.StartActiveSpan($"{nameof(AgentCollection)}.{nameof(TickSharedAsync)}");
 
 			await TerminateExpiredSessionsAsync(stoppingToken);
 			await DeleteExpiredEphemeralAgentsAsync(stoppingToken);
@@ -326,7 +326,7 @@ namespace HordeServer.Agents
 
 		private async Task TerminateExpiredSessionsAsync(CancellationToken cancellationToken)
 		{
-			using TelemetrySpan span = _tracer.StartActiveSpan($"{nameof(AgentService)}.{nameof(TerminateExpiredSessionsAsync)}");
+			using TelemetrySpan span = _tracer.StartActiveSpan($"{nameof(AgentCollection)}.{nameof(TerminateExpiredSessionsAsync)}");
 
 			int c = 0;
 			await foreach (RpcSession session in _scheduler.FindExpiredSessionsAsync(cancellationToken))
@@ -348,7 +348,7 @@ namespace HordeServer.Agents
 
 		private async Task DeleteExpiredEphemeralAgentsAsync(CancellationToken cancellationToken)
 		{
-			using TelemetrySpan span = _tracer.StartActiveSpan($"{nameof(AgentService)}.{nameof(DeleteExpiredEphemeralAgentsAsync)}");
+			using TelemetrySpan span = _tracer.StartActiveSpan($"{nameof(AgentCollection)}.{nameof(DeleteExpiredEphemeralAgentsAsync)}");
 			int numDeleted = 0;
 
 			List<AgentDocument> deletedDocuments = await _agentCollection.Find(x => x.Deleted).ToListAsync(cancellationToken);
