Horde: Allow specifying arbitary key/value metadata on tools, and add product ids for MSI installers. Unreal Toolbox will automatically install/uninstall MSIs using the given product id.

diff --git a/Engine/Source/Programs/Horde/BuildHorde.xml b/Engine/Source/Programs/Horde/BuildHorde.xml
--- a/Engine/Source/Programs/Horde/BuildHorde.xml
+++ b/Engine/Source/Programs/Horde/BuildHorde.xml
@@ -309,20 +309,19 @@
 				<!-- Add the bundled tools into the installer -->
 
 				<!-- General Tools -->
-				<DotNet Arguments="run --project &quot;$(RootDir)/Engine/Source/Programs/Horde/HordeServer/HordeServer.csproj&quot; -- generate tooldata -id=horde-agent -name=&quot;Horde Agent (Cross-Platform)&quot; -group=&quot;horde-agent&quot; -platforms=any -description=&quot;Cross-platform Horde Agent distribution. Allows hosts to take on work from the Horde Server.&quot; -version=&quot;$(InformationalVersion)&quot; -inputdir=&quot;$(ServerToolsStagingDir_HordeAgent)&quot; -serverdir=&quot;$(StagingDir)/Server/$(ServerArch)&quot; -public" DotNetPath="$(DotNetPath)" />				
-				<DotNet Arguments="run --project &quot;$(RootDir)/Engine/Source/Programs/Horde/HordeServer/HordeServer.csproj&quot; -- generate tooldata -id=horde-agent-msi -name=&quot;Horde Agent (Windows Installer)&quot; -group=&quot;horde-agent&quot; -platforms=win-x64 -description=&quot;Installer for the Horde Agent on Windows.&quot; -version=&quot;$(InformationalVersion)&quot; -inputdir=&quot;$(ServerToolsStagingDir_HordeAgentMsi)&quot; -serverdir=&quot;$(StagingDir)/Server/$(ServerArch)&quot; -public" DotNetPath="$(DotNetPath)" />				
-				<DotNet Arguments="run --project &quot;$(RootDir)/Engine/Source/Programs/Horde/HordeServer/HordeServer.csproj&quot; -- generate tooldata -id=unreal-toolbox-msi -name=&quot;Unreal Toolbox (Windows Installer)&quot; -description=&quot;Windows installer for the Unreal Toolbox.&quot; -version=&quot;$(InformationalVersion)&quot; -inputdir=&quot;$(ServerToolsStagingDir_UnrealToolboxMsi)&quot; -serverdir=&quot;$(StagingDir)/Server/$(ServerArch)&quot; -public -showinugs" DotNetPath="$(DotNetPath)" />				
+				<DotNet Arguments="run --project &quot;$(RootDir)/Engine/Source/Programs/Horde/HordeServer/HordeServer.csproj&quot; -- generate tooldata -id=horde-agent -name=&quot;Horde Agent (Cross-Platform)&quot; -group=&quot;horde-agent&quot; -platforms=any -description=&quot;Cross-platform Horde Agent distribution. Allows hosts to take on work from the Horde Server.&quot; -version=&quot;$(InformationalVersion)&quot; -inputdir=&quot;$(ServerToolsStagingDir_HordeAgent)&quot; -serverdir=&quot;$(StagingDir)/Server/$(ServerArch)&quot; -public" DotNetPath="$(DotNetPath)" />
+				<DotNet Arguments="run --project &quot;$(RootDir)/Engine/Source/Programs/Horde/HordeServer/HordeServer.csproj&quot; -- generate tooldata -id=horde-agent-msi -name=&quot;Horde Agent (Windows Installer)&quot; -group=&quot;horde-agent&quot; -platforms=win-x64 -description=&quot;Installer for the Horde Agent on Windows.&quot; -version=&quot;$(InformationalVersion)&quot; -inputdir=&quot;$(ServerToolsStagingDir_HordeAgentMsi)&quot; -serverdir=&quot;$(StagingDir)/Server/$(ServerArch)&quot; -public -showintoolbox -metadata=msi-product-id:1A853492-7FD0-4F88-B120-25292474A829" DotNetPath="$(DotNetPath)" />
+				<DotNet Arguments="run --project &quot;$(RootDir)/Engine/Source/Programs/Horde/HordeServer/HordeServer.csproj&quot; -- generate tooldata -id=unreal-toolbox-msi -name=&quot;Unreal Toolbox (Windows Installer)&quot; -description=&quot;Windows installer for the Unreal Toolbox.&quot; -version=&quot;$(InformationalVersion)&quot; -inputdir=&quot;$(ServerToolsStagingDir_UnrealToolboxMsi)&quot; -serverdir=&quot;$(StagingDir)/Server/$(ServerArch)&quot; -public -showinugs" DotNetPath="$(DotNetPath)" />
 				<DotNet Arguments="run --project &quot;$(RootDir)/Engine/Source/Programs/Horde/HordeServer/HordeServer.csproj&quot; -- generate tooldata -id=ugs-cmd -name=&quot;UnrealGameSync (Command-Line)&quot; -group=&quot;ugs&quot; -platforms=any -description=&quot;Command-line version of UnrealGameSync. Supports Windows, MacOS and Linux.&quot; -version=&quot;$(InformationalVersion)&quot; -inputdir=&quot;$(ServerToolsStagingDir_UnrealGameSyncCmd)&quot; -serverdir=&quot;$(StagingDir)/Server/$(ServerArch)&quot; -public -showintoolbox" DotNetPath="$(DotNetPath)" />
-				<DotNet Arguments="run --project &quot;$(RootDir)/Engine/Source/Programs/Horde/HordeServer/HordeServer.csproj&quot; -- generate tooldata -id=ugs-msi -name=&quot;UnrealGameSync (Windows Installer)&quot; -group=&quot;ugs&quot; -platforms=win-x64 -description=&quot;Installer for UnrealGameSync.&quot; -version=&quot;$(InformationalVersion)&quot; -inputdir=&quot;$(ServerToolsStagingDir_UnrealGameSyncMsi)&quot; -serverdir=&quot;$(StagingDir)/Server/$(ServerArch)&quot; -public -showintoolbox" DotNetPath="$(DotNetPath)" />
+				<DotNet Arguments="run --project &quot;$(RootDir)/Engine/Source/Programs/Horde/HordeServer/HordeServer.csproj&quot; -- generate tooldata -id=ugs-msi -name=&quot;UnrealGameSync (Windows Installer)&quot; -group=&quot;ugs&quot; -platforms=win-x64 -description=&quot;Utility to allow syncing and building from Perforce.&quot; -version=&quot;$(InformationalVersion)&quot; -inputdir=&quot;$(ServerToolsStagingDir_UnrealGameSyncMsi)&quot; -serverdir=&quot;$(StagingDir)/Server/$(ServerArch)&quot; -public -showintoolbox -metadata=msi-product-id:6465D644-E775-44B5-B978-7B9EF8F0EF33" DotNetPath="$(DotNetPath)" />
 
 				<!-- Internal Category -->
 				<DotNet Arguments="run --project &quot;$(RootDir)/Engine/Source/Programs/Horde/HordeServer/HordeServer.csproj&quot; -- generate tooldata -id=horde-cmd -name=&quot;Horde (Command-Line)&quot; -Category=Internal -description=&quot;Update package for the Horde command-line tool.&quot; -version=&quot;$(InformationalVersion)&quot; -inputdir=&quot;$(ServerToolsStagingDir_HordeCmd)&quot; -serverdir=&quot;$(StagingDir)/Server/$(ServerArch)&quot; -public -showinugs -showintoolbox" DotNetPath="$(DotNetPath)" />
 				<DotNet Arguments="run --project &quot;$(RootDir)/Engine/Source/Programs/Horde/HordeServer/HordeServer.csproj&quot; -- generate tooldata -id=horde-agent-win64 -name=&quot;Horde Agent (Windows)&quot; -category=Internal -description=&quot;Update package for older installations of Horde Agent through the MSI installer.&quot; -version=&quot;$(InformationalVersion)&quot; -inputdir=&quot;$(ServerToolsStagingDir_HordeAgentWin64)&quot; -serverdir=&quot;$(StagingDir)/Server/$(ServerArch)&quot; -public -showindashboard=false" DotNetPath="$(DotNetPath)" />
-				<DotNet Arguments="run --project &quot;$(RootDir)/Engine/Source/Programs/Horde/HordeServer/HordeServer.csproj&quot; -- generate tooldata -id=horde-agent-msi-toolbox -name=&quot;Horde Agent&quot; -description=&quot;Installer for the Horde Agent on Windows, for Unreal Toolbox.&quot; -version=&quot;$(InformationalVersion)&quot; -inputdir=&quot;$(ServerToolsStagingDir_HordeAgentMsi)&quot; -serverdir=&quot;$(StagingDir)/Server/$(ServerArch)&quot; -public -showintoolbox" DotNetPath="$(DotNetPath)" />
 				<DotNet Arguments="run --project &quot;$(RootDir)/Engine/Source/Programs/Horde/HordeServer/HordeServer.csproj&quot; -- generate tooldata -id=unreal-toolbox -name=&quot;Unreal Toolbox&quot; -Category=Internal -description=&quot;Update package for the Unreal Toolbox.&quot; -version=&quot;$(InformationalVersion)&quot; -inputdir=&quot;$(ServerToolsStagingDir_UnrealToolbox)&quot; -serverdir=&quot;$(StagingDir)/Server/$(ServerArch)&quot; -public -showinugs" DotNetPath="$(DotNetPath)" />
 				<DotNet Arguments="run --project &quot;$(RootDir)/Engine/Source/Programs/Horde/HordeServer/HordeServer.csproj&quot; -- generate tooldata -id=p4v-utils -name=&quot;P4V Extensions&quot; -category=Internal -description=&quot;Update package for the P4VUtils tool.&quot; -version=&quot;$(InformationalVersion)&quot; -inputdir=&quot;$(ServerToolsStagingDir_P4VUtils)&quot; -serverdir=&quot;$(StagingDir)/Server/$(ServerArch)&quot; -public -showinugs -showintoolbox" DotNetPath="$(DotNetPath)" />
 				<DotNet Arguments="run --project &quot;$(RootDir)/Engine/Source/Programs/Horde/HordeServer/HordeServer.csproj&quot; -- generate tooldata -id=ugs-win -name=&quot;UnrealGameSync (Windows Update Package)&quot; -category=Internal -description=&quot;Update package for UGS on Windows.&quot; -version=&quot;$(InformationalVersion)&quot; -inputdir=&quot;$(ServerToolsStagingDir_UnrealGameSync)&quot; -serverdir=&quot;$(StagingDir)/Server/$(ServerArch)&quot; -public -showindashboard=false" DotNetPath="$(DotNetPath)" />
-				<DotNet Arguments="run --project &quot;$(RootDir)/Engine/Source/Programs/Horde/HordeServer/HordeServer.csproj&quot; -- generate tooldata -id=ushell -name=&quot;Ushell&quot; -Category=Internal -description=&quot;Update package for ushell.&quot; -version=&quot;$(InformationalVersion)&quot; -inputdir=&quot;$(RootDir)/Engine/Extras/ushell&quot; -serverdir=&quot;$(StagingDir)/Server/$(ServerArch)&quot; -public -showinugs -showintoolbox" DotNetPath="$(DotNetPath)" />			
+				<DotNet Arguments="run --project &quot;$(RootDir)/Engine/Source/Programs/Horde/HordeServer/HordeServer.csproj&quot; -- generate tooldata -id=ushell -name=&quot;Ushell&quot; -Category=Internal -description=&quot;Update package for ushell.&quot; -version=&quot;$(InformationalVersion)&quot; -inputdir=&quot;$(RootDir)/Engine/Extras/ushell&quot; -serverdir=&quot;$(StagingDir)/Server/$(ServerArch)&quot; -public -showinugs -showintoolbox" DotNetPath="$(DotNetPath)" />
 
 				<!-- Separate copy of the command line tool for installation with the server on Windows -->
 				<Delete Files="$(StagingDir)/Tools/..."/>
diff --git a/Engine/Source/Programs/Horde/Plugins/Tools/HordeServer.Tools/Commands/Generate/ToolDataCommand.cs b/Engine/Source/Programs/Horde/Plugins/Tools/HordeServer.Tools/Commands/Generate/ToolDataCommand.cs
--- a/Engine/Source/Programs/Horde/Plugins/Tools/HordeServer.Tools/Commands/Generate/ToolDataCommand.cs
+++ b/Engine/Source/Programs/Horde/Plugins/Tools/HordeServer.Tools/Commands/Generate/ToolDataCommand.cs
@@ -62,6 +62,10 @@ namespace HordeServer.Commands.Generate
 		[Description("Shows the tool for download on the dashboard")]
 		public bool ShowInDashboard { get; set; } = true;
 
+		[CommandLine(ListSeparator = ';')]
+		[Description("Semicolon delimited list of key=value metadata pairs")]
+		public List<string> Metadata { get; set; } = new List<string>();
+
 		[CommandLine(Required = true)]
 		[Description("Source directory for tool data")]
 		public DirectoryReference InputDir { get; set; } = null!;
@@ -146,6 +150,21 @@ namespace HordeServer.Commands.Generate
 					bundledTool[nameof(BundledToolConfig.Platforms)] = platforms;
 				}
 
+				if (Metadata.Count > 0)
+				{
+					JsonObject metadata = new JsonObject();
+					foreach (string metadataItem in Metadata)
+					{
+						int equalsIdx = metadataItem.IndexOf(':', StringComparison.Ordinal);
+						if (equalsIdx == -1)
+						{
+							throw new InvalidOperationException($"Missing colon symbol in {metadataItem}");
+						}
+						metadata.Add(metadataItem.Substring(0, equalsIdx), metadataItem.Substring(equalsIdx + 1));
+					}
+					bundledTool[nameof(BundledToolConfig.Metadata)] = metadata;
+				}
+
 				bundledTool[nameof(BundledToolConfig.RefName)] = refName.ToString();
 
 				if (Public)
diff --git a/Engine/Source/Programs/Horde/Plugins/Tools/HordeServer.Tools/Tools/ToolCollection.cs b/Engine/Source/Programs/Horde/Plugins/Tools/HordeServer.Tools/Tools/ToolCollection.cs
--- a/Engine/Source/Programs/Horde/Plugins/Tools/HordeServer.Tools/Tools/ToolCollection.cs
+++ b/Engine/Source/Programs/Horde/Plugins/Tools/HordeServer.Tools/Tools/ToolCollection.cs
@@ -105,6 +105,7 @@ namespace HordeServer.Tools
 			public bool ShowInUgs => _config.ShowInUgs;
 			public bool ShowInDashboard => _config.ShowInDashboard;
 			public bool ShowInToolbox => _config.ShowInToolbox;
+			public IReadOnlyDictionary<string, string> Metadata => _config.Metadata;
 			public IReadOnlyList<IToolDeployment> Deployments => _deployments;
 
 			public Tool(ToolCollection collection, ToolDocument document, ToolConfig config, DateTime utcNow)
diff --git a/Engine/Source/Programs/Horde/Plugins/Tools/HordeServer.Tools/Tools/ToolConfig.cs b/Engine/Source/Programs/Horde/Plugins/Tools/HordeServer.Tools/Tools/ToolConfig.cs
--- a/Engine/Source/Programs/Horde/Plugins/Tools/HordeServer.Tools/Tools/ToolConfig.cs
+++ b/Engine/Source/Programs/Horde/Plugins/Tools/HordeServer.Tools/Tools/ToolConfig.cs
@@ -7,6 +7,7 @@ using EpicGames.Horde.Acls;
 using EpicGames.Horde.Storage;
 using EpicGames.Horde.Tools;
 using HordeServer.Acls;
+using HordeServer.Utilities;
 
 #pragma warning disable CA2227 // Change x to be read-only by removing the property setter
 
@@ -75,6 +76,11 @@ namespace HordeServer.Tools
 		/// </summary>
 		public bool ShowInToolbox { get; set; }
 
+		/// <summary>
+		/// Metadata for this tool
+		/// </summary>
+		public CaseInsensitiveDictionary<string> Metadata { get; set; } = new CaseInsensitiveDictionary<string>();
+
 		/// <summary>
 		/// Default namespace for new deployments of this tool
 		/// </summary>
diff --git a/Engine/Source/Programs/Horde/Plugins/Tools/HordeServer.Tools/Tools/ToolsController.cs b/Engine/Source/Programs/Horde/Plugins/Tools/HordeServer.Tools/Tools/ToolsController.cs
--- a/Engine/Source/Programs/Horde/Plugins/Tools/HordeServer.Tools/Tools/ToolsController.cs
+++ b/Engine/Source/Programs/Horde/Plugins/Tools/HordeServer.Tools/Tools/ToolsController.cs
@@ -204,7 +204,7 @@ namespace HordeServer.Tools
 		static GetToolSummaryResponse CreateGetToolSummaryResponse(ITool tool)
 		{
 			IToolDeployment? deployment = (tool.Deployments.Count == 0) ? null : tool.Deployments[^1];
-			return new GetToolSummaryResponse(tool.Id, tool.Name, tool.Description, tool.Category, tool.Group, tool.Platforms?.ToList(), deployment?.Version, deployment?.Id, deployment?.State, deployment?.Progress, tool.Bundled, tool.ShowInUgs, tool.ShowInDashboard, tool.ShowInToolbox);
+			return new GetToolSummaryResponse(tool.Id, tool.Name, tool.Description, tool.Category, tool.Group, tool.Platforms?.ToList(), deployment?.Version, deployment?.Id, deployment?.State, deployment?.Progress, tool.Bundled, tool.ShowInUgs, tool.ShowInDashboard, tool.ShowInToolbox, new Dictionary<string, string>(tool.Metadata, StringComparer.OrdinalIgnoreCase));
 		}
 
 		/// <summary>
@@ -248,7 +248,7 @@ namespace HordeServer.Tools
 
 		static GetToolResponse CreateGetToolResponse(ITool tool, List<GetToolDeploymentResponse> deployments)
 		{
-			return new GetToolResponse(tool.Id, tool.Name, tool.Description, tool.Category, tool.Group, tool.Platforms?.ToList(), deployments, tool.Public, tool.Bundled, tool.ShowInUgs, tool.ShowInDashboard, tool.ShowInToolbox);
+			return new GetToolResponse(tool.Id, tool.Name, tool.Description, tool.Category, tool.Group, tool.Platforms?.ToList(), deployments, tool.Public, tool.Bundled, tool.ShowInUgs, tool.ShowInDashboard, tool.ShowInToolbox, new Dictionary<string, string>(tool.Metadata, StringComparer.OrdinalIgnoreCase));
 		}
 
 		/// <summary>
diff --git a/Engine/Source/Programs/Shared/EpicGames.Horde/Tools/HttpToolCollection.cs b/Engine/Source/Programs/Shared/EpicGames.Horde/Tools/HttpToolCollection.cs
--- a/Engine/Source/Programs/Shared/EpicGames.Horde/Tools/HttpToolCollection.cs
+++ b/Engine/Source/Programs/Shared/EpicGames.Horde/Tools/HttpToolCollection.cs
@@ -18,6 +18,8 @@ namespace EpicGames.Horde.Tools
 		[DebuggerDisplay("{Id}")]
 		class Tool : ITool
 		{
+			static readonly IReadOnlyDictionary<string, string> s_defaultMetadata = new Dictionary<string, string>();
+
 			readonly GetToolResponse _response;
 			readonly List<ToolDeployment> _deployments;
 
@@ -33,6 +35,7 @@ namespace EpicGames.Horde.Tools
 			bool ITool.ShowInUgs => _response.ShowInUgs;
 			bool ITool.ShowInDashboard => _response.ShowInDashboard;
 			bool ITool.ShowInToolbox => _response.ShowInToolbox;
+			IReadOnlyDictionary<string, string> ITool.Metadata => _response.Metadata ?? s_defaultMetadata;
 			IReadOnlyList<IToolDeployment> ITool.Deployments => _deployments;
 
 			public Tool(HttpToolCollection collection, GetToolResponse response)
diff --git a/Engine/Source/Programs/Shared/EpicGames.Horde/Tools/ITool.cs b/Engine/Source/Programs/Shared/EpicGames.Horde/Tools/ITool.cs
--- a/Engine/Source/Programs/Shared/EpicGames.Horde/Tools/ITool.cs
+++ b/Engine/Source/Programs/Shared/EpicGames.Horde/Tools/ITool.cs
@@ -75,6 +75,11 @@ namespace EpicGames.Horde.Tools
 		/// </summary>
 		bool ShowInToolbox { get; }
 
+		/// <summary>
+		/// User-defined metadata for this tool
+		/// </summary>
+		IReadOnlyDictionary<string, string> Metadata { get; }
+
 		/// <summary>
 		/// Current deployments of this tool, sorted by time.
 		/// </summary>
diff --git a/Engine/Source/Programs/Shared/EpicGames.Horde/Tools/ToolMessages.cs b/Engine/Source/Programs/Shared/EpicGames.Horde/Tools/ToolMessages.cs
--- a/Engine/Source/Programs/Shared/EpicGames.Horde/Tools/ToolMessages.cs
+++ b/Engine/Source/Programs/Shared/EpicGames.Horde/Tools/ToolMessages.cs
@@ -22,7 +22,8 @@ namespace EpicGames.Horde.Tools
 	/// <param name="ShowInUgs">Whether to show this tool for download inside UGS</param>
 	/// <param name="ShowInDashboard">Whether to show this tool for download on the dashboard</param>
 	/// <param name="ShowInToolbox">Whether to show this tool for download in Unreal Toolbox</param>
-	public record class GetToolResponse(ToolId Id, string Name, string Description, string? Category, string? Group, List<string>? Platforms, List<GetToolDeploymentResponse> Deployments, bool Public, bool Bundled, bool ShowInUgs, bool ShowInDashboard, bool ShowInToolbox);
+	/// <param name="Metadata">Metadata for the tool</param>
+	public record class GetToolResponse(ToolId Id, string Name, string Description, string? Category, string? Group, List<string>? Platforms, List<GetToolDeploymentResponse> Deployments, bool Public, bool Bundled, bool ShowInUgs, bool ShowInDashboard, bool ShowInToolbox, Dictionary<string, string>? Metadata = null);
 
 	/// <summary>
 	/// Summary for a particular tool.
@@ -41,7 +42,8 @@ namespace EpicGames.Horde.Tools
 	/// <param name="ShowInUgs">Whether to show this tool for download inside UGS</param>
 	/// <param name="ShowInDashboard">Whether to show this tool for download on the dashboard</param>
 	/// <param name="ShowInToolbox">Whether to show this tool for download in the launcher</param>
-	public record class GetToolSummaryResponse(ToolId Id, string Name, string Description, string? Category, string? Group, List<string>? Platforms, string? Version, ToolDeploymentId? DeploymentId, ToolDeploymentState? DeploymentState, double? DeploymentProgress, bool Bundled, bool ShowInUgs, bool ShowInDashboard, bool ShowInToolbox);
+	/// <param name="Metadata">Metadata for the tool</param>
+	public record class GetToolSummaryResponse(ToolId Id, string Name, string Description, string? Category, string? Group, List<string>? Platforms, string? Version, ToolDeploymentId? DeploymentId, ToolDeploymentState? DeploymentState, double? DeploymentProgress, bool Bundled, bool ShowInUgs, bool ShowInDashboard, bool ShowInToolbox, Dictionary<string, string>? Metadata = null);
 
 	/// <summary>
 	/// Response when querying all tools
diff --git a/Engine/Source/Programs/UnrealToolbox/App.axaml.cs b/Engine/Source/Programs/UnrealToolbox/App.axaml.cs
--- a/Engine/Source/Programs/UnrealToolbox/App.axaml.cs
+++ b/Engine/Source/Programs/UnrealToolbox/App.axaml.cs
@@ -55,6 +55,7 @@ namespace UnrealToolbox
 			_serviceProvider = serviceCollection.BuildServiceProvider();
 		}
 
+		/// <inheritdoc/>
 		public async ValueTask DisposeAsync()
 		{
 			if (_settingsThreadStop != null)
diff --git a/Engine/Source/Programs/UnrealToolbox/ToolCatalog.cs b/Engine/Source/Programs/UnrealToolbox/ToolCatalog.cs
--- a/Engine/Source/Programs/UnrealToolbox/ToolCatalog.cs
+++ b/Engine/Source/Programs/UnrealToolbox/ToolCatalog.cs
@@ -1,13 +1,14 @@
 // Copyright Epic Games, Inc. All Rights Reserved.
 
-using System.Diagnostics;
-using System.Runtime.InteropServices;
-using System.Text.Json;
 using EpicGames.Core;
 using EpicGames.Horde;
 using EpicGames.Horde.Storage.Nodes;
 using EpicGames.Horde.Tools;
 using Microsoft.Extensions.Logging;
+using Microsoft.Win32;
+using System.Diagnostics;
+using System.Runtime.InteropServices;
+using System.Text.Json;
 
 #nullable enable
 
@@ -15,7 +16,7 @@ namespace UnrealToolbox
 {
 	class ToolCatalog : IToolCatalog, IAsyncDisposable
 	{
-		record class ItemState(string Name, string Description, CurrentToolDeploymentInfo? Current, PendingToolDeploymentInfo? Pending, IToolDeployment? Latest);
+		record class ItemState(string Name, string Description, string? MsiProductId, CurrentToolDeploymentInfo? Current, PendingToolDeploymentInfo? Pending, IToolDeployment? Latest);
 
 		[DebuggerDisplay("{Id}")]
 		class Item : IToolCatalogItem
@@ -150,6 +151,7 @@ namespace UnrealToolbox
 			public ToolId Id { get; set; }
 			public string Name { get; set; }
 			public string Description { get; set; }
+			public string? MsiProductId { get; set; }
 			public JsonDeploymentState? Current { get; set; }
 
 			public JsonItemState()
@@ -171,6 +173,7 @@ namespace UnrealToolbox
 		{
 			public ToolDeploymentId Id { get; set; }
 			public string Version { get; set; }
+			public string? MsiProductId { get; set; }
 
 			public JsonDeploymentState()
 			{
@@ -389,7 +392,7 @@ namespace UnrealToolbox
 							current = new CurrentToolDeploymentInfo(jsonItem.Current.Id, jsonItem.Current.Version, toolDir, toolConfig);
 						}
 
-						ItemState itemState = new ItemState(jsonItem.Name, jsonItem.Description, current, null, null);
+						ItemState itemState = new ItemState(jsonItem.Name, jsonItem.Description, jsonItem.MsiProductId, current, null, null);
 						items.Add(jsonItem.Id, new Item(this, jsonItem.Id, itemState));
 					}
 				}
@@ -500,38 +503,81 @@ namespace UnrealToolbox
 				int exitCode = await RunCommandAsync(item.Id.ToString(), toolConfig.InstallCommand, toolDir.FullName, cancellationToken);
 				if (exitCode != 0)
 				{
-					pending = new PendingToolDeploymentInfo(true, $"Installation failed ({exitCode}). See log.", deploymentInfo);
+					pending = new PendingToolDeploymentInfo(true, $"Installation failed ({exitCode}).", deploymentInfo);
 					UpdateStateAndNotify(item, state => state with { Pending = pending });
 					return;
 				}
 			}
+			else if (OperatingSystem.IsWindows() && !String.IsNullOrEmpty(item._state.MsiProductId))
+			{
+				PendingToolDeploymentInfo pending = new PendingToolDeploymentInfo(false, "Running installer...", deploymentInfo);
+				UpdateStateAndNotify(item, state => state with { Pending = pending });
+
+				FileReference? msiFile = DirectoryReference.EnumerateFiles(toolDir, "*.msi").FirstOrDefault();
+				if (msiFile == null)
+				{
+					pending = new PendingToolDeploymentInfo(true, $"MSI file not found", deploymentInfo);
+					UpdateStateAndNotify(item, state => state with { Pending = pending });
+					return;
+				}
+
+				ToolCommand installCommand = new ToolCommand { FileName = "msiexec.exe" };
+				installCommand.Arguments = new List<string> { "/I", msiFile.FullName };
+
+				int exitCode = await RunCommandAsync(item.Id.ToString(), installCommand, toolDir.FullName, cancellationToken);
+				if (exitCode != 0 || !IsMsiInstalled(item._state.MsiProductId))
+				{
+					UpdateStateAndNotify(item, state => state with { Pending = null });
+					return;
+				}
+			}
 
 			CurrentToolDeploymentInfo current = new CurrentToolDeploymentInfo(deploymentInfo.Id, deploymentInfo.Version, toolDir, toolConfig);
 			UpdateStateAndNotify(item, state => state with { Pending = null, Current = current });
 		}
 
+		static bool IsMsiInstalled(string msiProductId)
+		{
+			string? value =
+				(Registry.GetValue($"HKEY_LOCAL_MACHINE\\SOFTWARE\\Microsoft\\Windows\\CurrentVersion\\Uninstall\\{{{msiProductId}}}", "UninstallString", null) as string) ??
+				(Registry.GetValue($"HKEY_LOCAL_MACHINE\\SOFTWARE\\WOW6432Node\\Microsoft\\Windows\\CurrentVersion\\Uninstall\\{{{msiProductId}}}", "UninstallString", null) as string);
+
+			return !String.IsNullOrEmpty(value);
+		}
+
 		async Task UninstallAsync(Item item, CancellationToken cancellationToken)
 		{
-			CurrentToolDeploymentInfo? current = null;
-			UpdateStateAndNotify(item, state =>
+			CurrentToolDeploymentInfo? current = item._state.Current;
+			if (current != null)
 			{
-				current = state.Current;
-
-				PendingToolDeploymentInfo? pending = null;
-				if (current != null)
+				if (current.Config.UninstallCommand != null)
 				{
-					pending = new PendingToolDeploymentInfo(false, $"Removing {current.Version}", null);
+					PendingToolDeploymentInfo? pending = new PendingToolDeploymentInfo(false, $"Removing {current.Version}", null);
+					UpdateStateAndNotify(item, state => state with { Pending = pending, Current = null });
+
+					await RunCommandAsync(current.Id.ToString(), current.Config.UninstallCommand, current.Dir.FullName, cancellationToken);
 				}
+				else if (OperatingSystem.IsWindows() && !String.IsNullOrEmpty(item._state.MsiProductId) && IsMsiInstalled(item._state.MsiProductId))
+				{
+					PendingToolDeploymentInfo pending = new PendingToolDeploymentInfo(false, $"Running installer...", current);
+					UpdateStateAndNotify(item, state => state with { Pending = pending });
 
-				return state with { Pending = pending, Current = null };
-			});
+					ToolCommand uninstallCommand = new ToolCommand { FileName = "msiexec.exe" };
+					uninstallCommand.Arguments = new List<string> { "/x", $"{{{item._state.MsiProductId}}}" };
 
-			if (current != null)
-			{
-				ToolCommand? uninstallCommand = current.Config.UninstallCommand;
-				if (uninstallCommand != null)
-				{
-					await RunCommandAsync(current.Id.ToString(), uninstallCommand, current.Dir.FullName, cancellationToken);
+					int exitCode = await RunCommandAsync(current.Id.ToString(), uninstallCommand, current.Dir.FullName, cancellationToken);
+					if (exitCode == 1602)
+					{
+						// User cancelled (https://learn.microsoft.com/en-us/windows/win32/msi/error-codes)
+						UpdateStateAndNotify(item, state => state with { Pending = null });
+						return;
+					}
+					else if (exitCode != 0 || IsMsiInstalled(item._state.MsiProductId))
+					{
+						pending = new PendingToolDeploymentInfo(true, $"Uninstall failed ({exitCode}).", current);
+						UpdateStateAndNotify(item, state => state with { Pending = pending });
+						return;
+					}
 				}
 			}
 
@@ -793,15 +839,27 @@ namespace UnrealToolbox
 				{
 					if (ShouldIncludeTool(tool))
 					{
+						string? msiProductId;
+						if (!tool.Metadata.TryGetValue("msi-product-id", out msiProductId))
+						{
+							msiProductId = null;
+						}
+
 						IToolDeployment? latestDeployment = (tool.Deployments.Count > 0) ? tool.Deployments[^1] : null;
 						if (_items.TryGetValue(tool.Id, out Item? item))
 						{
-							item._state = new ItemState(tool.Name, tool.Description, item._state.Current, item._state.Pending, latestDeployment);
+							item._state = new ItemState(tool.Name, tool.Description, msiProductId, item._state.Current, item._state.Pending, latestDeployment);
 						}
 						else
 						{
-							item = new Item(this, tool.Id, new ItemState(tool.Name, tool.Description, null, null, latestDeployment));
+							item = new Item(this, tool.Id, new ItemState(tool.Name, tool.Description, msiProductId, null, null, latestDeployment));
 						}
+
+						if (OperatingSystem.IsWindows() && item._state.Current != null && msiProductId != null && !IsMsiInstalled(msiProductId))
+						{
+							item._state = item._state with { Current = null };
+						}
+
 						newItems[item.Id] = item;
 					}
 				}
diff --git a/Engine/Source/Programs/UnrealToolbox/ToolCatalogViewModel.cs b/Engine/Source/Programs/UnrealToolbox/ToolCatalogViewModel.cs
--- a/Engine/Source/Programs/UnrealToolbox/ToolCatalogViewModel.cs
+++ b/Engine/Source/Programs/UnrealToolbox/ToolCatalogViewModel.cs
@@ -144,7 +144,7 @@ namespace UnrealToolbox
 				{
 					Status = $"Updating to {pending.Deployment.Version}...";
 				}
-				IsBusy = true;
+				IsBusy = !pending.Failed;
 			}
 			else
 			{
