Horde: Add more logging and tracing for updating agents via REST API

diff --git a/Engine/Source/Programs/Horde/Plugins/Compute/HordeServer.Compute/Agents/AgentsController.cs b/Engine/Source/Programs/Horde/Plugins/Compute/HordeServer.Compute/Agents/AgentsController.cs
--- a/Engine/Source/Programs/Horde/Plugins/Compute/HordeServer.Compute/Agents/AgentsController.cs
+++ b/Engine/Source/Programs/Horde/Plugins/Compute/HordeServer.Compute/Agents/AgentsController.cs
@@ -254,12 +254,16 @@ namespace HordeServer.Agents
 		[Route("/api/v1/agents/{agentId}")]
 		public async Task<ActionResult> UpdateAgentAsync(AgentId agentId, [FromBody] UpdateAgentRequest update)
 		{
+			using TelemetrySpan span = _tracer.StartActiveSpan($"{nameof(AgentsController)}.{nameof(UpdateAgentAsync)}");
+			
 			if (!_computeConfig.Value.Authorize(AgentAclAction.UpdateAgent, User))
 			{
 				return Forbid(AgentAclAction.UpdateAgent, agentId);
 			}
 
 			string userName = User.GetUser() ?? "Unknown";
+			string? clientIp = HttpContext.Connection.RemoteIpAddress?.ToString();
+			span.SetAttribute("claims", User.Claims.Select(x => x.ToString()).ToArray());
 
 			for (; ; )
 			{
@@ -278,35 +282,35 @@ namespace HordeServer.Agents
 				IAuditLogChannel<AgentId> logger = _agentService.Agents.GetLogger(agent.Id);
 				if (agent.Enabled != newAgent.Enabled)
 				{
-					logger.LogInformation("Setting changed: Enabled = {State} ({UserName})", newAgent.Enabled, userName);
+					logger.LogInformation("Setting changed: Enabled = {State} ({UserName} via IP {Ip})", newAgent.Enabled, userName, clientIp);
 				}
 				if (agent.RequestConform != newAgent.RequestConform)
 				{
-					logger.LogInformation("Setting changed: RequestConform = {State} ({UserName})", newAgent.RequestConform, userName);
+					logger.LogInformation("Setting changed: RequestConform = {State} ({UserName} via IP {Ip})", newAgent.RequestConform, userName, clientIp);
 				}
 				if (agent.RequestFullConform != newAgent.RequestFullConform)
 				{
-					logger.LogInformation("Setting changed: RequestFullConform = {State} ({UserName})", newAgent.RequestFullConform, userName);
+					logger.LogInformation("Setting changed: RequestFullConform = {State} ({UserName} via IP {Ip})", newAgent.RequestFullConform, userName, clientIp);
 				}
 				if (agent.RequestRestart != newAgent.RequestRestart)
 				{
-					logger.LogInformation("Setting changed: RequestRestart = {State} ({UserName})", newAgent.RequestRestart, userName);
+					logger.LogInformation("Setting changed: RequestRestart = {State} ({UserName} via IP {Ip})", newAgent.RequestRestart, userName, clientIp);
 				}
 				if (agent.RequestShutdown != newAgent.RequestShutdown)
 				{
-					logger.LogInformation("Setting changed: RequestShutdown = {State} ({UserName})", newAgent.RequestShutdown, userName);
+					logger.LogInformation("Setting changed: RequestShutdown = {State} ({UserName} via IP {Ip})", newAgent.RequestShutdown, userName, clientIp);
 				}
 				if (agent.Comment != newAgent.Comment)
 				{
-					logger.LogInformation("Setting changed: Comment = \"{Comment}\" ({UserName})", update.Comment, userName);
+					logger.LogInformation("Setting changed: Comment = \"{Comment}\" ({UserName} via IP {Ip})", update.Comment, userName, clientIp);
 				}
 				foreach (PoolId addedPool in newAgent.ExplicitPools.Except(agent.ExplicitPools))
 				{
-					logger.LogInformation("Added to pool {PoolId} ({UserName})", addedPool, userName);
+					logger.LogInformation("Added to pool {PoolId} ({UserName} via IP {Ip})", addedPool, userName, clientIp);
 				}
 				foreach (PoolId removedPool in agent.ExplicitPools.Except(newAgent.ExplicitPools))
 				{
-					logger.LogInformation("Removed from pool {PoolId} ({UserName})", removedPool, userName);
+					logger.LogInformation("Removed from pool {PoolId} ({UserName} via IP {Ip})", removedPool, userName, clientIp);
 				}
 				break;
 			}
