Horde: Add logging for workspace differences.

#rnx

diff --git a/Engine/Source/Programs/Horde/Plugins/Build/HordeServer.Build/Jobs/JobRpcService.cs b/Engine/Source/Programs/Horde/Plugins/Build/HordeServer.Build/Jobs/JobRpcService.cs
--- a/Engine/Source/Programs/Horde/Plugins/Build/HordeServer.Build/Jobs/JobRpcService.cs
+++ b/Engine/Source/Programs/Horde/Plugins/Build/HordeServer.Build/Jobs/JobRpcService.cs
@@ -15,6 +15,7 @@ using EpicGames.Horde.Logs;
 using Google.Protobuf.WellKnownTypes;
 using Grpc.Core;
 using Horde.Common.Rpc;
+using HordeCommon.Rpc.Messages;
 using HordeServer.Acls;
 using HordeServer.Agents;
 using HordeServer.Agents.Pools;
@@ -109,6 +110,19 @@ namespace HordeServer.Jobs
 					{
 						response.Retry = await _conformTaskSource.GetWorkspacesAsync(agent, response.PendingWorkspaces, context.CancellationToken);
 						response.RemoveUntrackedFiles = request.RemoveUntrackedFiles || agent.RequestFullConform;
+
+						if (response.Retry)
+						{
+							_logger.LogInformation("Retrying conform for {AgentId}", agent.Id);
+							foreach (RpcAgentWorkspace workspace in request.Workspaces)
+							{
+								_logger.LogInformation("{AgentId} has workspace {Workspace}", agent.Id, workspace.ToJson());
+							}
+							foreach (RpcAgentWorkspace workspace in response.PendingWorkspaces)
+							{
+								_logger.LogInformation("{AgentId} needs workspace {Workspace}", agent.Id, workspace.ToJson());
+							}
+						}
 					}
 					return response;
 				}
