Horde: Additional logging for conform retries.

#rnx

diff --git a/Engine/Source/Programs/Horde/Plugins/Build/HordeServer.Build/Jobs/JobRpcService.cs b/Engine/Source/Programs/Horde/Plugins/Build/HordeServer.Build/Jobs/JobRpcService.cs
--- a/Engine/Source/Programs/Horde/Plugins/Build/HordeServer.Build/Jobs/JobRpcService.cs
+++ b/Engine/Source/Programs/Horde/Plugins/Build/HordeServer.Build/Jobs/JobRpcService.cs
@@ -15,7 +15,6 @@ using EpicGames.Horde.Logs;
 using Google.Protobuf.WellKnownTypes;
 using Grpc.Core;
 using Horde.Common.Rpc;
-using HordeCommon.Rpc.Messages;
 using HordeServer.Acls;
 using HordeServer.Agents;
 using HordeServer.Agents.Pools;
@@ -114,13 +113,13 @@ namespace HordeServer.Jobs
 						if (response.Retry)
 						{
 							_logger.LogInformation("Retrying conform for {AgentId}", agent.Id);
-							foreach (RpcAgentWorkspace workspace in request.Workspaces)
+							foreach (AgentWorkspaceInfo workspace in newWorkspaces)
 							{
-								_logger.LogInformation("{AgentId} has workspace {Workspace}", agent.Id, workspace.ToJson());
+								_logger.LogInformation("{AgentId} has workspace {Workspace}", agent.Id, JsonSerializer.Serialize(workspace));
 							}
-							foreach (RpcAgentWorkspace workspace in response.PendingWorkspaces)
+							foreach (AgentWorkspaceInfo workspace in conformWorkspaces)
 							{
-								_logger.LogInformation("{AgentId} needs workspace {Workspace}", agent.Id, workspace.ToJson());
+								_logger.LogInformation("{AgentId} needs workspace {Workspace}", agent.Id, JsonSerializer.Serialize(workspace));
 							}
 						}
 					}
