Horde: More logging for conform updates.

#rnx

diff --git a/Engine/Source/Programs/Horde/Plugins/Build/HordeServer.Build/Jobs/JobRpcService.cs b/Engine/Source/Programs/Horde/Plugins/Build/HordeServer.Build/Jobs/JobRpcService.cs
--- a/Engine/Source/Programs/Horde/Plugins/Build/HordeServer.Build/Jobs/JobRpcService.cs
+++ b/Engine/Source/Programs/Horde/Plugins/Build/HordeServer.Build/Jobs/JobRpcService.cs
@@ -112,14 +112,14 @@ namespace HordeServer.Jobs
 
 						if (response.Retry)
 						{
-							_logger.LogInformation("Retrying conform for {AgentId}", agent.Id);
-							foreach (AgentWorkspaceInfo workspace in newWorkspaces)
-							{
-								_logger.LogInformation("{AgentId} has workspace {Workspace}", agent.Id, JsonSerializer.Serialize(workspace));
-							}
+							_logger.LogInformation("Retrying conform for {AgentId} ({BoolA},{BoolB},{BoolC})", agent.Id, conformWorkspaces.SetEquals(newWorkspaces), agent.RequestFullConform, request.RemoveUntrackedFiles);
 							foreach (AgentWorkspaceInfo workspace in conformWorkspaces)
 							{
-								_logger.LogInformation("{AgentId} needs workspace {Workspace}", agent.Id, JsonSerializer.Serialize(workspace));
+								AgentWorkspaceInfo? prevWorkspace = newWorkspaces.FirstOrDefault(x => x.Identifier == workspace.Identifier);
+								if (prevWorkspace != null)
+								{
+									_logger.LogInformation("{AgentId} {Identifier}: {OldWorkspace} -> {NewWorkspace} ({Equal})", agent.Id, workspace.Identifier, JsonSerializer.Serialize(workspace), JsonSerializer.Serialize(prevWorkspace), workspace.Equals(prevWorkspace));
+								}
 							}
 						}
 					}
