Horde: Fix AgentWorkspaceInfo objects constructed from RPC workspace messages having an empty string for 'method' instead of null, causing agents to get stuck in a conform loop.

diff --git a/Engine/Source/Programs/Horde/Plugins/Build/HordeServer.Build/Jobs/JobRpcService.cs b/Engine/Source/Programs/Horde/Plugins/Build/HordeServer.Build/Jobs/JobRpcService.cs
--- a/Engine/Source/Programs/Horde/Plugins/Build/HordeServer.Build/Jobs/JobRpcService.cs
+++ b/Engine/Source/Programs/Horde/Plugins/Build/HordeServer.Build/Jobs/JobRpcService.cs
@@ -109,19 +109,6 @@ namespace HordeServer.Jobs
 					{
 						response.Retry = await _conformTaskSource.GetWorkspacesAsync(agent, response.PendingWorkspaces, context.CancellationToken);
 						response.RemoveUntrackedFiles = request.RemoveUntrackedFiles || agent.RequestFullConform;
-
-						if (response.Retry)
-						{
-							_logger.LogInformation("Retrying conform for {AgentId} ({BoolA},{BoolB},{BoolC})", agent.Id, conformWorkspaces.SetEquals(newWorkspaces), agent.RequestFullConform, request.RemoveUntrackedFiles);
-							foreach (AgentWorkspaceInfo workspace in conformWorkspaces)
-							{
-								AgentWorkspaceInfo? prevWorkspace = newWorkspaces.FirstOrDefault(x => x.Identifier == workspace.Identifier);
-								if (prevWorkspace != null)
-								{
-									_logger.LogInformation("{AgentId} {Identifier}: {OldWorkspace} -> {NewWorkspace} ({Equal})", agent.Id, workspace.Identifier, JsonSerializer.Serialize(workspace), JsonSerializer.Serialize(prevWorkspace), workspace.Equals(prevWorkspace));
-								}
-							}
-						}
 					}
 					return response;
 				}
diff --git a/Engine/Source/Programs/Horde/Plugins/Compute/HordeServer.Compute/Agents/IAgent.cs b/Engine/Source/Programs/Horde/Plugins/Compute/HordeServer.Compute/Agents/IAgent.cs
--- a/Engine/Source/Programs/Horde/Plugins/Compute/HordeServer.Compute/Agents/IAgent.cs
+++ b/Engine/Source/Programs/Horde/Plugins/Compute/HordeServer.Compute/Agents/IAgent.cs
@@ -730,6 +730,8 @@ namespace HordeServer.Agents
 		public AgentWorkspaceInfo(RpcAgentWorkspace workspace)
 			: this(workspace.ConfiguredCluster, workspace.ConfiguredUserName, workspace.Identifier, workspace.Stream, (workspace.View.Count > 0) ? workspace.View.ToList() : null, workspace.Incremental, workspace.Method, workspace.MinScratchSpace, workspace.ConformDiskFreeSpace)
 		{
+			Method = String.IsNullOrEmpty(Method) ? null : Method;
+
 			// Treat zero value as null as Protobuf cannot store null 
 			MinScratchSpace = workspace.MinScratchSpace == 0 ? null : workspace.MinScratchSpace;
 			ConformDiskFreeSpace = workspace.ConformDiskFreeSpace == 0 ? null : workspace.ConformDiskFreeSpace;
