Horde: Make imports for uploaded blobs non-optional.

#rnx

diff --git a/Engine/Source/Programs/Horde/Plugins/Storage/HordeServer.Storage/Storage/StorageController.cs b/Engine/Source/Programs/Horde/Plugins/Storage/HordeServer.Storage/Storage/StorageController.cs
--- a/Engine/Source/Programs/Horde/Plugins/Storage/HordeServer.Storage/Storage/StorageController.cs
+++ b/Engine/Source/Programs/Horde/Plugins/Storage/HordeServer.Storage/Storage/StorageController.cs
@@ -117,9 +117,10 @@ namespace HordeServer.Storage
 		/// <returns>Information about the written blob, or redirect information</returns>
 		public static async Task<ActionResult<WriteBlobResponse>> WriteBlobAsync(IStorageBackend storageBackend, WriteBlobRequest request, CancellationToken cancellationToken = default)
 		{
+			IReadOnlyCollection<BlobLocator> imports = request.Imports ?? (IReadOnlyCollection<BlobLocator>)Array.Empty<BlobLocator>();
 			if (request.File == null)
 			{
-				(BlobLocator Path, Uri UploadUrl)? result = await storageBackend.TryGetBlobWriteRedirectAsync(request.Imports, request.Prefix ?? String.Empty, cancellationToken);
+				(BlobLocator Path, Uri UploadUrl)? result = await storageBackend.TryGetBlobWriteRedirectAsync(imports, request.Prefix ?? String.Empty, cancellationToken);
 				if (result == null)
 				{
 					return new WriteBlobResponse { SupportsRedirects = false };
@@ -130,7 +131,7 @@ namespace HordeServer.Storage
 			else
 			{
 				using Stream stream = request.File.OpenReadStream();
-				BlobLocator locator = await storageBackend.WriteBlobAsync(stream, request.Imports, request.Prefix, cancellationToken);
+				BlobLocator locator = await storageBackend.WriteBlobAsync(stream, imports, request.Prefix, cancellationToken);
 				return new WriteBlobResponse { Blob = locator.ToString(), SupportsRedirects = storageBackend.SupportsRedirects };
 			}
 		}
diff --git a/Engine/Source/Programs/Horde/Plugins/Storage/HordeServer.Storage/Storage/StorageService.cs b/Engine/Source/Programs/Horde/Plugins/Storage/HordeServer.Storage/Storage/StorageService.cs
--- a/Engine/Source/Programs/Horde/Plugins/Storage/HordeServer.Storage/Storage/StorageService.cs
+++ b/Engine/Source/Programs/Horde/Plugins/Storage/HordeServer.Storage/Storage/StorageService.cs
@@ -88,7 +88,7 @@ namespace HordeServer.Storage
 			}
 
 			/// <inheritdoc/>
-			public async Task<BlobLocator> WriteBlobAsync(Stream stream, IReadOnlyCollection<BlobLocator>? imports, string? basePath = null, CancellationToken cancellationToken = default)
+			public async Task<BlobLocator> WriteBlobAsync(Stream stream, IReadOnlyCollection<BlobLocator> imports, string? basePath = null, CancellationToken cancellationToken = default)
 			{
 				BlobLocator locator = StorageHelpers.CreateUniqueLocator(basePath);
 
@@ -106,7 +106,7 @@ namespace HordeServer.Storage
 			}
 
 			/// <inheritdoc/>
-			public async ValueTask<(BlobLocator, Uri)?> TryGetBlobWriteRedirectAsync(IReadOnlyCollection<BlobLocator>? imports = null, string? prefix = null, CancellationToken cancellationToken = default)
+			public async ValueTask<(BlobLocator, Uri)?> TryGetBlobWriteRedirectAsync(IReadOnlyCollection<BlobLocator> imports, string? prefix = null, CancellationToken cancellationToken = default)
 			{
 				if (!_store.SupportsRedirects)
 				{
@@ -607,29 +607,26 @@ namespace HordeServer.Storage
 		#region Blobs
 
 		/// <inheritdoc/>
-		async Task AddBlobAsync(NamespaceId namespaceId, BlobLocator locator, IReadOnlyCollection<BlobLocator>? imports = null, List<AliasInfo>? exports = null, CancellationToken cancellationToken = default)
+		async Task AddBlobAsync(NamespaceId namespaceId, BlobLocator locator, IReadOnlyCollection<BlobLocator> imports, List<AliasInfo>? exports = null, CancellationToken cancellationToken = default)
 		{
 			ObjectId id = ObjectId.GenerateNewId(_clock.UtcNow);
 			BlobInfo blobInfo = new BlobInfo(id, namespaceId, locator);
 			blobInfo.Aliases = exports;
 
-			if (imports != null)
+			if (imports.Count == 0)
 			{
-				if (imports.Count == 0)
-				{
-					blobInfo.Imports = new List<ObjectId>();
-				}
-				else
-				{
-					List<string> paths = imports.Select(x => x.BaseLocator.ToString()).ToList();
+				blobInfo.Imports = new List<ObjectId>();
+			}
+			else
+			{
+				List<string> paths = imports.Select(x => x.BaseLocator.ToString()).ToList();
 
-					FilterDefinition<BlobInfo> filter =
-						Builders<BlobInfo>.Filter.Eq(x => x.NamespaceId, namespaceId) &
-						Builders<BlobInfo>.Filter.In(x => x.Path, paths);
+				FilterDefinition<BlobInfo> filter =
+					Builders<BlobInfo>.Filter.Eq(x => x.NamespaceId, namespaceId) &
+					Builders<BlobInfo>.Filter.In(x => x.Path, paths);
 
-					blobInfo.Imports = await _blobCollection.Find(filter).Project(x => x.Id).ToListAsync(cancellationToken);
-					blobInfo.Imports.Sort();
-				}
+				blobInfo.Imports = await _blobCollection.Find(filter).Project(x => x.Id).ToListAsync(cancellationToken);
+				blobInfo.Imports.Sort();
 			}
 
 			await _blobCollection.InsertOneAsync(blobInfo, new InsertOneOptions { }, cancellationToken);
diff --git a/Engine/Source/Programs/Shared/EpicGames.Horde.Tests/BackendTests.cs b/Engine/Source/Programs/Shared/EpicGames.Horde.Tests/BackendTests.cs
--- a/Engine/Source/Programs/Shared/EpicGames.Horde.Tests/BackendTests.cs
+++ b/Engine/Source/Programs/Shared/EpicGames.Horde.Tests/BackendTests.cs
@@ -61,7 +61,7 @@ namespace EpicGames.Horde.Tests
 			{
 				using MemoryMappedFileCache memoryMappedFileCache = new MemoryMappedFileCache();
 				FileStorageBackend backend = new FileStorageBackend(tempDir.Location, memoryMappedFileCache, NullLogger.Instance);
-				BlobLocator path = await backend.WriteBytesAsync(Encoding.UTF8.GetBytes("hello world"));
+				BlobLocator path = await backend.WriteBytesAsync(Encoding.UTF8.GetBytes("hello world"), Array.Empty<BlobLocator>());
 
 				using (IReadOnlyMemoryOwner<byte> handle = await backend.ReadBlobAsync(path, 0, null))
 				{
@@ -95,15 +95,15 @@ namespace EpicGames.Horde.Tests
 				Assert.IsTrue(value.SequenceEqual(Encoding.UTF8.GetBytes("item 2")));
 
 				byte[] data3 = Encoding.UTF8.GetBytes("3");
-				BlobLocator path3 = await cacheBackend.WriteBytesAsync(data3);
+				BlobLocator path3 = await cacheBackend.WriteBytesAsync(data3, Array.Empty<BlobLocator>());
 				await cacheBackend.ReadBytesAsync(path3);
 
 				byte[] data4 = Encoding.UTF8.GetBytes("4");
-				BlobLocator path4 = await cacheBackend.WriteBytesAsync(data4);
+				BlobLocator path4 = await cacheBackend.WriteBytesAsync(data4, Array.Empty<BlobLocator>());
 				await cacheBackend.ReadBytesAsync(path4);
 
 				byte[] data5 = Encoding.UTF8.GetBytes("5");
-				BlobLocator path5 = await cacheBackend.WriteBytesAsync(data5);
+				BlobLocator path5 = await cacheBackend.WriteBytesAsync(data5, Array.Empty<BlobLocator>());
 				await cacheBackend.ReadBytesAsync(path5);
 
 				HashSet<BlobLocator> paths = new HashSet<BlobLocator>(cache.GetLocators());
@@ -113,7 +113,7 @@ namespace EpicGames.Horde.Tests
 				Assert.IsTrue(paths.Contains(path5));
 
 				byte[] data6 = Encoding.UTF8.GetBytes("12345678901");
-				BlobLocator path6 = await cacheBackend.WriteBytesAsync(data6);
+				BlobLocator path6 = await cacheBackend.WriteBytesAsync(data6, Array.Empty<BlobLocator>());
 				await cacheBackend.ReadBytesAsync(path6);
 
 				paths = new HashSet<BlobLocator>(cache.GetLocators());
@@ -136,10 +136,10 @@ namespace EpicGames.Horde.Tests
 			byte[] data1 = Encoding.UTF8.GetBytes("hello world");
 			byte[] data2 = Encoding.UTF8.GetBytes("item 2");
 
-			BlobLocator path1 = await backend.WriteBytesAsync(data1);
+			BlobLocator path1 = await backend.WriteBytesAsync(data1, Array.Empty<BlobLocator>());
 			byte[] outputData1 = await backend.ReadBytesAsync(path1);
 
-			BlobLocator path2 = await backend.WriteBytesAsync(data2);
+			BlobLocator path2 = await backend.WriteBytesAsync(data2, Array.Empty<BlobLocator>());
 			byte[] outputData2 = await backend.ReadBytesAsync(path2);
 
 			Assert.IsTrue(data1.SequenceEqual(outputData1));
diff --git a/Engine/Source/Programs/Shared/EpicGames.Horde/Storage/IStorageBackend.cs b/Engine/Source/Programs/Shared/EpicGames.Horde/Storage/IStorageBackend.cs
--- a/Engine/Source/Programs/Shared/EpicGames.Horde/Storage/IStorageBackend.cs
+++ b/Engine/Source/Programs/Shared/EpicGames.Horde/Storage/IStorageBackend.cs
@@ -48,11 +48,11 @@ namespace EpicGames.Horde.Storage
 		/// Writes a stream to the storage backend. If the stream throws an exception during read, the write will be aborted.
 		/// </summary>
 		/// <param name="stream">Data stream</param>
-		/// <param name="imports">Imported blobs. If omitted, the backend will parse them from the stream data.</param>
+		/// <param name="imports">List of referenced blobs</param>
 		/// <param name="prefix">Path prefix for the uploaded data</param>
 		/// <param name="cancellationToken">Cancellation token for the operation</param>
 		/// <returns>Path to the uploaded object</returns>
-		Task<BlobLocator> WriteBlobAsync(Stream stream, IReadOnlyCollection<BlobLocator>? imports = null, string? prefix = null, CancellationToken cancellationToken = default);
+		Task<BlobLocator> WriteBlobAsync(Stream stream, IReadOnlyCollection<BlobLocator> imports, string? prefix = null, CancellationToken cancellationToken = default);
 
 		/// <summary>
 		/// Gets a HTTP redirect for a read request
@@ -69,7 +69,7 @@ namespace EpicGames.Horde.Storage
 		/// <param name="prefix">Prefix for the uploaded data</param>
 		/// <param name="cancellationToken">Cancellation token for the operation</param>
 		/// <returns>Path for retrieval, and URI to upload the data to</returns>
-		ValueTask<(BlobLocator, Uri)?> TryGetBlobWriteRedirectAsync(IReadOnlyCollection<BlobLocator>? imports = null, string? prefix = null, CancellationToken cancellationToken = default);
+		ValueTask<(BlobLocator, Uri)?> TryGetBlobWriteRedirectAsync(IReadOnlyCollection<BlobLocator> imports, string? prefix = null, CancellationToken cancellationToken = default);
 
 		#endregion
 
@@ -217,7 +217,7 @@ namespace EpicGames.Horde.Storage
 		/// <param name="imports"></param>
 		/// <param name="prefix">Prefix for the uploaded data</param>
 		/// <param name="cancellationToken">Cancellation token for the operation</param>
-		public static async Task<BlobLocator> WriteBytesAsync(this IStorageBackend storageBackend, ReadOnlyMemory<byte> data, IReadOnlyCollection<BlobLocator>? imports = null, string? prefix = null, CancellationToken cancellationToken = default)
+		public static async Task<BlobLocator> WriteBytesAsync(this IStorageBackend storageBackend, ReadOnlyMemory<byte> data, IReadOnlyCollection<BlobLocator> imports, string? prefix = null, CancellationToken cancellationToken = default)
 		{
 			using (ReadOnlyMemoryStream stream = new ReadOnlyMemoryStream(data))
 			{
diff --git a/Engine/Source/Programs/Shared/EpicGames.Horde/Storage/StorageBackendCache.cs b/Engine/Source/Programs/Shared/EpicGames.Horde/Storage/StorageBackendCache.cs
--- a/Engine/Source/Programs/Shared/EpicGames.Horde/Storage/StorageBackendCache.cs
+++ b/Engine/Source/Programs/Shared/EpicGames.Horde/Storage/StorageBackendCache.cs
@@ -129,13 +129,13 @@ namespace EpicGames.Horde.Storage
 				return storageObject.Slice(offset, length);
 			}
 
-			public Task<BlobLocator> WriteBlobAsync(Stream stream, IReadOnlyCollection<BlobLocator>? imports, string? prefix = null, CancellationToken cancellationToken = default) 
+			public Task<BlobLocator> WriteBlobAsync(Stream stream, IReadOnlyCollection<BlobLocator> imports, string? prefix = null, CancellationToken cancellationToken = default) 
 				=> _inner.WriteBlobAsync(stream, imports, prefix, cancellationToken);
 
 			public ValueTask<Uri?> TryGetBlobReadRedirectAsync(BlobLocator locator, CancellationToken cancellationToken = default) 
 				=> _inner.TryGetBlobReadRedirectAsync(locator, cancellationToken);
 
-			public ValueTask<(BlobLocator, Uri)?> TryGetBlobWriteRedirectAsync(IReadOnlyCollection<BlobLocator>? imports = null, string? prefix = null, CancellationToken cancellationToken = default) 
+			public ValueTask<(BlobLocator, Uri)?> TryGetBlobWriteRedirectAsync(IReadOnlyCollection<BlobLocator> imports, string? prefix = null, CancellationToken cancellationToken = default) 
 				=> _inner.TryGetBlobWriteRedirectAsync(imports, prefix, cancellationToken);
 
 			public void GetStats(StorageStats stats)
