Horde: Fix issue with artfifact log rendering when there is no artifact type
#skipci

diff --git a/Engine/Source/Programs/Horde/HordeDashboard/src/components/LogRender.tsx b/Engine/Source/Programs/Horde/HordeDashboard/src/components/LogRender.tsx
--- a/Engine/Source/Programs/Horde/HordeDashboard/src/components/LogRender.tsx
+++ b/Engine/Source/Programs/Horde/HordeDashboard/src/components/LogRender.tsx
@@ -268,20 +268,18 @@ const renderTags = (navigate: NavigateFunction, line: LogLine, lineNumber: numbe
 
          return <a key={key} rel="noreferrer" href={record.target} onClick={(ev) => ev.stopPropagation()}><Highlight search={search ? search : ""} className={logStyle.logLine}>{text}</Highlight></a>;
       } else if (tagType === TagType.ArtifactId) {
-         const artifactType = properties.ArtifactType;
-         if (artifactType?.length) {
-
-            const search = new URLSearchParams(window.location.search);
-            search.set("artifactContext", encodeURIComponent(artifactType as string));
-            if (properties.ArtifactId) {
-               if (properties.ArtifactId["$text"]) {
-                  search.set("artifactId", encodeURIComponent(properties.ArtifactId["$text"] as string));
-               }               
+         const artifactType = properties.ArtifactType ?? "_none";
+
+         const search = new URLSearchParams(window.location.search);
+         search.set("artifactContext", encodeURIComponent(artifactType as string));
+         if (properties.ArtifactId) {
+            if (properties.ArtifactId["$text"]) {
+               search.set("artifactId", encodeURIComponent(properties.ArtifactId["$text"] as string));
             }
-            
-            const url = `${window.location.pathname}?` + search.toString();
-            return <a key={key} href="/" onClick={(ev) => { ev.stopPropagation(); ev.preventDefault(); navigate(url, { replace: true }) }}><Highlight search={search ? search : ""} className={logStyle.logLine}>{text}</Highlight></a>;
          }
+
+         const url = `${window.location.pathname}?` + search.toString();
+         return <a key={key} href="/" onClick={(ev) => { ev.stopPropagation(); ev.preventDefault(); navigate(url, { replace: true }) }}><Highlight search={search ? search : ""} className={logStyle.logLine}>{text}</Highlight></a>;
       }
 
       return <span key={key} />;
