Horde: Fix bug where leases weren't immediately aborted when agent is marked as busy

Busy value was one update session loop iteration behind leading to extra delay. Draining of leases still has a grace period of 6 seconds to let the termination file signaling take place. Primarily provided for UBA. Can potentially be shortened but requires changes to UBA.

diff --git a/Engine/Source/Programs/Horde/HordeAgent/Leases/LeaseManager.cs b/Engine/Source/Programs/Horde/HordeAgent/Leases/LeaseManager.cs
--- a/Engine/Source/Programs/Horde/HordeAgent/Leases/LeaseManager.cs
+++ b/Engine/Source/Programs/Horde/HordeAgent/Leases/LeaseManager.cs
@@ -310,7 +310,6 @@ namespace HordeAgent.Leases
 				}
 
 				// Get the new agent status to be reported back to server
-				bool busy = _statusService.IsBusy;
 				if (stopping)
 				{
 					updateSessionRequest.Status = RpcAgentStatus.Stopping;
@@ -319,7 +318,7 @@ namespace HordeAgent.Leases
 				{
 					updateSessionRequest.Status = RpcAgentStatus.Unhealthy;
 				}
-				else if (busy)
+				else if (_statusService.IsBusy)
 				{
 					updateSessionRequest.Status = RpcAgentStatus.Busy;
 				}
@@ -413,7 +412,7 @@ namespace HordeAgent.Leases
 					{
 						_statusService.Set(false, _activeLeases.Count, "Attempting to connect to server...");
 					}
-					else if (busy)
+					else if (_statusService.IsBusy)
 					{
 						_statusService.Set(true, 0, "Paused");
 
