Horde: Fix issues with agent enrollment and deletion latency
#skipci

diff --git a/Engine/Source/Programs/Horde/HordeDashboard/src/backend/AgentStore.ts b/Engine/Source/Programs/Horde/HordeDashboard/src/backend/AgentStore.ts
--- a/Engine/Source/Programs/Horde/HordeDashboard/src/backend/AgentStore.ts
+++ b/Engine/Source/Programs/Horde/HordeDashboard/src/backend/AgentStore.ts
@@ -116,7 +116,7 @@ export class AgentStore {
 
     }
 
-    async update(slim = false): Promise<void> {
+    async update(slim:boolean = false, invalidateCache: boolean = false): Promise<void> {
         return new Promise<void>((resolve, reject) => {
 
             if (this.inflight !== undefined && slim) {                
@@ -128,7 +128,7 @@ export class AgentStore {
 
             const filter = "id,name,sessionId,sessionExpiresAt,online,enabled,ephemeral,comment,version,forceVersion,pools,capabilities,leases,acl,updateTime,deleted,pendingConform,pendingFullConform,conformAttemptCount,lastConformTime,nextConformTime,lastShutdownReason,pendingShutdown,status";
             const promises: any[] = [];
-            promises.push(backend.getAgents({ includeDeleted: false, modifiedAfter: this.modifiedAfterDate?.toISOString(), filter:filter    }));
+            promises.push(backend.getAgents({ includeDeleted: false, modifiedAfter: this.modifiedAfterDate?.toISOString(), filter:filter,invalidateCache:invalidateCache}));
             if (!slim) {
                 promises.push(backend.getPools());
             }
diff --git a/Engine/Source/Programs/Horde/HordeDashboard/src/backend/Api.ts b/Engine/Source/Programs/Horde/HordeDashboard/src/backend/Api.ts
--- a/Engine/Source/Programs/Horde/HordeDashboard/src/backend/Api.ts
+++ b/Engine/Source/Programs/Horde/HordeDashboard/src/backend/Api.ts
@@ -295,6 +295,7 @@ export type AgentQuery = {
 	modifiedAfter?: string;
 	poolId?: string;
 	includeDeleted?: boolean;
+	invalidateCache?: boolean;
 	condition?: string;
 	filter?: string;
 }
diff --git a/Engine/Source/Programs/Horde/HordeDashboard/src/components/AgentView.tsx b/Engine/Source/Programs/Horde/HordeDashboard/src/components/AgentView.tsx
--- a/Engine/Source/Programs/Horde/HordeDashboard/src/components/AgentView.tsx
+++ b/Engine/Source/Programs/Horde/HordeDashboard/src/components/AgentView.tsx
@@ -438,7 +438,7 @@ class LocalState {
       }).finally(function () {
          that.selection.setAllSelected(false);
          that.currentSelection = [];
-         agentStore.update();
+         agentStore.update(false, true);
       });
    };
 
diff --git a/Engine/Source/Programs/Horde/HordeDashboard/src/components/agents/AgentRequestsView.tsx b/Engine/Source/Programs/Horde/HordeDashboard/src/components/agents/AgentRequestsView.tsx
--- a/Engine/Source/Programs/Horde/HordeDashboard/src/components/agents/AgentRequestsView.tsx
+++ b/Engine/Source/Programs/Horde/HordeDashboard/src/components/agents/AgentRequestsView.tsx
@@ -93,7 +93,7 @@ const AgentsPanel: React.FC = observer(() => {
 
       const queryAgents = async (queryTime: string) => {
 
-         const nagents = await backend.getAgents({ modifiedAfter: queryTime });
+         const nagents = await backend.getAgents({ modifiedAfter: queryTime, invalidateCache: true });
          if (nagents.length) {
             await agentStore.update(false);
             navigate("/agents");
