Horde: Fix agent store updates upon cache invalidation/deletes
#skipci

diff --git a/Engine/Source/Programs/Horde/HordeDashboard/src/backend/AgentStore.ts b/Engine/Source/Programs/Horde/HordeDashboard/src/backend/AgentStore.ts
--- a/Engine/Source/Programs/Horde/HordeDashboard/src/backend/AgentStore.ts
+++ b/Engine/Source/Programs/Horde/HordeDashboard/src/backend/AgentStore.ts
@@ -56,9 +56,11 @@ export class AgentStore {
                     toSet.push(updatedAgent);
                 }
             });            
-            this._agents = toSet.filter(a => true /*!!a.deleted*/);
-            this.agentsUpdated++;
+            this._agents = toSet.filter(a => true /*!!a.deleted*/);            
+        } else {
+            this._agents = [];
         }
+        this.agentsUpdated++;
     }
 
     async updateAgent(agentId: string): Promise<GetAgentResponse> {
@@ -124,6 +126,10 @@ export class AgentStore {
                 return;
             }
 
+            if (invalidateCache) {
+                this.modifiedAfterDate = undefined;
+            }
+
             this.inflight = true;
 
             const filter = "id,name,sessionId,sessionExpiresAt,online,enabled,ephemeral,comment,version,forceVersion,pools,capabilities,leases,acl,updateTime,deleted,pendingConform,pendingFullConform,conformAttemptCount,lastConformTime,nextConformTime,lastShutdownReason,pendingShutdown,status";
diff --git a/Engine/Source/Programs/Horde/HordeDashboard/src/components/agents/AgentRequestsView.tsx b/Engine/Source/Programs/Horde/HordeDashboard/src/components/agents/AgentRequestsView.tsx
--- a/Engine/Source/Programs/Horde/HordeDashboard/src/components/agents/AgentRequestsView.tsx
+++ b/Engine/Source/Programs/Horde/HordeDashboard/src/components/agents/AgentRequestsView.tsx
@@ -98,7 +98,7 @@ const AgentsPanel: React.FC = observer(() => {
             await agentStore.update(false);
             navigate("/agents");
          } else {
-            setTimeout(queryAgents, 3000);
+            setTimeout(queryAgents, 2000);
          }
       }
 
