Horde: Fix build errors for Docker image with bundled tools

Paths changed for pre-built tools.

diff --git a/Engine/Source/Programs/Horde/BuildHorde.xml b/Engine/Source/Programs/Horde/BuildHorde.xml
--- a/Engine/Source/Programs/Horde/BuildHorde.xml
+++ b/Engine/Source/Programs/Horde/BuildHorde.xml
@@ -375,7 +375,7 @@
 
 			<!-- A Horde server image has already been built, and is referenced inside Dockerfile.bundled -->
 			<Copy From="Engine/Source/Programs/Horde/HordeServer/Dockerfile.bundled" To="$(StagingDir)/Dockerfile.bundled" />
-			<Tag Files="$(StagingDir)/Dockerfile*;$(StagingDir)/Dashboard/...;$(StagingDir)/Server/Tools/...;$(StagingDir)/Server/*.json" With="#DockerFiles"/>
+			<Tag Files="$(StagingDir)/Dockerfile*;$(StagingDir)/Dashboard/...;$(StagingDir)/Server/win-x64/Tools/...;$(StagingDir)/Server/win-x64/*.json" With="#DockerFiles"/>
 			<Docker-Build BaseDir="$(StagingDir)" Files="#DockerFiles" Tag="horde-server:bundled" DockerFile="$(StagingDir)/Dockerfile.bundled" />
 		</Node>
 
