Horde: Send Slack notifications for job steps timing out.

#rnx

diff --git a/Engine/Source/Programs/Horde/Plugins/Build/HordeServer.Build.Tests/Notifications/NotificationServiceTests.cs b/Engine/Source/Programs/Horde/Plugins/Build/HordeServer.Build.Tests/Notifications/NotificationServiceTests.cs
--- a/Engine/Source/Programs/Horde/Plugins/Build/HordeServer.Build.Tests/Notifications/NotificationServiceTests.cs
+++ b/Engine/Source/Programs/Horde/Plugins/Build/HordeServer.Build.Tests/Notifications/NotificationServiceTests.cs
@@ -34,7 +34,7 @@ namespace HordeServer.Tests.Notifications
 
 		public Task NotifyJobCompleteAsync(IJob job, IGraph graph, LabelOutcome outcome, CancellationToken cancellationToken) { throw new NotImplementedException(); }
 		public Task NotifyJobCompleteAsync(IUser user, IJob job, IGraph graph, LabelOutcome outcome, CancellationToken cancellationToken) { throw new NotImplementedException(); }
-		public Task NotifyJobStepCompleteAsync(IUser user, IJob job, IJobStepBatch batch, IJobStep step, INode node, List<ILogEventData> jobStepEventData, CancellationToken cancellationToken) { throw new NotImplementedException(); }
+		public Task NotifyJobStepCompleteAsync(IEnumerable<IUser> usersToNotify, IJob job, IJobStepBatch batch, IJobStep step, INode node, List<ILogEventData> jobStepEventData, CancellationToken cancellationToken) { throw new NotImplementedException(); }
 		public Task NotifyLabelCompleteAsync(IUser user, IJob job, ILabel label, int labelIdx, LabelOutcome outcome, List<(string, JobStepOutcome, Uri)> stepData, CancellationToken cancellationToken) { throw new NotImplementedException(); }
 		public Task NotifyIssueUpdatedAsync(IIssue issue, CancellationToken cancellationToken) { throw new NotImplementedException(); }
 		public Task NotifyConfigUpdateAsync(ConfigUpdateInfo info, CancellationToken cancellationToken) => Task.CompletedTask;
diff --git a/Engine/Source/Programs/Horde/Plugins/Build/HordeServer.Build/Notifications/INotificationSink.cs b/Engine/Source/Programs/Horde/Plugins/Build/HordeServer.Build/Notifications/INotificationSink.cs
--- a/Engine/Source/Programs/Horde/Plugins/Build/HordeServer.Build/Notifications/INotificationSink.cs
+++ b/Engine/Source/Programs/Horde/Plugins/Build/HordeServer.Build/Notifications/INotificationSink.cs
@@ -52,7 +52,7 @@ namespace HordeServer.Notifications
 		/// <summary>
 		/// Send notifications that a job step has completed
 		/// </summary>
-		/// <param name="user">User to notify</param>
+		/// <param name="usersToNotify">Users to notify</param>
 		/// <param name="job">The job containing the step</param>
 		/// <param name="batch">Unique id of the batch</param>
 		/// <param name="step">The step id</param>
@@ -60,7 +60,7 @@ namespace HordeServer.Notifications
 		/// <param name="jobStepEventData"></param>
 		/// <param name="cancellationToken">Cancellation token for the operation</param>
 		/// <returns>Async task</returns>
-		Task NotifyJobStepCompleteAsync(IUser user, IJob job, IJobStepBatch batch, IJobStep step, INode node, List<ILogEventData> jobStepEventData, CancellationToken cancellationToken);
+		Task NotifyJobStepCompleteAsync(IEnumerable<IUser> usersToNotify, IJob job, IJobStepBatch batch, IJobStep step, INode node, List<ILogEventData> jobStepEventData, CancellationToken cancellationToken);
 
 		/// <summary>
 		/// Send notifications that a job step has completed
diff --git a/Engine/Source/Programs/Horde/Plugins/Build/HordeServer.Build/Notifications/NotificationService.cs b/Engine/Source/Programs/Horde/Plugins/Build/HordeServer.Build/Notifications/NotificationService.cs
--- a/Engine/Source/Programs/Horde/Plugins/Build/HordeServer.Build/Notifications/NotificationService.cs
+++ b/Engine/Source/Programs/Horde/Plugins/Build/HordeServer.Build/Notifications/NotificationService.cs
@@ -688,6 +688,7 @@ namespace HordeServer.Notifications
 				jobStepEventData.Add(eventData);
 			}
 
+			List<IUser> filteredUsersToNotify = new List<IUser>();
 			foreach (IUser slackUser in usersToNotify)
 			{
 				if (job.PreflightCommitId != null)
@@ -697,8 +698,11 @@ namespace HordeServer.Notifications
 						continue;
 					}
 				}
-				EnqueueTasks((sink, ctx) => sink.NotifyJobStepCompleteAsync(slackUser, job, batch, step, node, jobStepEventData, ctx));
+				filteredUsersToNotify.Add(slackUser);
 			}
+
+			EnqueueTasks((sink, ctx) => sink.NotifyJobStepCompleteAsync(filteredUsersToNotify, job, batch, step, node, jobStepEventData, ctx));
+
 			_logger.LogDebug("Finished sending notifications for step {JobId}:{BatchId}:{StepId}", job.Id, batchId, stepId);
 		}
 
diff --git a/Engine/Source/Programs/Horde/Plugins/Build/HordeServer.Build/Notifications/Sinks/SlackNotificationSink.cs b/Engine/Source/Programs/Horde/Plugins/Build/HordeServer.Build/Notifications/Sinks/SlackNotificationSink.cs
--- a/Engine/Source/Programs/Horde/Plugins/Build/HordeServer.Build/Notifications/Sinks/SlackNotificationSink.cs
+++ b/Engine/Source/Programs/Horde/Plugins/Build/HordeServer.Build/Notifications/Sinks/SlackNotificationSink.cs
@@ -596,14 +596,41 @@ namespace HordeServer.Notifications.Sinks
 		#region Job step complete
 
 		/// <inheritdoc/>
-		public async Task NotifyJobStepCompleteAsync(IUser slackUser, IJob job, IJobStepBatch batch, IJobStep step, INode node, List<ILogEventData> jobStepEventData, CancellationToken cancellationToken)
+		public async Task NotifyJobStepCompleteAsync(IEnumerable<IUser> usersToNotify, IJob job, IJobStepBatch batch, IJobStep step, INode node, List<ILogEventData> jobStepEventData, CancellationToken cancellationToken)
 		{
-			_logger.LogInformation("Sending Slack notification for job {JobId}, batch {BatchId}, step {StepId}, outcome {Outcome} to {SlackUser} ({UserId})", job.Id, batch.Id, step.Id, step.Outcome, slackUser.Name, slackUser.Id);
+			if (step.Error == JobStepError.TimedOut)
+			{
+				if (_settings.JobNotificationChannel != null)
+				{
+					BuildConfig buildConfig = _buildConfig.CurrentValue;
+					if (buildConfig.TryGetStream(job.StreamId, out StreamConfig? streamConfig))
+					{
+						foreach (string channel in _settings.JobNotificationChannel.Split(';'))
+						{
+							Uri jobStepLink = new Uri($"{_serverInfo.DashboardUrl}job/{job.Id}?step={step.Id}");
+							Uri jobStepLogLink = new Uri($"{_serverInfo.DashboardUrl}log/{step.LogId}");
 
-			string? slackUserId = await GetSlackUserIdAsync(slackUser, cancellationToken);
-			if (slackUserId != null)
+							SlackAttachment attachment = new SlackAttachment();
+							attachment.FallbackText = $"{streamConfig.Name} - {GetJobChangeText(job)} - {job.Name} - {node.Name} - {step.Outcome}";
+							attachment.Color = ErrorColor;
+							attachment.AddSection($"*<{jobStepLink}|{streamConfig.Name} - {GetJobChangeText(job)} - {job.Name} - {node.Name}>*");
+							attachment.AddSection($"*Job Step Aborted due to Timeout*");
+
+							await SendMessageAsync($"#{channel}", attachment, cancellationToken);
+						}
+					}
+				}
+			}
+
+			foreach (IUser slackUser in usersToNotify)
 			{
-				await SendJobStepCompleteMessageAsync(slackUserId, job, step, node, jobStepEventData, cancellationToken);
+				_logger.LogInformation("Sending Slack notification for job {JobId}, batch {BatchId}, step {StepId}, outcome {Outcome} to {SlackUser} ({UserId})", job.Id, batch.Id, step.Id, step.Outcome, slackUser.Name, slackUser.Id);
+
+				string? slackUserId = await GetSlackUserIdAsync(slackUser, cancellationToken);
+				if (slackUserId != null)
+				{
+					await SendJobStepCompleteMessageAsync(slackUserId, job, step, node, jobStepEventData, cancellationToken);
+				}
 			}
 		}
 
