Horde: Support running jobs in streams that do not have the engine directly under the stream root. The "EnginePath" property in the stream config can be used to configure the path to the engine folder.

diff --git a/Engine/Source/Programs/Horde/Drivers/JobDriver/Execution/JobExecutor.cs b/Engine/Source/Programs/Horde/Drivers/JobDriver/Execution/JobExecutor.cs
--- a/Engine/Source/Programs/Horde/Drivers/JobDriver/Execution/JobExecutor.cs
+++ b/Engine/Source/Programs/Horde/Drivers/JobDriver/Execution/JobExecutor.cs
@@ -199,14 +199,16 @@ namespace JobDriver.Execution
 		const string ScriptArgumentPrefix = "-Script=";
 		const string TargetArgumentPrefix = "-Target=";
 
-		const string PreprocessedScript = "Engine/Saved/Horde/Preprocessed.xml";
-		const string PreprocessedSchema = "Engine/Saved/Horde/Preprocessed.xsd";
+		string PreprocessedScript => $"{EnginePath}/Saved/Horde/Preprocessed.xml";
+		string PreprocessedSchema => $"{EnginePath}/Saved/Horde/Preprocessed.xsd";
 
 		protected List<string> _targets = new List<string>();
 		protected string? _scriptFileName;
 		protected bool _preprocessScript;
 		protected bool _savePreprocessedScript;
-		
+
+		protected string EnginePath => String.IsNullOrEmpty(Batch.EnginePath)? "Engine" : Batch.EnginePath;
+
 		protected Tracer Tracer { get; }
 
 		/// <summary>
@@ -714,8 +716,8 @@ namespace JobDriver.Execution
 
 			FileReference[] automationToolPaths = new FileReference[]
 			{
-				FileReference.Combine(buildDir, "Engine", "Binaries", "DotNET", "AutomationTool.exe"),
-				FileReference.Combine(buildDir, "Engine", "Binaries", "DotNET", "AutomationTool", "AutomationTool.exe")
+				FileReference.Combine(buildDir, $"{EnginePath}/Binaries/DotNET/AutomationTool.exe"),
+				FileReference.Combine(buildDir, $"{EnginePath}/Binaries/DotNET/AutomationTool/AutomationTool.exe")
 			};
 
 			if (automationToolPaths.Any(automationTool => FileReference.Exists(automationTool)))
@@ -735,9 +737,9 @@ namespace JobDriver.Execution
 			}
 		}
 
-		protected static void DeleteCachedBuildGraphManifests(DirectoryReference workspaceDir, ILogger logger)
+		protected void DeleteCachedBuildGraphManifests(DirectoryReference workspaceDir, ILogger logger)
 		{
-			DirectoryReference manifestDir = DirectoryReference.Combine(workspaceDir, "Engine", "Saved", "BuildGraph");
+			DirectoryReference manifestDir = DirectoryReference.Combine(workspaceDir, $"{EnginePath}/Saved/BuildGraph");
 			if (DirectoryReference.Exists(manifestDir))
 			{
 				try
@@ -765,7 +767,7 @@ namespace JobDriver.Execution
 
 		protected virtual async Task<bool> SetupAsync(JobStepInfo step, DirectoryReference workspaceDir, bool? useP4, ILogger logger, CancellationToken cancellationToken)
 		{
-			FileReference definitionFile = FileReference.Combine(workspaceDir, "Engine", "Saved", "Horde", "Exported.json");
+			FileReference definitionFile = FileReference.Combine(workspaceDir, $"{EnginePath}/Saved/Horde/Exported.json");
 
 			StringBuilder arguments = new StringBuilder($"BuildGraph");
 			if (_scriptFileName != null)
@@ -1185,7 +1187,7 @@ namespace JobDriver.Execution
 
 		private async Task<bool> ExecuteWithTempStorageAsync(JobStepInfo step, DirectoryReference workspaceDir, string arguments, bool? useP4, ILogger logger, CancellationToken cancellationToken)
 		{
-			DirectoryReference manifestDir = DirectoryReference.Combine(workspaceDir, "Engine", "Saved", "BuildGraph");
+			DirectoryReference manifestDir = DirectoryReference.Combine(workspaceDir, $"{EnginePath}/Saved/BuildGraph");
 
 			// Create the mapping of tag names to file sets
 			Dictionary<string, HashSet<FileReference>> tagNameToFileSet = new Dictionary<string, HashSet<FileReference>>();
@@ -1505,22 +1507,22 @@ namespace JobDriver.Execution
 
 			if (RuntimeInformation.IsOSPlatform(OSPlatform.Windows))
 			{
-				result = await ExecuteCommandAsync(step, workspaceDir, Environment.GetEnvironmentVariable("COMSPEC") ?? "cmd.exe", $"/C \"\"{workspaceDir}\\Engine\\Build\\BatchFiles\\RunUAT.bat\" {arguments}\"", span, logger, cancellationToken);
+				result = await ExecuteCommandAsync(step, workspaceDir, Environment.GetEnvironmentVariable("COMSPEC") ?? "cmd.exe", $"/C \"\"{workspaceDir}\\{EnginePath}\\Build\\BatchFiles\\RunUAT.bat\" {arguments}\"", span, logger, cancellationToken);
 			}
 			else if (RuntimeInformation.IsOSPlatform(OSPlatform.Linux))
 			{
-				string args = $"\"{workspaceDir}/Engine/Build/BatchFiles/RunUAT.sh\" {arguments}";
+				string args = $"\"{workspaceDir}/{EnginePath}/Build/BatchFiles/RunUAT.sh\" {arguments}";
 
 				if (JobOptions.UseWine is true)
 				{
-					args = $"\"{workspaceDir}/Engine/Build/BatchFiles/RunWineUAT.sh\" {arguments}";
+					args = $"\"{workspaceDir}/{EnginePath}/Build/BatchFiles/RunWineUAT.sh\" {arguments}";
 				}
 
 				result = await ExecuteCommandAsync(step, workspaceDir, "/bin/bash", args, span, logger, cancellationToken);
 			}
 			else if (RuntimeInformation.IsOSPlatform(OSPlatform.OSX))
 			{
-				result = await ExecuteCommandAsync(step, workspaceDir, "/bin/sh", $"\"{workspaceDir}/Engine/Build/BatchFiles/RunUAT.sh\" {arguments}", span, logger, cancellationToken);
+				result = await ExecuteCommandAsync(step, workspaceDir, "/bin/sh", $"\"{workspaceDir}/{EnginePath}/Build/BatchFiles/RunUAT.sh\" {arguments}", span, logger, cancellationToken);
 			}
 			else
 			{
@@ -1544,10 +1546,10 @@ namespace JobDriver.Execution
 			}
 		}
 
-		static async Task<List<string>> ReadIgnorePatternsAsync(DirectoryReference workspaceDir, ILogger logger)
+		async Task<List<string>> ReadIgnorePatternsAsync(DirectoryReference workspaceDir, ILogger logger)
 		{
 			List<DirectoryReference> baseDirs = new List<DirectoryReference>();
-			baseDirs.Add(DirectoryReference.Combine(workspaceDir, "Engine"));
+			baseDirs.Add(DirectoryReference.Combine(workspaceDir, EnginePath));
 			AddRestrictedDirs(baseDirs, "Restricted");
 			AddRestrictedDirs(baseDirs, "Platforms");
 
@@ -1809,19 +1811,19 @@ namespace JobDriver.Execution
 			// Add all the other Horde-specific variables
 			newEnvVars["IsBuildMachine"] = "1";
 
-			DirectoryReference logDir = DirectoryReference.Combine(workspaceDir, "Engine", "Programs", "AutomationTool", "Saved", "Logs");
+			DirectoryReference logDir = DirectoryReference.Combine(workspaceDir, $"{EnginePath}/Programs/AutomationTool/Saved/Logs");
 			FileUtils.ForceDeleteDirectoryContents(logDir);
 			newEnvVars["uebp_LogFolder"] = logDir.FullName;
 
-			DirectoryReference telemetryDir = DirectoryReference.Combine(workspaceDir, "Engine", "Programs", "AutomationTool", "Saved", "Telemetry");
+			DirectoryReference telemetryDir = DirectoryReference.Combine(workspaceDir, $"{EnginePath}/Programs/AutomationTool/Saved/Telemetry");
 			FileUtils.ForceDeleteDirectoryContents(telemetryDir);
 			newEnvVars["UE_TELEMETRY_DIR"] = telemetryDir.FullName;
 
-			DirectoryReference testDataDir = DirectoryReference.Combine(workspaceDir, "Engine", "Programs", "AutomationTool", "Saved", "TestData");
+			DirectoryReference testDataDir = DirectoryReference.Combine(workspaceDir, $"{EnginePath}/Programs/AutomationTool/Saved/TestData");
 			FileUtils.ForceDeleteDirectoryContents(testDataDir);
 			newEnvVars["UE_TESTDATA_DIR"] = testDataDir.FullName;
 
-			FileReference graphUpdateFile = FileReference.Combine(workspaceDir, "Engine", "Saved", "Horde", "Graph.json");
+			FileReference graphUpdateFile = FileReference.Combine(workspaceDir, $"{EnginePath}/Saved/Horde/Graph.json");
 			FileUtils.ForceDeleteFile(graphUpdateFile);
 			newEnvVars["UE_HORDE_GRAPH_UPDATE"] = graphUpdateFile.FullName;
 
diff --git a/Engine/Source/Programs/Horde/Drivers/JobDriver/Execution/PerforceExecutor.cs b/Engine/Source/Programs/Horde/Drivers/JobDriver/Execution/PerforceExecutor.cs
--- a/Engine/Source/Programs/Horde/Drivers/JobDriver/Execution/PerforceExecutor.cs
+++ b/Engine/Source/Programs/Horde/Drivers/JobDriver/Execution/PerforceExecutor.cs
@@ -131,7 +131,7 @@ namespace JobDriver.Execution
 				await _workspace.SyncAsync(perforce, Batch.Change, Batch.PreflightChange, null, cancellationToken);
 
 				// Remove any cached BuildGraph manifests
-				DirectoryReference manifestDir = DirectoryReference.Combine(_workspace.WorkspaceDir, "Engine", "Saved", "BuildGraph");
+				DirectoryReference manifestDir = DirectoryReference.Combine(_workspace.WorkspaceDir, $"{EnginePath}/Saved/BuildGraph");
 				if (DirectoryReference.Exists(manifestDir))
 				{
 					try
diff --git a/Engine/Source/Programs/Horde/Plugins/Build/HordeServer.Build/Jobs/JobRpcService.cs b/Engine/Source/Programs/Horde/Plugins/Build/HordeServer.Build/Jobs/JobRpcService.cs
--- a/Engine/Source/Programs/Horde/Plugins/Build/HordeServer.Build/Jobs/JobRpcService.cs
+++ b/Engine/Source/Programs/Horde/Plugins/Build/HordeServer.Build/Jobs/JobRpcService.cs
@@ -198,6 +198,7 @@ namespace HordeServer.Jobs
 			response.AgentType = graph.Groups[batch.GroupIdx].AgentType;
 			response.StreamId = streamConfig.Id.ToString();
 			response.StreamName = streamConfig.Name;
+			response.EnginePath = streamConfig.EnginePath;
 			response.Change = job.CommitId.GetPerforceChange();
 			response.CodeChange = job.CodeCommitId?.GetPerforceChange() ?? 0;
 			response.PreflightChange = job.PreflightCommitId?.GetPerforceChange() ?? 0;
diff --git a/Engine/Source/Programs/Horde/Plugins/Build/HordeServer.Build/Streams/StreamConfig.cs b/Engine/Source/Programs/Horde/Plugins/Build/HordeServer.Build/Streams/StreamConfig.cs
--- a/Engine/Source/Programs/Horde/Plugins/Build/HordeServer.Build/Streams/StreamConfig.cs
+++ b/Engine/Source/Programs/Horde/Plugins/Build/HordeServer.Build/Streams/StreamConfig.cs
@@ -126,6 +126,11 @@ namespace HordeServer.Streams
 		/// </summary>
 		public string Name { get; set; } = String.Empty;
 
+		/// <summary>
+		/// Path to the engine directory within the workspace. Used for launching UAT.
+		/// </summary>
+		public string EnginePath { get; set; } = "Engine";
+
 		/// <summary>
 		/// The perforce cluster containing the stream
 		/// </summary>
