Horde: Fix static analysis warning due to unused variable.

#rnx

diff --git a/Engine/Source/Programs/Horde/Plugins/Build/HordeServer.Build/Notifications/Sinks/SlackNotificationSink.cs b/Engine/Source/Programs/Horde/Plugins/Build/HordeServer.Build/Notifications/Sinks/SlackNotificationSink.cs
--- a/Engine/Source/Programs/Horde/Plugins/Build/HordeServer.Build/Notifications/Sinks/SlackNotificationSink.cs
+++ b/Engine/Source/Programs/Horde/Plugins/Build/HordeServer.Build/Notifications/Sinks/SlackNotificationSink.cs
@@ -608,7 +608,6 @@ namespace HordeServer.Notifications.Sinks
 						foreach (string channel in _settings.JobNotificationChannel.Split(';'))
 						{
 							Uri jobStepLink = new Uri($"{_serverInfo.DashboardUrl}job/{job.Id}?step={step.Id}");
-							Uri jobStepLogLink = new Uri($"{_serverInfo.DashboardUrl}log/{step.LogId}");
 
 							SlackAttachment attachment = new SlackAttachment();
 							attachment.FallbackText = $"{streamConfig.Name} - {GetJobChangeText(job)} - {job.Name} - {node.Name} - {step.Outcome}";
