Horde: Grant claims to leases identifying them as running certain projects, streams, and templates.

diff --git a/Engine/Source/Programs/Horde/HordeServer.Shared/Utilities/HordeClaimTypes.cs b/Engine/Source/Programs/Horde/HordeServer.Shared/Utilities/HordeClaimTypes.cs
--- a/Engine/Source/Programs/Horde/HordeServer.Shared/Utilities/HordeClaimTypes.cs
+++ b/Engine/Source/Programs/Horde/HordeServer.Shared/Utilities/HordeClaimTypes.cs
@@ -89,6 +89,21 @@ namespace HordeServer.Utilities
 		/// </summary>
 		public const string Lease = Prefix + "lease";
 
+		/// <summary>
+		/// Claim identifying an agent as executing a lease in a particular stream
+		/// </summary>
+		public const string LeaseStream = Prefix + "lease-stream";
+
+		/// <summary>
+		/// Claim identifying an agent as executing a lease in a particular project
+		/// </summary>
+		public const string LeaseProject = Prefix + "lease-project";
+
+		/// <summary>
+		/// Claim identifying an agent as executing a lease in a particular job template
+		/// </summary>
+		public const string LeaseTemplate = Prefix + "lease-template";
+
 		/// <summary>
 		/// Claim allowing a certain namespace and/or path to be read from. See <see cref="WriteNamespace"/>.
 		/// </summary>
diff --git a/Engine/Source/Programs/Horde/Plugins/Build/HordeServer.Build/Jobs/JobTaskSource.cs b/Engine/Source/Programs/Horde/Plugins/Build/HordeServer.Build/Jobs/JobTaskSource.cs
--- a/Engine/Source/Programs/Horde/Plugins/Build/HordeServer.Build/Jobs/JobTaskSource.cs
+++ b/Engine/Source/Programs/Horde/Plugins/Build/HordeServer.Build/Jobs/JobTaskSource.cs
@@ -886,6 +886,9 @@ namespace HordeServer.Jobs
 			List<AclClaimConfig> claims = new List<AclClaimConfig>();
 			claims.Add(HordeClaims.AgentRoleClaim);
 			claims.Add(new AclClaimConfig(HordeClaimTypes.Lease, leaseId.ToString()));
+			claims.Add(new AclClaimConfig(HordeClaimTypes.LeaseStream, streamConfig.Id.ToString()));
+			claims.Add(new AclClaimConfig(HordeClaimTypes.LeaseProject, streamConfig.ProjectConfig.Id.ToString()));
+			claims.Add(new AclClaimConfig(HordeClaimTypes.LeaseTemplate, job.TemplateId.ToString()));
 			claims.AddRange(job.Claims);
 
 			// Encode the payload
