Horde: Scan uploaded blobs for file sizes and store them in the database.

#rnx

diff --git a/Engine/Source/Programs/Horde/Plugins/Storage/HordeServer.Storage/Storage/ObjectStores/AwsObjectStore.cs b/Engine/Source/Programs/Horde/Plugins/Storage/HordeServer.Storage/Storage/ObjectStores/AwsObjectStore.cs
--- a/Engine/Source/Programs/Horde/Plugins/Storage/HordeServer.Storage/Storage/ObjectStores/AwsObjectStore.cs
+++ b/Engine/Source/Programs/Horde/Plugins/Storage/HordeServer.Storage/Storage/ObjectStores/AwsObjectStore.cs
@@ -551,6 +551,27 @@ namespace HordeServer.Storage.ObjectStores
 			}
 		}
 
+		/// <inheritdoc/>
+		public async Task<long> GetSizeAsync(ObjectKey key, CancellationToken cancellationToken)
+		{
+			using TelemetrySpan span = OpenTelemetryTracers.Horde.StartActiveSpan($"{nameof(AwsObjectStore)}.{nameof(ExistsAsync)}");
+			span.SetAttribute("path", key.ToString());
+
+			try
+			{
+				GetObjectMetadataRequest request = new GetObjectMetadataRequest();
+				request.BucketName = _options.AwsBucketName;
+				request.Key = GetFullPath(key);
+
+				GetObjectMetadataResponse response = await _client.GetObjectMetadataAsync(request, cancellationToken);
+				return response.ContentLength;
+			}
+			catch (AmazonS3Exception ex) when (ex.StatusCode == System.Net.HttpStatusCode.NotFound)
+			{
+				return -1;
+			}
+		}
+
 		/// <inheritdoc/>
 		public void GetStats(StorageStats stats) { }
 	}
diff --git a/Engine/Source/Programs/Horde/Plugins/Storage/HordeServer.Storage/Storage/ObjectStores/AzureObjectStore.cs b/Engine/Source/Programs/Horde/Plugins/Storage/HordeServer.Storage/Storage/ObjectStores/AzureObjectStore.cs
--- a/Engine/Source/Programs/Horde/Plugins/Storage/HordeServer.Storage/Storage/ObjectStores/AzureObjectStore.cs
+++ b/Engine/Source/Programs/Horde/Plugins/Storage/HordeServer.Storage/Storage/ObjectStores/AzureObjectStore.cs
@@ -129,6 +129,27 @@ namespace HordeServer.Storage.ObjectStores
 			return await blob.ExistsAsync(cancellationToken);
 		}
 
+		/// <inheritdoc/>
+		public async Task<long> GetSizeAsync(ObjectKey key, CancellationToken cancellationToken)
+		{
+			if (!await _blobContainer.ExistsAsync(cancellationToken))
+			{
+				return -1;
+			}
+
+			BlobClient blob = _blobContainer.GetBlobClient(key.ToString());
+
+			try
+			{
+				BlobProperties properties = await blob.GetPropertiesAsync(cancellationToken: cancellationToken);
+				return properties.ContentLength;
+			}
+			catch(RequestFailedException)
+			{
+				return -1;
+			}
+		}
+
 		/// <inheritdoc/>
 		public async Task DeleteAsync(ObjectKey key, CancellationToken cancellationToken)
 		{
diff --git a/Engine/Source/Programs/Horde/Plugins/Storage/HordeServer.Storage/Storage/ObjectStores/ChainedObjectStore.cs b/Engine/Source/Programs/Horde/Plugins/Storage/HordeServer.Storage/Storage/ObjectStores/ChainedObjectStore.cs
--- a/Engine/Source/Programs/Horde/Plugins/Storage/HordeServer.Storage/Storage/ObjectStores/ChainedObjectStore.cs
+++ b/Engine/Source/Programs/Horde/Plugins/Storage/HordeServer.Storage/Storage/ObjectStores/ChainedObjectStore.cs
@@ -35,6 +35,10 @@ namespace HordeServer.Storage.Storage.ObjectStores
 		public async Task<bool> ExistsAsync(ObjectKey key, CancellationToken cancellationToken = default)
 			=> await _primary.ExistsAsync(key, cancellationToken) || await _secondary.ExistsAsync(key, cancellationToken);
 
+		/// <inheritdoc/>
+		public Task<long> GetSizeAsync(ObjectKey key, CancellationToken cancellationToken = default)
+			=> _primary.GetSizeAsync(key, cancellationToken);
+
 		/// <inheritdoc/>
 		public void GetStats(StorageStats stats)
 		{
diff --git a/Engine/Source/Programs/Horde/Plugins/Storage/HordeServer.Storage/Storage/StorageService.cs b/Engine/Source/Programs/Horde/Plugins/Storage/HordeServer.Storage/Storage/StorageService.cs
--- a/Engine/Source/Programs/Horde/Plugins/Storage/HordeServer.Storage/Storage/StorageService.cs
+++ b/Engine/Source/Programs/Horde/Plugins/Storage/HordeServer.Storage/Storage/StorageService.cs
@@ -265,6 +265,9 @@ namespace HordeServer.Storage
 			[BsonElement("del"), BsonIgnoreIfDefault]
 			public int GcVersion { get; set; }
 
+			[BsonElement("len"), BsonIgnoreIfDefault]
+			public long Length { get; set; }
+
 			[BsonIgnore]
 			public BlobLocator Locator => new BlobLocator(Path);
 
@@ -365,6 +368,19 @@ namespace HordeServer.Storage
 			}
 		}
 
+		[SingletonDocument("length-scan-state")]
+		class LengthScanState : SingletonBase
+		{
+			public ObjectId LastImportBlobInfoId { get; set; }
+			public bool Reset { get; set; }
+
+			public void DoReset()
+			{
+				LastImportBlobInfoId = ObjectId.Empty;
+				Reset = false;
+			}
+		}
+
 		class GcNamespaceState
 		{
 			public NamespaceId Id { get; set; }
@@ -390,6 +406,9 @@ namespace HordeServer.Storage
 		readonly SingletonDocument<GcState> _gcState;
 		readonly ITicker _gcTicker;
 
+		readonly SingletonDocument<LengthScanState> _lengthScanState;
+		readonly ITicker _lengthScanTicker;
+
 		readonly object _lockObject = new object();
 
 		string? _lastConfigRevision;
@@ -446,6 +465,9 @@ namespace HordeServer.Storage
 
 			_gcState = new SingletonDocument<GcState>(mongoService);
 			_gcTicker = clock.AddTicker("Storage:GC", TimeSpan.FromMinutes(5.0), TickGcAsync, logger);
+
+			_lengthScanState = new SingletonDocument<LengthScanState>(mongoService);
+			_lengthScanTicker = clock.AddTicker("Storage:LengthScan", TimeSpan.FromSeconds(1.0), TickLengthsAsync, logger);
 		}
 
 		static string GetFieldName<TClass>(Expression<Func<TClass, object?>> expr)
@@ -461,6 +483,7 @@ namespace HordeServer.Storage
 			await _blobTicker.DisposeAsync();
 			await _refTicker.DisposeAsync();
 			await _gcTicker.DisposeAsync();
+			await _lengthScanTicker.DisposeAsync();
 		}
 
 		internal static ObjectKey GetObjectKey(BlobLocator locator) => new ObjectKey($"{locator.Path}.blob");
@@ -492,11 +515,13 @@ namespace HordeServer.Storage
 			await _blobTicker.StartAsync();
 			await _refTicker.StartAsync();
 			await _gcTicker.StartAsync();
+			await _lengthScanTicker.StartAsync();
 		}
 
 		/// <inheritdoc/>
 		public async Task StopAsync(CancellationToken cancellationToken)
 		{
+			await _lengthScanTicker.StopAsync();
 			await _gcTicker.StopAsync();
 			await _refTicker.StopAsync();
 			await _blobTicker.StopAsync();
@@ -698,8 +723,6 @@ namespace HordeServer.Storage
 			// Get the current state of the storage system
 			State state = CreateState(_storageConfig.CurrentValue);
 
-			Dictionary<NamespaceId, BundleStorageNamespace> cachedClients = new();
-
 			long ingestedCount = 0;
 
 			// Compute missing import info, by searching for blobs with an ObjectId timestamp after the last import compute cycle
@@ -1300,5 +1323,81 @@ namespace HordeServer.Storage
 		}
 
 		#endregion
+
+		#region Length scan
+
+		/// <summary>
+		/// Scan new blobs for their sizes
+		/// </summary>
+		/// <param name="cancellationToken">Cancellation token for the operation</param>
+		async ValueTask TickLengthsAsync(CancellationToken cancellationToken)
+		{
+			using TelemetrySpan span = _tracer.StartActiveSpan($"{nameof(StorageService)}.{nameof(TickLengthsAsync)}");
+
+			State state = CreateState(_storageConfig.CurrentValue);
+
+			Channel<BlobInfo> channel = Channel.CreateBounded<BlobInfo>(new BoundedChannelOptions(1000));
+
+			await using AsyncPipeline pipeline = new AsyncPipeline(cancellationToken);
+			pipeline.AddTask(ctx => FindBlobsForLengthScanAsync(channel.Writer, ctx));
+			pipeline.AddTasks(8, channel.Reader, (entry, ctx) => FindBlobLengthsAsync(entry, state, ctx));
+			await pipeline.WaitForCompletionAsync();
+		}
+
+		async Task FindBlobsForLengthScanAsync(ChannelWriter<BlobInfo> writer, CancellationToken cancellationToken)
+		{
+			DateTime latestTimeUtc = _clock.UtcNow - TimeSpan.FromMinutes(30.0);
+			ObjectId latestInfoId = ObjectId.GenerateNewId(latestTimeUtc);
+
+			long scannedCount = 0;
+			LengthScanState lengthScanState = await _lengthScanState.GetAsync(cancellationToken);
+
+			for(; ;)
+			{
+				// Reset the state
+				if (lengthScanState.Reset)
+				{
+					_logger.LogInformation("Resetting scan for blob lengths...");
+					lengthScanState = await _lengthScanState.UpdateAsync(x => x.DoReset(), cancellationToken);
+				}
+
+				// Fetch the next batch of blobs
+				List<BlobInfo> current = await _blobCollection.Find(x => x.Id > lengthScanState.LastImportBlobInfoId && x.Id < latestInfoId).SortBy(x => x.Id).Limit(2000).ToListAsync(cancellationToken);
+				if (current.Count == 0)
+				{
+					break;
+				}
+
+				// Add a check record for each new blob
+				_logger.LogDebug("Finding length for {NumBlobs} blobs ({FirstId} to {LastId})", current.Count, current[0].Id, current[^1].Id);
+				foreach (BlobInfo blobInfo in current)
+				{
+					await writer.WriteAsync(blobInfo, cancellationToken);
+				}
+
+				// Update the last imported blob id
+				lengthScanState = await _lengthScanState.UpdateAsync(state => state.LastImportBlobInfoId = current[^1].Id, cancellationToken);
+				scannedCount += current.Count;
+			}
+
+			_logger.LogInformation("Added lengths for {NumBlobs} blobs", scannedCount);
+			writer.Complete();
+		}
+
+		async Task FindBlobLengthsAsync(BlobInfo blobInfo, State state, CancellationToken cancellationToken)
+		{
+			NamespaceInfo? namespaceInfo;
+			if (state.Namespaces.TryGetValue(blobInfo.NamespaceId, out namespaceInfo))
+			{
+				ObjectKey key = GetObjectKey(blobInfo.Locator);
+				long length = await namespaceInfo.Store.GetSizeAsync(key, cancellationToken);
+
+				FilterDefinition<BlobInfo> filter = Builders<BlobInfo>.Filter.Eq(x => x.Id, blobInfo.Id);
+				UpdateDefinition<BlobInfo> update = Builders<BlobInfo>.Update.Set(x => x.Length, length);
+				await _blobCollection.UpdateOneAsync(filter, update, cancellationToken: cancellationToken);
+			}
+		}
+
+		#endregion
 	}
 }
diff --git a/Engine/Source/Programs/Shared/EpicGames.Horde/Storage/IObjectStore.cs b/Engine/Source/Programs/Shared/EpicGames.Horde/Storage/IObjectStore.cs
--- a/Engine/Source/Programs/Shared/EpicGames.Horde/Storage/IObjectStore.cs
+++ b/Engine/Source/Programs/Shared/EpicGames.Horde/Storage/IObjectStore.cs
@@ -52,9 +52,17 @@ namespace EpicGames.Horde.Storage
 		/// </summary>
 		/// <param name="key">Relative path within the bucket</param>
 		/// <param name="cancellationToken">Cancellation token for the operation</param>
-		/// <returns></returns>
+		/// <returns>True if the object exists</returns>
 		Task<bool> ExistsAsync(ObjectKey key, CancellationToken cancellationToken = default);
 
+		/// <summary>
+		/// Gets the size of a particular object
+		/// </summary>
+		/// <param name="key">Relative path within the bucket</param>
+		/// <param name="cancellationToken">Cancellation token for the operation</param>
+		/// <returns>Size of the object, or -1 if it does not exist</returns>
+		Task<long> GetSizeAsync(ObjectKey key, CancellationToken cancellationToken = default);
+
 		/// <summary>
 		/// Deletes a file with the given path
 		/// </summary>
@@ -110,6 +118,7 @@ namespace EpicGames.Horde.Storage
 			public Task WriteAsync(ObjectKey key, Stream stream, CancellationToken cancellationToken) => _inner.WriteAsync(key, stream, cancellationToken);
 			public Task DeleteAsync(ObjectKey path, CancellationToken cancellationToken) => _inner.DeleteAsync(path, cancellationToken);
 			public Task<bool> ExistsAsync(ObjectKey path, CancellationToken cancellationToken) => _inner.ExistsAsync(path, cancellationToken);
+			public Task<long> GetSizeAsync(ObjectKey path, CancellationToken cancellationToken) => _inner.GetSizeAsync(path, cancellationToken);
 			public ValueTask<Uri?> TryGetReadRedirectAsync(ObjectKey path, CancellationToken cancellationToken = default) => _inner.TryGetReadRedirectAsync(path, cancellationToken);
 			public ValueTask<Uri?> TryGetWriteRedirectAsync(ObjectKey path, CancellationToken cancellationToken = default) => _inner.TryGetWriteRedirectAsync(path, cancellationToken);
 			public void GetStats(StorageStats stats) => _inner.GetStats(stats);
diff --git a/Engine/Source/Programs/Shared/EpicGames.Horde/Storage/ObjectStores/FileObjectStore.cs b/Engine/Source/Programs/Shared/EpicGames.Horde/Storage/ObjectStores/FileObjectStore.cs
--- a/Engine/Source/Programs/Shared/EpicGames.Horde/Storage/ObjectStores/FileObjectStore.cs
+++ b/Engine/Source/Programs/Shared/EpicGames.Horde/Storage/ObjectStores/FileObjectStore.cs
@@ -107,6 +107,20 @@ namespace EpicGames.Horde.Storage.ObjectStores
 			return Task.FromResult(FileReference.Exists(location));
 		}
 
+		/// <inheritdoc/>
+		public Task<long> GetSizeAsync(ObjectKey key, CancellationToken cancellationToken)
+		{
+			FileInfo info = GetBlobFile(key).ToFileInfo();
+			if (info.Exists)
+			{
+				return Task.FromResult<long>(info.Length);
+			}
+			else
+			{
+				return Task.FromResult<long>(-1);
+			}
+		}
+
 		/// <summary>
 		/// Delete a file from the store
 		/// </summary>
diff --git a/Engine/Source/Programs/Shared/EpicGames.Horde/Storage/ObjectStores/HttpObjectStore.cs b/Engine/Source/Programs/Shared/EpicGames.Horde/Storage/ObjectStores/HttpObjectStore.cs
--- a/Engine/Source/Programs/Shared/EpicGames.Horde/Storage/ObjectStores/HttpObjectStore.cs
+++ b/Engine/Source/Programs/Shared/EpicGames.Horde/Storage/ObjectStores/HttpObjectStore.cs
@@ -53,6 +53,24 @@ namespace EpicGames.Horde.Storage.ObjectStores
 			return response.IsSuccessStatusCode;
 		}
 
+		/// <inheritdoc/>
+		public async Task<long> GetSizeAsync(ObjectKey key, CancellationToken cancellationToken = default)
+		{
+			using HttpClient client = _createHttpClient();
+
+			using HttpRequestMessage request = new HttpRequestMessage(HttpMethod.Head, GetUri(key));
+			using HttpResponseMessage response = await client.SendAsync(request, cancellationToken);
+
+			if (response.IsSuccessStatusCode)
+			{
+				return response.Content.Headers.ContentLength ?? throw new NotSupportedException("Missing Content-Length header");
+			}
+			else
+			{
+				return -1;
+			}
+		}
+
 		/// <inheritdoc/>
 		public void GetStats(StorageStats stats)
 		{
diff --git a/Engine/Source/Programs/Shared/EpicGames.Horde/Storage/ObjectStores/MemoryObjectStore.cs b/Engine/Source/Programs/Shared/EpicGames.Horde/Storage/ObjectStores/MemoryObjectStore.cs
--- a/Engine/Source/Programs/Shared/EpicGames.Horde/Storage/ObjectStores/MemoryObjectStore.cs
+++ b/Engine/Source/Programs/Shared/EpicGames.Horde/Storage/ObjectStores/MemoryObjectStore.cs
@@ -63,6 +63,20 @@ namespace EpicGames.Horde.Storage.ObjectStores
 			return Task.FromResult(_keyToData.ContainsKey(key));
 		}
 
+		/// <inheritdoc/>
+		public Task<long> GetSizeAsync(ObjectKey key, CancellationToken cancellationToken)
+		{
+			byte[]? data;
+			if (_keyToData.TryGetValue(key, out data))
+			{
+				return Task.FromResult<long>(data.Length);
+			}
+			else
+			{
+				return Task.FromResult<long>(-1);
+			}
+		}
+
 		/// <inheritdoc/>
 		public Task DeleteAsync(ObjectKey key, CancellationToken cancellationToken)
 		{
diff --git a/Engine/Source/Programs/Shared/EpicGames.Horde/Storage/ObjectStores/PrefixedObjectStore.cs b/Engine/Source/Programs/Shared/EpicGames.Horde/Storage/ObjectStores/PrefixedObjectStore.cs
--- a/Engine/Source/Programs/Shared/EpicGames.Horde/Storage/ObjectStores/PrefixedObjectStore.cs
+++ b/Engine/Source/Programs/Shared/EpicGames.Horde/Storage/ObjectStores/PrefixedObjectStore.cs
@@ -40,6 +40,9 @@ namespace EpicGames.Horde.Storage.ObjectStores
 		/// <inheritdoc/>
 		public Task<bool> ExistsAsync(ObjectKey locator, CancellationToken cancellationToken = default) => _inner.ExistsAsync(GetKeyWithPrefix(locator), cancellationToken);
 
+		/// <inheritdoc/>
+		public Task<long> GetSizeAsync(ObjectKey locator, CancellationToken cancellationToken = default) => _inner.GetSizeAsync(GetKeyWithPrefix(locator), cancellationToken);
+
 		/// <inheritdoc/>
 		public Task<Stream> OpenAsync(ObjectKey locator, int offset, int? length, CancellationToken cancellationToken = default) => _inner.OpenAsync(GetKeyWithPrefix(locator), offset, length, cancellationToken);
 
diff --git a/Engine/Source/Programs/Shared/EpicGames.Horde/Storage/StorageBackendCache.cs b/Engine/Source/Programs/Shared/EpicGames.Horde/Storage/StorageBackendCache.cs
--- a/Engine/Source/Programs/Shared/EpicGames.Horde/Storage/StorageBackendCache.cs
+++ b/Engine/Source/Programs/Shared/EpicGames.Horde/Storage/StorageBackendCache.cs
@@ -84,6 +84,7 @@ namespace EpicGames.Horde.Storage
 			public Task WriteAsync(ObjectKey key, Stream stream, CancellationToken cancellationToken = default) => _inner.WriteAsync(key, stream, cancellationToken);
 			public Task DeleteAsync(ObjectKey key, CancellationToken cancellationToken = default) => _inner.DeleteAsync(key, cancellationToken);
 			public Task<bool> ExistsAsync(ObjectKey key, CancellationToken cancellationToken = default) => _inner.ExistsAsync(key, cancellationToken);
+			public Task<long> GetSizeAsync(ObjectKey key, CancellationToken cancellationToken = default) => _inner.GetSizeAsync(key, cancellationToken);
 			public ValueTask<Uri?> TryGetReadRedirectAsync(ObjectKey key, CancellationToken cancellationToken = default) => _inner.TryGetReadRedirectAsync(key, cancellationToken);
 			public ValueTask<Uri?> TryGetWriteRedirectAsync(ObjectKey key, CancellationToken cancellationToken = default) => _inner.TryGetWriteRedirectAsync(key, cancellationToken);
 
