Horde: Updating schema documentation

#rnx

diff --git a/Engine/Source/Programs/Horde/Docs/Config/Schema/Streams.md b/Engine/Source/Programs/Horde/Docs/Config/Schema/Streams.md
--- a/Engine/Source/Programs/Horde/Docs/Config/Schema/Streams.md
+++ b/Engine/Source/Programs/Horde/Docs/Config/Schema/Streams.md
@@ -11,6 +11,7 @@ Name | Description
 `include` | [ConfigInclude](#configinclude)`[]`<br>Includes for other configuration files
 `macros` | [ConfigMacro](#configmacro)`[]`<br>Macros within this stream
 `name` | `string`<br>Name of the stream
+`enginePath` | `string`<br>Path to the engine directory within the workspace. Used for launching UAT.
 `clusterName` | `string`<br>The perforce cluster containing the stream
 `order` | `integer`<br>Order for this stream
 `initialAgentType` | `string`<br>Default initial agent type for templates
