Horde: Additional logging for storage service enumerating blob sizes.

#rnx

diff --git a/Engine/Source/Programs/Horde/Plugins/Storage/HordeServer.Storage/Storage/StorageService.cs b/Engine/Source/Programs/Horde/Plugins/Storage/HordeServer.Storage/Storage/StorageService.cs
--- a/Engine/Source/Programs/Horde/Plugins/Storage/HordeServer.Storage/Storage/StorageService.cs
+++ b/Engine/Source/Programs/Horde/Plugins/Storage/HordeServer.Storage/Storage/StorageService.cs
@@ -1411,6 +1411,7 @@ namespace HordeServer.Storage
 			{
 				ObjectKey key = GetObjectKey(blobInfo.Locator);
 				long length = await namespaceInfo.Store.GetSizeAsync(key, cancellationToken);
+				_logger.LogDebug("Length of blob {BlobId} ({Key}): {Length}", blobInfo.Id, key, length);
 
 				FilterDefinition<BlobInfo> filter = Builders<BlobInfo>.Filter.Eq(x => x.Id, blobInfo.Id);
 				UpdateDefinition<BlobInfo> update = Builders<BlobInfo>.Update.Set(x => x.Length, length);
