Horde: Prevent infinite loop when cancelling during agent registration

diff --git a/Engine/Source/Programs/Horde/HordeAgent/Services/SessionFactory.cs b/Engine/Source/Programs/Horde/HordeAgent/Services/SessionFactory.cs
--- a/Engine/Source/Programs/Horde/HordeAgent/Services/SessionFactory.cs
+++ b/Engine/Source/Programs/Horde/HordeAgent/Services/SessionFactory.cs
@@ -306,9 +306,19 @@ namespace HordeAgent.Services
 				}
 				catch (RpcException ex)
 				{
+					if (cancellationToken.IsCancellationRequested)
+					{
+						break;
+					}
+					
+					TimeSpan backoffDelay = TimeSpan.FromSeconds(15);
 					logger.LogWarning(ex, "Exception in RPC: {Message}", ex.Message);
+					logger.LogInformation("Waiting {BackoffDelay} seconds before retrying...", backoffDelay.TotalSeconds);
+					await Task.Delay(backoffDelay, CancellationToken.None);
 				}
 			}
+			
+			throw new Exception("Unable to register agent");
 		}
 
 		static string? GetProperty(RpcAgentCapabilities capabilities, string name)
