Horde: Add schedule auditing
#skipci

diff --git a/Engine/Source/Programs/Horde/HordeDashboard/src/App.tsx b/Engine/Source/Programs/Horde/HordeDashboard/src/App.tsx
--- a/Engine/Source/Programs/Horde/HordeDashboard/src/App.tsx
+++ b/Engine/Source/Programs/Horde/HordeDashboard/src/App.tsx
@@ -160,6 +160,7 @@ const Main: React.FC = () => {
                { path: "devices", element: <DeviceView /> },
                { path: "audit/agent/:agentId", element: <AuditLogView /> },
                { path: "audit/issue/:issueId", element: <AuditLogView /> },
+               { path: "audit/template/:streamId/:templateId", element: <AuditLogView /> },
                { path: "automation", element: <AutomationView /> },
                { path: "tools", element: <ToolView /> },
                { path: "lease/:leaseId", element: <DebugView /> },
diff --git a/Engine/Source/Programs/Horde/HordeDashboard/src/backend/Backend.ts b/Engine/Source/Programs/Horde/HordeDashboard/src/backend/Backend.ts
--- a/Engine/Source/Programs/Horde/HordeDashboard/src/backend/Backend.ts
+++ b/Engine/Source/Programs/Horde/HordeDashboard/src/backend/Backend.ts
@@ -958,6 +958,14 @@ export class Backend {
         });
     }
 
+    getTemplateHistory(streamId: string, templateId: string, query: AuditLogQuery): Promise<AuditLogEntry[]> {
+        return new Promise<AuditLogEntry[]>((resolve, reject) => {
+            this._fetch.get(`/api/v1/streams/${encodeURIComponent(streamId)}/templates/${encodeURIComponent(templateId)}/history`, { params: query }).then((response) => {
+                const history = (response.data?.entries ?? []) as AuditLogEntry[];
+                resolve(history);
+            }).catch(reason => { reject(reason); });
+        });
+    }
 
     getIssuesV2(queryIn?: IssueQueryV2): Promise<FindIssueResponse[]> {
 
diff --git a/Engine/Source/Programs/Horde/HordeDashboard/src/components/AuditLog.tsx b/Engine/Source/Programs/Horde/HordeDashboard/src/components/AuditLog.tsx
--- a/Engine/Source/Programs/Horde/HordeDashboard/src/components/AuditLog.tsx
+++ b/Engine/Source/Programs/Horde/HordeDashboard/src/components/AuditLog.tsx
@@ -20,6 +20,7 @@ import { getLogStyles, logMetricNormal } from "./LogStyle";
 import { TopNav } from "./TopNav";
 import { getHordeStyling } from "../styles/Styles";
 import { getHordeTheme } from "../styles/theme";
+import { projectStore } from "horde/backend/ProjectStore";
 
 
 let _auditStyleNormal: any;
@@ -32,17 +33,22 @@ const getAuditStyleNormal = () => {
    const auditStyleNormal = _auditStyleNormal ?? mergeStyleSets(logStyleNormal, {
 
       logLine: [
-          {
+         {
             selectors: {
                '&:hover': { background: theme.palette.neutralLight }
-             },
-          }
+            },
+         }
       ]
    });
 
    _auditStyleNormal = auditStyleNormal;
 
-   return auditStyleNormal;  
+   return auditStyleNormal;
+}
+
+export type TemplateAudit = {
+   streamId: string;
+   templateId: string;
 }
 
 
@@ -59,7 +65,8 @@ class AuditLogHandler {
 
       this.agentId = undefined;
       this.issueId = undefined;
-      
+      this.templateAudit = undefined;
+
       this.timeSelectKey = "time_live_tail";
 
       clearTimeout(this.timeoutId);
@@ -109,15 +116,34 @@ class AuditLogHandler {
       this.update();
    }
 
+   setTemplateAudit(templateAudit?: TemplateAudit) {
 
-   get isAgentLog():boolean {
+      if (!templateAudit) {
+         this.clear();
+         return;
+      }
+
+      if (this.templateAudit?.streamId === templateAudit.streamId && this.templateAudit.templateId === templateAudit.templateId) {
+         return;
+      }
+
+      this.templateAudit = templateAudit;
+      this.update();
+   }
+
+
+   get isAgentLog(): boolean {
       return !!this.agentId;
    }
 
-   get isIssueLog():boolean {
+   get isIssueLog(): boolean {
       return !!this.issueId;
    }
 
+   get isTemplateAuditLog(): boolean {
+      return !!this.templateAudit;
+   }
+
    async update() {
 
       clearTimeout(this.timeoutId);
@@ -129,7 +155,7 @@ class AuditLogHandler {
          if (this.updating) {
             return;
          }
-   
+
       }
 
       // cancel any pending        
@@ -137,7 +163,7 @@ class AuditLogHandler {
          this.canceled.add(i);
       }
 
-      if (!this.agentId && !this.issueId) {
+      if (!this.agentId && !this.issueId && !this.templateAudit) {
          return;
       }
 
@@ -152,15 +178,17 @@ class AuditLogHandler {
          let maxTime = this.maxDate;
 
          if (this.tailing) {
-            
+
             if (this.issueId) {
                minTime = undefined;
             } else if (this.agentId) {
                // live tail of the last 4 days
                // We need to optimize the agent history endpoint, this was an attempt to limit results though doesn't work 
                minTime = undefined; //new Date(new Date().valueOf() - (60 * 24 * 4 * 60000));
+            } else if (this.templateAudit) {
+               minTime = undefined;
             }
-            
+
             maxTime = new Date();
          }
 
@@ -169,11 +197,15 @@ class AuditLogHandler {
          if (this.agentId) {
             entries = await backend.getAgentHistory(this.agentId, { minTime: minTime?.toISOString(), maxTime: maxTime?.toISOString(), count: 8192 * 2 });
          }
-         
+
          if (this.issueId) {
             entries = await backend.getIssueHistory(this.issueId, { minTime: minTime?.toISOString(), maxTime: maxTime?.toISOString(), count: 8192 * 2 });
          }
 
+         if (this.templateAudit) {
+            entries = await backend.getTemplateHistory(this.templateAudit.streamId, this.templateAudit.templateId, { minTime: minTime?.toISOString(), maxTime: maxTime?.toISOString(), count: 8192 * 2 });
+         }
+
          // check for canceled during graph request
          if (this.canceled.has(cancelId)) {
             return;
@@ -240,6 +272,7 @@ class AuditLogHandler {
 
    agentId?: string;
    issueId?: string;
+   templateAudit?: TemplateAudit;
 
    timeSelectKey?: string;
 
@@ -310,7 +343,7 @@ const timeSelections: TimeSelection[] = [
 
 const selection = new Selection({ selectionMode: SelectionMode.multiple });
 
-export const AuditLogPanel: React.FC<{ agentId?: string, issueId?: string }> = observer(({ agentId, issueId }) => {
+export const AuditLogPanel: React.FC<{ agentId?: string, issueId?: string, templateAudit?: TemplateAudit }> = observer(({ agentId, issueId, templateAudit }) => {
 
    const windowSize = useWindowSize();
    const [showDatePicker, setShowDatePicker] = useState(false);
@@ -342,7 +375,11 @@ export const AuditLogPanel: React.FC<{ agentId?: string, issueId?: string }> = o
       handler.setIssue(issueId);
    }
 
-   if (!agentId && !issueId) {
+   if (templateAudit) {
+      handler.setTemplateAudit(templateAudit);
+   }
+
+   if (!agentId && !issueId && !templateAudit) {
       return null;
    }
 
@@ -439,6 +476,19 @@ export const AuditLogPanel: React.FC<{ agentId?: string, issueId?: string }> = o
 
    }
 
+   if (templateAudit) {
+      const stream = projectStore.streamById(templateAudit.streamId)!;
+      const templateName = stream?.templates.find(t => t.id === templateAudit.templateId)?.name;
+
+      crumbItems = [
+         {
+            text: `Template`,
+         },
+         {
+            text: `Audit - ${stream.fullname ?? templateAudit.streamId} - ${templateName ?? templateAudit.templateId}`,
+         }
+      ];
+   }
 
    const onRenderCell = (item?: AuditLogItem, index?: number, isScrolling?: boolean): JSX.Element => {
       const entry = item!.entry;
@@ -453,17 +503,17 @@ export const AuditLogPanel: React.FC<{ agentId?: string, issueId?: string }> = o
       let timestamp = `[${tm.format(format)}]`;
 
       return (
-            <Stack className={auditStyleNormal.logLine} key={`key_log_line_${item?.entry.time}`} style={{ width: "100%", height: logMetricNormal.lineHeight }}>
-               <div style={{ position: "relative" }}>
-                  <Stack tokens={{ childrenGap: 8 }} horizontal disableShrink={true}>
-                     <Stack horizontal disableShrink={true} >
-                        <Stack styles={{ root: { width: 140, whiteSpace: "nowrap", fontSize: logMetricNormal.fontSize, userSelect: "none" } }}> {timestamp}</Stack>
-                        <Stack styles={{ root: { color: "#8a8a8a", width: 92, paddingRight: 8, whiteSpace: "nowrap", textAlign: "right", fontSize: logMetricNormal.fontSize, userSelect: "none" } }}> [{item!.type}]</Stack>
-                        <div className={auditStyleNormal.logLineOuter}> <Stack styles={{ root: { paddingLeft: 8, paddingRight: 8 } }}> {renderAuditEntry(entry, search)}</Stack></div>
-                     </Stack>
+         <Stack className={auditStyleNormal.logLine} key={`key_log_line_${item?.entry.time}`} style={{ width: "100%", height: logMetricNormal.lineHeight }}>
+            <div style={{ position: "relative" }}>
+               <Stack tokens={{ childrenGap: 8 }} horizontal disableShrink={true}>
+                  <Stack horizontal disableShrink={true} >
+                     <Stack styles={{ root: { width: 140, whiteSpace: "nowrap", fontSize: logMetricNormal.fontSize, userSelect: "none" } }}> {timestamp}</Stack>
+                     <Stack styles={{ root: { color: "#8a8a8a", width: 92, paddingRight: 8, whiteSpace: "nowrap", textAlign: "right", fontSize: logMetricNormal.fontSize, userSelect: "none" } }}> [{item!.type}]</Stack>
+                     <div className={auditStyleNormal.logLineOuter}> <Stack styles={{ root: { paddingLeft: 8, paddingRight: 8 } }}> {renderAuditEntry(entry, search)}</Stack></div>
                   </Stack>
-               </div>
-            </Stack>
+               </Stack>
+            </div>
+         </Stack>
       );
    }
 
@@ -482,7 +532,7 @@ export const AuditLogPanel: React.FC<{ agentId?: string, issueId?: string }> = o
       <Breadcrumbs items={crumbItems} />
       <Stack tokens={{ childrenGap: 12 }}>
          <Stack horizontal>
-            <div key={`windowsize_logview1_${windowSize.width}_${windowSize.height}`} style={{ width: vw / 2 - (1440/2) - 230, flexShrink: 0, backgroundColor: modeColors.background }} />
+            <div key={`windowsize_logview1_${windowSize.width}_${windowSize.height}`} style={{ width: vw / 2 - (1440 / 2) - 230, flexShrink: 0, backgroundColor: modeColors.background }} />
             <Stack tokens={{ childrenGap: 0 }} styles={{ root: { backgroundColor: modeColors.background, margin: "auto", paddingTop: 12, paddingRight: 10 } }}>
                <Stack horizontal styles={{ root: { paddingLeft: 0, paddingBottom: 4, paddingRight: 12, width: 1440 } }}>
                   <Stack horizontal tokens={{ childrenGap: 12 }}>
@@ -581,11 +631,11 @@ export const AuditLogPanel: React.FC<{ agentId?: string, issueId?: string }> = o
             <Stack tokens={{ childrenGap: 0 }}>
                {showDatePicker && < DateTimeRange onChange={(minDate, maxDate) => { handler.setTimeRange(minDate, maxDate); setShowDatePicker(false) }} onDismiss={() => { handler.setTimeSelection({ text: "Live Tail", key: "time_live_tail", minutes: 0 }); setShowDatePicker(false) }} />}
                {viewAgent && <HistoryModal agentId={agentId} onDismiss={() => { setViewAgent(false) }} />}
-               {viewIssue && <IssueModalV2 issueId={issueId} onCloseExternal={() => { setViewIssue(false);} } popHistoryOnClose={false} />}
+               {viewIssue && <IssueModalV2 issueId={issueId} onCloseExternal={() => { setViewIssue(false); }} popHistoryOnClose={false} />}
                <FocusZone direction={FocusZoneDirection.vertical} isInnerZoneKeystroke={() => { return true; }} defaultActiveElement="#LogList" style={{ padding: 0, margin: 0 }} >
                   <div className={auditStyleNormal.container} style={{ height: 'calc(100vh - 260px)', position: 'relative' }} data-is-scrollable={true}>
                      <Stack horizontal>
-                        <div key={`windowsize_logview2_${windowSize.width}_${windowSize.height}`} style={{ width: vw / 2 - (1440/2) - 24, flexShrink: 0 }} />
+                        <div key={`windowsize_logview2_${windowSize.width}_${windowSize.height}`} style={{ width: vw / 2 - (1440 / 2) - 24, flexShrink: 0 }} />
                         <Stack styles={{ root: { backgroundColor: modeColors.background, paddingLeft: "0px", paddingRight: "0px" } }}>
                            {!handler.haveUpdated && <Spinner size={SpinnerSize.large} />}
                            {!!handler.haveUpdated && !logItems.length && <Stack style={{ paddingLeft: 0 }}><Text variant="mediumPlus">No audit entries found</Text></Stack>}
@@ -617,12 +667,22 @@ export const AuditLogView: React.FC = () => {
 
    const { hordeClasses } = getHordeStyling();
 
-   const { agentId, issueId } = useParams<{ agentId: string, issueId: string }>();
+   const { agentId, issueId, streamId, templateId } = useParams<{ agentId: string, issueId: string, streamId: string, templateId: string }>();
+
+   let templateAudit: TemplateAudit | undefined;
+   if (!!streamId && !!templateId) {
+      templateAudit = {
+         streamId: streamId,
+         templateId: templateId
+      }
+   }
+
    return <Stack className={hordeClasses.horde}>
       <TopNav />
       <Stack>
          {!!agentId && <AuditLogPanel agentId={agentId} />}
          {!!issueId && <AuditLogPanel issueId={issueId} />}
+         {!!templateAudit && <AuditLogPanel templateAudit={templateAudit} />}
       </Stack>
    </Stack>
 
@@ -636,7 +696,8 @@ export const AuditLogView: React.FC = () => {
 enum AuditLogType {
    Unknown,
    Agent,
-   Issue
+   Issue,
+   Template
 }
 
 type AuditProperty = {
@@ -666,6 +727,10 @@ class AuditLine {
          instance = new IssueAuditLine(logType, entry);
       }
 
+      if (logType === AuditLogType.Template) {
+         instance = new TemplateAuditLine(logType, entry);
+      }
+
       if (instance) {
          instance.search = search;
          instance.process();
@@ -705,7 +770,7 @@ class AuditLine {
 
          }
 
-         if (typeof (property) === "string" || typeof (property) === "number" || typeof(property) === "boolean") {
+         if (typeof (property) === "string" || typeof (property) === "number" || typeof (property) === "boolean") {
 
             this.properties.set(m, {
                tag: m,
@@ -783,7 +848,7 @@ class AuditLine {
          // LogId
          if (tag === "LogId") {
 
-            const logId = property.value as string;            
+            const logId = property.value as string;
 
             if (logId) {
                let to = `/log/${logId}`;
@@ -859,7 +924,7 @@ class AuditLine {
             return t;
          }
 
-         const rtags:any = [];
+         const rtags: any = [];
 
          const key = `log_line_${idx}_${index}_fragment`;
 
@@ -900,6 +965,9 @@ class IssueAuditLine extends AuditLine {
 
 }
 
+class TemplateAuditLine extends AuditLine {
+
+}
 
 // Line Rendering -------------------------------------------------------------------------------------------------------------------
 
@@ -916,6 +984,10 @@ const renderAuditEntry = (entry: AuditLogEntry, search?: string) => {
       type = AuditLogType.Issue;
    }
 
+   if (handler.isTemplateAuditLog) {
+      type = AuditLogType.Template;
+   }
+
    if (type === AuditLogType.Unknown) {
       return <div>Unknown Log Type</div>
    }
diff --git a/Engine/Source/Programs/Horde/HordeDashboard/src/components/SchedulePane.tsx b/Engine/Source/Programs/Horde/HordeDashboard/src/components/SchedulePane.tsx
--- a/Engine/Source/Programs/Horde/HordeDashboard/src/components/SchedulePane.tsx
+++ b/Engine/Source/Programs/Horde/HordeDashboard/src/components/SchedulePane.tsx
@@ -3,7 +3,7 @@ import { mergeStyleSets, Stack } from '@fluentui/react';
 import { DetailsList, DetailsListLayoutMode, DetailsRow, IColumn, IDetailsListProps, SelectionMode } from '@fluentui/react/lib/DetailsList';
 import moment from 'moment-timezone';
 import React, { useState } from 'react';
-import { Link } from 'react-router-dom';
+import { Link, useNavigate } from 'react-router-dom';
 import backend from '../backend';
 import { GetJobResponse, GetTemplateRefResponse, JobState } from '../backend/Api';
 import dashboard from '../backend/Dashboard';
@@ -27,9 +27,10 @@ type ScheduleItem = {
 
 const scheduledJobs = new Map<string, GetJobResponse | null>();
 
-export const SchedulePane: React.FC<{ templates: GetTemplateRefResponse[] }> = ({ templates }) => {
+export const SchedulePane: React.FC<{ streamId: string, templates: GetTemplateRefResponse[] }> = ({ streamId, templates }) => {
 
    let [updated, setUpdated] = useState(new Date().getTime());
+   const navigate = useNavigate();
 
    const ids = new Set<string>();
 
@@ -76,7 +77,8 @@ export const SchedulePane: React.FC<{ templates: GetTemplateRefResponse[] }> = (
       { key: 'schedule_column7', name: 'time_4', minWidth: 100, maxWidth: 100, isResizable: false },
       { key: 'schedule_column8', name: 'time_5', minWidth: 100, maxWidth: 100, isResizable: false },
       { key: 'schedule_column9', name: 'time_6', minWidth: 100, maxWidth: 100, isResizable: false },
-      { key: 'schedule_column10', name: 'space', minWidth: 2, isResizable: false }
+      { key: 'schedule_column10', name: 'audit', minWidth: 64, maxWidth: 64, isResizable: false },
+      { key: 'schedule_column11', name: 'space', minWidth: 2, isResizable: false }
 
    ];
 
@@ -94,7 +96,7 @@ export const SchedulePane: React.FC<{ templates: GetTemplateRefResponse[] }> = (
          let gateTemplate = templates.find(t => t.id === templateId);
 
 
-         return <Stack horizontalAlign="start" verticalFill verticalAlign="center" tokens={{ childrenGap: 4 }}>
+         return <Stack horizontalAlign="start" verticalFill verticalAlign="center" tokens={{ childrenGap: 4 }} style={{paddingLeft: 12}}>
             <Stack style={{ fontWeight: 600 }}>{name}</Stack>
             {!!gateTemplate && <Stack tokens={{ childrenGap: 4 }}>
                <Stack horizontal tokens={{ childrenGap: 8 }}>
@@ -110,6 +112,16 @@ export const SchedulePane: React.FC<{ templates: GetTemplateRefResponse[] }> = (
          </Stack>;
       }
 
+      if (name === "audit") {
+         return <Link to={`/audit/template/${encodeURIComponent(streamId)}/${encodeURIComponent(template.id)}`}>
+            <Stack horizontalAlign="end" verticalFill verticalAlign="center" tokens={{ childrenGap: 4 }}>
+            <Stack tokens={{ childrenGap: 4 }}>
+               <Stack >View Log</Stack>
+            </Stack>
+         </Stack>
+         </Link>
+      }
+
       if (name.startsWith("time_")) {
 
          const index = parseInt(name.slice(5));
@@ -174,10 +186,10 @@ export const SchedulePane: React.FC<{ templates: GetTemplateRefResponse[] }> = (
 
          let background: string | undefined;
          if (props.itemIndex % 2 === 0) {
-            background  =  dashboard.darktheme ? "#1D2021" : "#FAF9F9";
+            background = dashboard.darktheme ? "#1D2021" : "#FAF9F9";
          }
 
-         return <DetailsRow {...nprops} styles={{root: {background: background}}}/>
+         return <DetailsRow {...nprops} styles={{ root: { background: background } }} />
 
       }
       return null;
diff --git a/Engine/Source/Programs/Horde/HordeDashboard/src/components/StreamSummary.tsx b/Engine/Source/Programs/Horde/HordeDashboard/src/components/StreamSummary.tsx
--- a/Engine/Source/Programs/Horde/HordeDashboard/src/components/StreamSummary.tsx
+++ b/Engine/Source/Programs/Horde/HordeDashboard/src/components/StreamSummary.tsx
@@ -778,7 +778,7 @@ const SchedulePanel: React.FC = observer(() => {
          <Stack tokens={{ childrenGap: 12 }}>
             <Text variant="mediumPlus" styles={{ root: { fontFamily: "Horde Open Sans SemiBold" } }}>Schedule</Text>
             <Stack styles={{ root: { paddingLeft: 4, paddingRight: 0, paddingTop: 8, paddingBottom: 4 } }}>
-               <SchedulePane templates={templates} />
+               <SchedulePane streamId={stream.id} templates={templates} />
             </Stack>
          </Stack>
       </Stack>
