Horde: Manually merging changes due to RM blockage.

#rnx

diff --git a/Engine/Source/Programs/Horde/HordeDashboard/src/components/AuditLog.tsx b/Engine/Source/Programs/Horde/HordeDashboard/src/components/AuditLog.tsx
--- a/Engine/Source/Programs/Horde/HordeDashboard/src/components/AuditLog.tsx
+++ b/Engine/Source/Programs/Horde/HordeDashboard/src/components/AuditLog.tsx
@@ -530,9 +530,9 @@ export const AuditLogPanel: React.FC<{ agentId?: string, issueId?: string, templ
 
    return <Stack>
       <Breadcrumbs items={crumbItems} />
-      <Stack tokens={{ childrenGap: 12 }}>
+      <Stack tokens={{ childrenGap: 12 }} style={{backgroundColor: modeColors.background}}>
          <Stack horizontal>
-            <div key={`windowsize_logview1_${windowSize.width}_${windowSize.height}`} style={{ width: vw / 2 - (1440 / 2) - 230, flexShrink: 0, backgroundColor: modeColors.background }} />
+            <div key={`windowsize_logview1_${windowSize.width}_${windowSize.height}`} style={{ width: vw / 2 - (1440 / 2) - 990, flexShrink: 0, backgroundColor: modeColors.background }} />
             <Stack tokens={{ childrenGap: 0 }} styles={{ root: { backgroundColor: modeColors.background, margin: "auto", paddingTop: 12, paddingRight: 10 } }}>
                <Stack horizontal styles={{ root: { paddingLeft: 0, paddingBottom: 4, paddingRight: 12, width: 1440 } }}>
                   <Stack horizontal tokens={{ childrenGap: 12 }}>
diff --git a/Engine/Source/Programs/Horde/Plugins/Storage/HordeServer.Storage/Storage/StorageService.cs b/Engine/Source/Programs/Horde/Plugins/Storage/HordeServer.Storage/Storage/StorageService.cs
--- a/Engine/Source/Programs/Horde/Plugins/Storage/HordeServer.Storage/Storage/StorageService.cs
+++ b/Engine/Source/Programs/Horde/Plugins/Storage/HordeServer.Storage/Storage/StorageService.cs
@@ -376,6 +376,12 @@ namespace HordeServer.Storage
 			public List<GcNamespaceState> Namespaces { get; set; } = new List<GcNamespaceState>();
 			public bool Reset { get; set; }
 
+			public void DoReset()
+			{
+				LastImportBlobInfoId = ObjectId.Empty;
+				Reset = false;
+			}
+
 			public GcNamespaceState FindOrAddNamespace(NamespaceId namespaceId)
 			{
 				GcNamespaceState? namespaceState = Namespaces.FirstOrDefault(x => x.Id == namespaceId);
@@ -487,7 +493,7 @@ namespace HordeServer.Storage
 			_gcTicker = clock.AddTicker("Storage:GC", TimeSpan.FromMinutes(5.0), TickGcAsync, logger);
 
 			_lengthScanState = new SingletonDocument<LengthScanState>(mongoService);
-			_lengthScanTicker = clock.AddTicker("Storage:LengthScan", TimeSpan.FromSeconds(1.0), TickLengthsAsync, logger);
+			_lengthScanTicker = clock.AddSharedTicker("Storage:LengthScan", TimeSpan.FromMinutes(5.0), TickLengthsAsync, logger);
 		}
 
 		static string GetFieldName<TClass>(Expression<Func<TClass, object?>> expr)
@@ -728,15 +734,6 @@ namespace HordeServer.Storage
 			using TelemetrySpan span = _tracer.StartActiveSpan($"{nameof(StorageService)}.{nameof(TickBlobsAsync)}");
 
 			GcState gcState = await _gcState.GetAsync(cancellationToken);
-			if (gcState.Reset)
-			{
-				_logger.LogInformation("Resetting scan for new blobs...");
-				gcState = await _gcState.UpdateAsync(x =>
-				{
-					x.Reset = false;
-					x.LastImportBlobInfoId = ObjectId.Empty;
-				}, cancellationToken);
-			}
 
 			DateTime ingestTimeUtc = _clock.UtcNow - TimeSpan.FromMinutes(30.0);
 
@@ -747,8 +744,15 @@ namespace HordeServer.Storage
 
 			// Compute missing import info, by searching for blobs with an ObjectId timestamp after the last import compute cycle
 			ObjectId latestInfoId = ObjectId.GenerateNewId(ingestTimeUtc);
-			while (!gcState.Reset)
+			for(; ;)
 			{
+				// Check for the reset flag being set
+				if (gcState.Reset)
+				{
+					_logger.LogInformation("Resetting scan for new blobs...");
+					gcState = await _gcState.UpdateAsync(x => x.DoReset(), cancellationToken);
+				}
+
 				// Fetch the next batch of blobs
 				List<BlobInfo> current = await _blobCollection.Find(x => x.Id > gcState.LastImportBlobInfoId && x.Id < latestInfoId).SortBy(x => x.Id).Limit(500).ToListAsync(cancellationToken);
 				if (current.Count == 0)
diff --git a/Engine/Source/Programs/Shared/EpicGames.Horde/HordeHttpMessageHandler.cs b/Engine/Source/Programs/Shared/EpicGames.Horde/HordeHttpMessageHandler.cs
--- a/Engine/Source/Programs/Shared/EpicGames.Horde/HordeHttpMessageHandler.cs
+++ b/Engine/Source/Programs/Shared/EpicGames.Horde/HordeHttpMessageHandler.cs
@@ -50,7 +50,7 @@ namespace EpicGames.Horde
 				return Task.CompletedTask;
 			}
 
-			AsyncTimeoutPolicy<HttpResponseMessage> timeoutPolicy = Policy.TimeoutAsync<HttpResponseMessage>(30, OnTimeoutAsync);
+			AsyncTimeoutPolicy<HttpResponseMessage> timeoutPolicy = Policy.TimeoutAsync<HttpResponseMessage>(60, OnTimeoutAsync);
 
 			// Retry twice after a timeout
 			void OnRetry(Exception ex, TimeSpan timespan)
@@ -70,7 +70,7 @@ namespace EpicGames.Horde
 		{
 			Task OnTimeoutAsync(DelegateResult<HttpResponseMessage> outcome, TimeSpan timespan, int retryAttempt, Context context)
 			{
-				logger.LogDebug(KnownLogEvents.Systemic_Horde_Http, "{Method} {Url} failed ({Result}). Delaying for {DelayMs}ms (attempt #{RetryNum}).", request.Method, request.RequestUri, outcome.Result?.StatusCode, timespan.TotalMilliseconds, retryAttempt);
+				logger.LogInformation(KnownLogEvents.Systemic_Horde_Http, "{Method} {Url} failed ({Result}). Delaying for {DelayMs}ms (attempt #{RetryNum}).", request.Method, request.RequestUri, outcome.Result?.StatusCode, timespan.TotalMilliseconds, retryAttempt);
 				return Task.CompletedTask;
 			}
 
