Horde: Allow resetting blob id for length scan.

diff --git a/Engine/Source/Programs/Horde/Plugins/Storage/HordeServer.Storage/Storage/StorageService.cs b/Engine/Source/Programs/Horde/Plugins/Storage/HordeServer.Storage/Storage/StorageService.cs
--- a/Engine/Source/Programs/Horde/Plugins/Storage/HordeServer.Storage/Storage/StorageService.cs
+++ b/Engine/Source/Programs/Horde/Plugins/Storage/HordeServer.Storage/Storage/StorageService.cs
@@ -1400,7 +1400,11 @@ namespace HordeServer.Storage
 				}
 
 				// Update the last imported blob id
-				lengthScanState = await _lengthScanState.UpdateAsync(state => state.LastImportBlobInfoId = current[^1].Id, cancellationToken);
+				lengthScanState.LastImportBlobInfoId = current[^1].Id;
+				if (!await _lengthScanState.TryUpdateAsync(lengthScanState, cancellationToken))
+				{
+					lengthScanState = await _lengthScanState.GetAsync(cancellationToken);
+				}
 				scannedCount += current.Count;
 			}
 
