Horde: Distinguish between no matching compute resources vs all resources in use

diff --git a/Engine/Source/Programs/Horde/Plugins/Build/HordeServer.Build.Tests/Compute/ComputeServiceTest.cs b/Engine/Source/Programs/Horde/Plugins/Build/HordeServer.Build.Tests/Compute/ComputeServiceTest.cs
--- a/Engine/Source/Programs/Horde/Plugins/Build/HordeServer.Build.Tests/Compute/ComputeServiceTest.cs
+++ b/Engine/Source/Programs/Horde/Plugins/Build/HordeServer.Build.Tests/Compute/ComputeServiceTest.cs
@@ -191,19 +191,19 @@ namespace HordeServer.Tests.Compute
 			ClusterId clusterId = new("default");
 
 			AllocateResourceParams arp1 = new(clusterId, ComputeProtocol.Latest, new Requirements { Pool = "foo" }) { RequestId = "req1", RequesterIp = ip, ParentLeaseId = null };
-			ComputeResource? resource1 = await ComputeService.TryAllocateResourceAsync(arp1, CancellationToken.None);
+			ComputeResource resource1 = await ComputeService.TryAllocateResourceAsync(arp1, CancellationToken.None);
 			Assert.AreEqual(agent1.Id, resource1!.AgentId);
 
 			AllocateResourceParams arp2 = new(clusterId, ComputeProtocol.Latest, new Requirements { Pool = "bar-%REQUESTER_NETWORK_ID%" }) { RequestId = "req2", RequesterIp = IPAddress.Parse("15.0.0.1"), ParentLeaseId = null };
-			ComputeResource? resource2 = await ComputeService.TryAllocateResourceAsync(arp2, CancellationToken.None);
+			ComputeResource resource2 = await ComputeService.TryAllocateResourceAsync(arp2, CancellationToken.None);
 			Assert.AreEqual(agent2.Id, resource2!.AgentId);
 
 			AllocateResourceParams arp3 = new(clusterId, ComputeProtocol.Latest, new Requirements { Pool = "bar-%REQUESTER_NETWORK_ID%" }) { RequestId = "req3", RequesterIp = ip, ParentLeaseId = null };
-			ComputeResource? resource3 = await ComputeService.TryAllocateResourceAsync(arp3, CancellationToken.None);
+			ComputeResource resource3 = await ComputeService.TryAllocateResourceAsync(arp3, CancellationToken.None);
 			Assert.AreEqual(agent3.Id, resource3!.AgentId);
 
 			AllocateResourceParams arp4 = new(clusterId, ComputeProtocol.Latest, new Requirements { Pool = "qux-%REQUESTER_COMPUTE_ID%" }) { RequestId = "req4", RequesterIp = ip, ParentLeaseId = null };
-			ComputeResource? resource4 = await ComputeService.TryAllocateResourceAsync(arp4, CancellationToken.None);
+			ComputeResource resource4 = await ComputeService.TryAllocateResourceAsync(arp4, CancellationToken.None);
 			Assert.AreEqual(agent4.Id, resource4!.AgentId);
 		}
 
@@ -217,10 +217,10 @@ namespace HordeServer.Tests.Compute
 
 			IAgent agent = await CreateAgentAsync(new PoolId("foo"), properties: ["ComputeIp=11.0.0.1", "ComputePort=5000"]);
 
-			ComputeResource? resource1 = await ComputeService.TryAllocateResourceAsync(CreateParams("11.0.0.1"), CancellationToken.None);
-			Assert.IsNull(resource1);
+			await Assert.ThrowsExceptionAsync<NoComputeResourcesException>(() => ComputeService.TryAllocateResourceAsync(CreateParams("11.0.0.1"), CancellationToken.None));
+			await Assert.ThrowsExceptionAsync<NoComputeResourcesException>(() => ComputeService.TryAllocateResourceAsync(CreateParams("11.0.0.1"), CancellationToken.None));
 
-			ComputeResource? resource2 = await ComputeService.TryAllocateResourceAsync(CreateParams("11.0.0.2"), CancellationToken.None);
+			ComputeResource resource2 = await ComputeService.TryAllocateResourceAsync(CreateParams("11.0.0.2"), CancellationToken.None);
 			Assert.AreEqual(agent.Id, resource2!.AgentId);
 		}
 
@@ -235,13 +235,30 @@ namespace HordeServer.Tests.Compute
 			IAgent agent1 = await CreateAgentAsync(new PoolId("foo"), properties: ["ComputeIp=11.0.0.1", "ComputePort=5000"]);
 			IAgent agent2 = await CreateAgentAsync(new PoolId("bar"), properties: ["ComputeIp=11.0.0.2", "ComputePort=5000"]);
 
-			ComputeResource? resource1 = await ComputeService.TryAllocateResourceAsync(CreateParams("foo"), CancellationToken.None);
+			ComputeResource resource1 = await ComputeService.TryAllocateResourceAsync(CreateParams("foo"), CancellationToken.None);
 			Assert.AreEqual(agent1.Id, resource1!.AgentId);
 
-			ComputeResource? resource2 = await ComputeService.TryAllocateResourceAsync(CreateParams("bar"), CancellationToken.None);
+			ComputeResource resource2 = await ComputeService.TryAllocateResourceAsync(CreateParams("bar"), CancellationToken.None);
 			Assert.AreEqual(agent2.Id, resource2!.AgentId);
 
-			Assert.IsNull(await ComputeService.TryAllocateResourceAsync(CreateParams("does-not-exist"), CancellationToken.None));
+			await Assert.ThrowsExceptionAsync<NoComputeResourcesException>(() => ComputeService.TryAllocateResourceAsync(CreateParams("does-not-exist"), CancellationToken.None));
+		}
+		
+		[TestMethod]
+		public async Task Assignment_NoResources_Async()
+		{
+			AllocateResourceParams CreateParams(string pool)
+			{
+				return new AllocateResourceParams(new ClusterId("default"), ComputeProtocol.Latest, new Requirements { Pool = pool, Exclusive = true });
+			}
+			
+			await CreateAgentAsync(new PoolId("foo"), properties: ["ComputeIp=11.0.0.1", "ComputePort=5000"]);
+			await ComputeService.TryAllocateResourceAsync(CreateParams("foo"), CancellationToken.None);
+			NoComputeResourcesException ex1 = await Assert.ThrowsExceptionAsync<NoComputeResourcesException>(() => ComputeService.TryAllocateResourceAsync(CreateParams("foo"), CancellationToken.None));
+			Assert.AreEqual(NoComputeResourcesException.AllResourcesInUse, ex1.Message);
+			
+			NoComputeResourcesException ex2 = await Assert.ThrowsExceptionAsync<NoComputeResourcesException>(() => ComputeService.TryAllocateResourceAsync(CreateParams("bar"), CancellationToken.None));
+			Assert.AreEqual(NoComputeResourcesException.NoMatchingResources, ex2.Message);
 		}
 
 		[TestMethod]
@@ -269,12 +286,12 @@ namespace HordeServer.Tests.Compute
 				Assert.AreEqual(agent1.Id, (await ComputeService.TryAllocateResourceAsync(arpFoo, CancellationToken.None))!.AgentId);
 
 				AllocateResourceParams arpBar = new(cluster1, ComputeProtocol.Latest, new Requirements { Pool = poolBar.ToString() });
-				Assert.IsNull(await ComputeService.TryAllocateResourceAsync(arpBar, CancellationToken.None));
+				await Assert.ThrowsExceptionAsync<NoComputeResourcesException>(() => ComputeService.TryAllocateResourceAsync(arpBar, CancellationToken.None));
 			}
 
 			{
 				AllocateResourceParams arpFoo = new(cluster2, ComputeProtocol.Latest, new Requirements { Pool = poolFoo.ToString() });
-				Assert.IsNull(await ComputeService.TryAllocateResourceAsync(arpFoo, CancellationToken.None));
+				await Assert.ThrowsExceptionAsync<NoComputeResourcesException>(() => ComputeService.TryAllocateResourceAsync(arpFoo, CancellationToken.None));
 
 				AllocateResourceParams arpBar = new(cluster2, ComputeProtocol.Latest, new Requirements { Pool = poolBar.ToString() });
 				Assert.AreEqual(agent2.Id, (await ComputeService.TryAllocateResourceAsync(arpBar, CancellationToken.None))!.AgentId);
diff --git a/Engine/Source/Programs/Horde/Plugins/Compute/HordeServer.Compute/Compute/ComputeController.cs b/Engine/Source/Programs/Horde/Plugins/Compute/HordeServer.Compute/Compute/ComputeController.cs
--- a/Engine/Source/Programs/Horde/Plugins/Compute/HordeServer.Compute/Compute/ComputeController.cs
+++ b/Engine/Source/Programs/Horde/Plugins/Compute/HordeServer.Compute/Compute/ComputeController.cs
@@ -122,18 +122,20 @@ namespace HordeServer.Compute
 				UserId = User.GetUserId(),
 			};
 
-			ComputeResource? computeResource;
+			ComputeResource computeResource;
 			try
 			{
 				computeResource = await _computeService.TryAllocateResourceAsync(arp, cancellationToken);
-				if (computeResource == null)
-				{
-					return StatusCode((int)HttpStatusCode.ServiceUnavailable, "No resources available");
-				}
+			}
+			catch (NoComputeResourcesException cse)
+			{
+				int statusCode = (int)HttpStatusCode.ServiceUnavailable;
+				return cse.ShowToUser ? StatusCode(statusCode, cse.Message) : StatusCode(statusCode);
 			}
 			catch (ComputeServiceException cse)
 			{
-				return cse.ShowToUser ? StatusCode((int)HttpStatusCode.InternalServerError, cse.Message) : StatusCode((int)HttpStatusCode.InternalServerError);
+				int statusCode = (int)HttpStatusCode.InternalServerError;
+				return cse.ShowToUser ? StatusCode(statusCode, cse.Message) : StatusCode(statusCode);
 			}
 
 			Dictionary<string, ConnectionMetadataPort> responsePorts = new();
diff --git a/Engine/Source/Programs/Horde/Plugins/Compute/HordeServer.Compute/Compute/ComputeService.cs b/Engine/Source/Programs/Horde/Plugins/Compute/HordeServer.Compute/Compute/ComputeService.cs
--- a/Engine/Source/Programs/Horde/Plugins/Compute/HordeServer.Compute/Compute/ComputeService.cs
+++ b/Engine/Source/Programs/Horde/Plugins/Compute/HordeServer.Compute/Compute/ComputeService.cs
@@ -48,7 +48,16 @@ namespace HordeServer.Compute
 		{
 		}
 	}
-
+	
+	/// <summary>
+	/// Exception when no compute resources available
+	/// </summary>
+	public class NoComputeResourcesException(string? message, Exception? innerException = null) : ComputeServiceException(message, innerException)
+	{
+		internal const string AllResourcesInUse = "Matching resources available but all are in use";
+		internal const string NoMatchingResources = "No matching resources available (in use or not)";
+	}
+	
 	/// <summary>
 	/// Outcome for a compute allocation request
 	/// </summary>
@@ -350,7 +359,7 @@ namespace HordeServer.Compute
 		/// <param name="arp">Allocation parameters</param>
 		/// <param name="cancellationToken">Cancellation token</param>
 		/// <returns>A compute resource if successful</returns>
-		public async Task<ComputeResource?> TryAllocateResourceAsync(AllocateResourceParams arp, CancellationToken cancellationToken)
+		public async Task<ComputeResource> TryAllocateResourceAsync(AllocateResourceParams arp, CancellationToken cancellationToken)
 		{
 			using TelemetrySpan span = _tracer.StartActiveSpan($"{nameof(ComputeService)}.{nameof(TryAllocateResourceAsync)}");
 			span.SetAttribute("clusterId", arp.ClusterId.ToString());
@@ -397,11 +406,16 @@ namespace HordeServer.Compute
 			try
 			{
 				PoolId? poolId = arp.Requirements.Pool != null ? new PoolId(arp.Requirements.Pool) : null;
+				int numMatchingAgents = 0;
 				await foreach (IAgent agent in _agentCollection.FindAsync(poolId, consistentRead: false, cancellationToken: cancellationToken))
 				{
 					Dictionary<string, int> assignedResources = new Dictionary<string, int>();
 
 					bool isMemberOfCluster = clusterConfig.Condition == null || agent.SatisfiesCondition(clusterConfig.Condition);
+					if (isMemberOfCluster)
+					{
+						numMatchingAgents++;
+					}
 					bool match = isMemberOfCluster && agent.MeetsRequirements(arp.Requirements, assignedResources, conditions);
 					if (match)
 					{
@@ -456,7 +470,12 @@ namespace HordeServer.Compute
 				}
 
 				await LogRequestAsync(AllocationOutcome.Denied, arp.RequestId, arp.Requirements, arp.ParentLeaseId, span, cancellationToken);
-				return null;
+				
+				if (numMatchingAgents == 0)
+				{
+					throw new NoComputeResourcesException(NoComputeResourcesException.NoMatchingResources);
+				}
+				throw new NoComputeResourcesException(NoComputeResourcesException.AllResourcesInUse);
 			}
 			catch (AgentRelayException are)
 			{
