Horde: Ensure auto-conforming is requested only when not already pending

diff --git a/Engine/Source/Programs/Horde/Plugins/Build/HordeServer.Build.Tests/Agents/Pools/PoolUpdateServiceTest.cs b/Engine/Source/Programs/Horde/Plugins/Build/HordeServer.Build.Tests/Agents/Pools/PoolUpdateServiceTest.cs
--- a/Engine/Source/Programs/Horde/Plugins/Build/HordeServer.Build.Tests/Agents/Pools/PoolUpdateServiceTest.cs
+++ b/Engine/Source/Programs/Horde/Plugins/Build/HordeServer.Build.Tests/Agents/Pools/PoolUpdateServiceTest.cs
@@ -127,6 +127,7 @@ public class PoolUpdateServiceTest : BuildTestSetup
 		IAgent agent = await CreateAutoConformAgentAsync(100, autoConformThresholdsMb);
 
 		// Act
+		await _pus.UpdatePoolsAsync(CancellationToken.None);
 		await _pus.AutoConformAgentsAsync(CancellationToken.None);
 		agent = (await AgentService.GetAgentAsync(agent.Id))!;
 
diff --git a/Engine/Source/Programs/Horde/Plugins/Build/HordeServer.Build/Agents/Pools/PoolUpdateService.cs b/Engine/Source/Programs/Horde/Plugins/Build/HordeServer.Build/Agents/Pools/PoolUpdateService.cs
--- a/Engine/Source/Programs/Horde/Plugins/Build/HordeServer.Build/Agents/Pools/PoolUpdateService.cs
+++ b/Engine/Source/Programs/Horde/Plugins/Build/HordeServer.Build/Agents/Pools/PoolUpdateService.cs
@@ -152,7 +152,8 @@ namespace HordeServer.Agents.Pools
 					}
 				}
 
-				if (freeDiskSpace != null && conformDiskSpaceNeeded > 0 && freeDiskSpace < conformDiskSpaceNeeded && agent.ConformAttemptCount is null or 0)
+				if (freeDiskSpace != null && conformDiskSpaceNeeded > 0 && freeDiskSpace < conformDiskSpaceNeeded &&
+				    agent.ConformAttemptCount is null or 0 && !agent.RequestFullConform)
 				{
 					await agent.TryUpdateAsync(new UpdateAgentOptions { RequestFullConform = true }, cancellationToken: cancellationToken);
 					_logger.LogInformation("Auto-conforming {AgentId} as workspace conform disk space needed ({ConformDiskSpace:F1} MB) is less than free disk space ({FreeDiskSpace:F1} MB)",
@@ -166,7 +167,7 @@ namespace HordeServer.Agents.Pools
 		/// </summary>
 		/// <param name="cancellationToken">Cancellation token for the async task</param>
 		/// <returns>Async task</returns>
-		async ValueTask UpdatePoolsAsync(CancellationToken cancellationToken)
+		internal async ValueTask UpdatePoolsAsync(CancellationToken cancellationToken)
 		{
 			using TelemetrySpan span = _tracer.StartActiveSpan($"{nameof(PoolUpdateService)}.{nameof(UpdatePoolsAsync)}");
 
