Horde: Add Horde's built-in user ID claim when issuing OIDC userinfo for Horde accounts

diff --git a/Engine/Source/Programs/Horde/HordeServer/Accounts/OAuthController.cs b/Engine/Source/Programs/Horde/HordeServer/Accounts/OAuthController.cs
--- a/Engine/Source/Programs/Horde/HordeServer/Accounts/OAuthController.cs
+++ b/Engine/Source/Programs/Horde/HordeServer/Accounts/OAuthController.cs
@@ -485,6 +485,7 @@ namespace HordeServer.Accounts
 				properties[ClaimTypes.Email] = account.Email;
 			}
 			properties[HordeClaimTypes.AccountId] = account.Id.ToString();
+			properties[HordeClaimTypes.User] = account.Login;
 
 			foreach (IUserClaim claim in account.Claims)
 			{
