Horde: Improve C# docs for ServerUrl and AuthMode in ServerSettings

#rnx

diff --git a/Engine/Source/Programs/Horde/HordeServer/ServerSettings.cs b/Engine/Source/Programs/Horde/HordeServer/ServerSettings.cs
--- a/Engine/Source/Programs/Horde/HordeServer/ServerSettings.cs
+++ b/Engine/Source/Programs/Horde/HordeServer/ServerSettings.cs
@@ -140,6 +140,7 @@ namespace HordeServer
 
 		/// <summary>
 		/// Type of authentication (e.g anonymous, OIDC, built-in Horde accounts)
+		/// If "Horde" auth mode is used, be sure to configure "ServerUrl" as well.
 		/// </summary>
 		public AuthMethod AuthMethod { get; set; } = AuthMethod.Anonymous;
 		
@@ -213,7 +214,10 @@ namespace HordeServer
 		public string[] OidcClaimHordePerforceUserMapping { get; set; } = { "preferred_username", "email" };
 
 		/// <summary>
-		/// Name of this machine 
+		/// Base URL this Horde server is accessible from
+		/// For example https://horde.mystudio.com
+		/// If not set, a default is used based on current hostname.
+		/// In more advanced setups where a reverse proxy is present in front of Horde, this must manually be set.
 		/// </summary>
 		public Uri ServerUrl
 		{
