Horde: Don't add all files to filter when there is no selection in the artifact modal to avoid overflowing the max request length
#skipci

diff --git a/Engine/Source/Programs/Horde/HordeDashboard/src/backend/Api.ts b/Engine/Source/Programs/Horde/HordeDashboard/src/backend/Api.ts
--- a/Engine/Source/Programs/Horde/HordeDashboard/src/backend/Api.ts
+++ b/Engine/Source/Programs/Horde/HordeDashboard/src/backend/Api.ts
@@ -1656,7 +1656,7 @@ export type ArtifactContextType = "step-trace" | "step-output" | "step-saved" |
 /// Request to create a zip file with artifact data
 export type CreateZipRequest = {
 	/// Filter lines for the zip. Uses standard <see cref="FileFilter"/> syntax.
-	filter: string[];
+	filter?: string[];
 }
 
 
diff --git a/Engine/Source/Programs/Horde/HordeDashboard/src/backend/Backend.ts b/Engine/Source/Programs/Horde/HordeDashboard/src/backend/Backend.ts
--- a/Engine/Source/Programs/Horde/HordeDashboard/src/backend/Backend.ts
+++ b/Engine/Source/Programs/Horde/HordeDashboard/src/backend/Backend.ts
@@ -571,12 +571,13 @@ export class Backend {
         window.location.assign(url);
     }
 
-    downloadArtifactZipV2(artifactId: string, request: CreateZipRequest) {
+    downloadArtifactZipV2(artifactId: string, request?: CreateZipRequest) {        
 
-        const filter = request.filter.map(f => `filter=${encodeURIComponent(f)}`).join("&")
+        const filter = request?.filter?.map(f => `filter=${encodeURIComponent(f)}`).join("&")
 
         let url = `/api/v2/artifacts/${artifactId}/zip`;
-        if (filter.length) {
+        
+        if (filter?.length) {
             url += "?" + filter;
         }
 
diff --git a/Engine/Source/Programs/Horde/HordeDashboard/src/components/artifacts/ArtifactsModal.tsx b/Engine/Source/Programs/Horde/HordeDashboard/src/components/artifacts/ArtifactsModal.tsx
--- a/Engine/Source/Programs/Horde/HordeDashboard/src/components/artifacts/ArtifactsModal.tsx
+++ b/Engine/Source/Programs/Horde/HordeDashboard/src/components/artifacts/ArtifactsModal.tsx
@@ -6,7 +6,7 @@ import { observer } from "mobx-react-lite";
 import { useEffect, useState } from "react";
 import { NavigateFunction, useNavigate } from "react-router-dom";
 import backend from "../../backend";
-import { ArtifactContextType, GetArtifactDirectoryEntryResponse, GetArtifactDirectoryResponse, GetArtifactFileEntryResponse, GetArtifactResponse } from "../../backend/Api";
+import { ArtifactContextType, CreateZipRequest, GetArtifactDirectoryEntryResponse, GetArtifactDirectoryResponse, GetArtifactFileEntryResponse, GetArtifactResponse } from "../../backend/Api";
 import dashboard, { StatusColor } from "../../backend/Dashboard";
 import { getHordeStyling } from "../../styles/Styles";
 
@@ -545,21 +545,30 @@ const DownloadButton: React.FC<{ handler: ArtifactsHandler, openArtifactInfo: ()
 
          const path = handler.path ? handler.path + "/" : "";
 
-         const filters = selection.map(s => {
-            const item = s as BrowserItem;
+         let zipRequest: CreateZipRequest | undefined;
 
-            if (item.type === BrowserType.NavigateUp) {
-               return "";
-            }
+         if ((handler.selection?.count ?? 0) > 0) {
 
-            if (item.type === BrowserType.Directory) {
-               return `${path}${item.text}/...`;
-            }
-            return `${path}${item.text}`;
-         }).filter(f => !!f);
+            const filter = selection.map(s => {
+               const item = s as BrowserItem;
+   
+               if (item.type === BrowserType.NavigateUp) {
+                  return "";
+               }
+   
+               if (item.type === BrowserType.Directory) {
+                  return `${path}${item.text}/...`;
+               }
+               return `${path}${item.text}`;
+            }).filter(f => !!f);
+
+            zipRequest = {
+               filter: filter
+            }   
+         }
 
          try {
-            backend.downloadArtifactZipV2(handler.artifact.id, { filter: filters });
+            backend.downloadArtifactZipV2(handler.artifact.id, zipRequest);
          } catch (err) {
             console.error(err);
          } finally {
