Horde: Fix issues with theming and login/setup view for installer build
#skipci

diff --git a/Engine/Source/Programs/Horde/HordeDashboard/src/App.tsx b/Engine/Source/Programs/Horde/HordeDashboard/src/App.tsx
--- a/Engine/Source/Programs/Horde/HordeDashboard/src/App.tsx
+++ b/Engine/Source/Programs/Horde/HordeDashboard/src/App.tsx
@@ -55,6 +55,8 @@ const Main: React.FC = () => {
 
    const [init, setInit] = useState(false);
 
+   document.body.setAttribute('style', `background: ${dashboard.darktheme ? "#0F0F0F" : "#FAF9F9"}`)   
+
    if (window.location.pathname === "/login") {
       return <HordeLoginView />
    }
@@ -85,9 +87,6 @@ const Main: React.FC = () => {
          ErrorHandler.set({ title: "Error initializing site", reason: reason }, true);
       });
 
-      // avoid flash when initializing dashboard
-      document.body.setAttribute('style', `background: ${dashboard.darktheme ? "#0F0F0F" : "#FAF9F9"}`)   
-
       return (<ThemeProvider applyTo='element' theme={dashboard.darktheme ? darkTheme : lightTheme}>
          <div style={{ position: 'absolute', left: '50%', top: '50%', transform: 'translate(-50%, -50%)' }}>
             <Stack horizontalAlign="center" styles={{ root: { padding: 20, minWidth: 200, minHeight: 100 } }}>
@@ -186,10 +185,6 @@ const Main: React.FC = () => {
 
       router = createBrowserRouter(routes);
    }
-
-   // avoid white bar on body
-   
-   document.body.setAttribute('style', `background: ${dashboard.darktheme ? "#0F0F0F" : "#FAF9F9"}`)   
    
    
    return (
diff --git a/Engine/Source/Programs/Horde/HordeDashboard/src/components/accounts/HordeLoginView.tsx b/Engine/Source/Programs/Horde/HordeDashboard/src/components/accounts/HordeLoginView.tsx
--- a/Engine/Source/Programs/Horde/HordeDashboard/src/components/accounts/HordeLoginView.tsx
+++ b/Engine/Source/Programs/Horde/HordeDashboard/src/components/accounts/HordeLoginView.tsx
@@ -14,6 +14,8 @@ export const HordeLoginView: React.FC = () => {
 
    const error = state.error;
 
+   document.body.setAttribute('style', `background: ${dashboard.darktheme ? "#0F0F0F" : "#FFFFFFF"}`)   
+
    const onLogin = async () => {
       const username = state.username?.trim();
       const password = state.password?.trim();
diff --git a/Engine/Source/Programs/Horde/HordeDashboard/src/components/setup/HordeSetupView.tsx b/Engine/Source/Programs/Horde/HordeDashboard/src/components/setup/HordeSetupView.tsx
--- a/Engine/Source/Programs/Horde/HordeDashboard/src/components/setup/HordeSetupView.tsx
+++ b/Engine/Source/Programs/Horde/HordeDashboard/src/components/setup/HordeSetupView.tsx
@@ -14,6 +14,8 @@ export const HordeSetupView: React.FC = () => {
 
    const username = "Admin";
 
+   document.body.setAttribute('style', `background: ${dashboard.darktheme ? "#0F0F0F" : "#FFFFFFF"}`)   
+
    const onLoginAdmin = async () => {
 
       const password = secrets.newPassword?.trim();
