Horde: Use P4 syntax on artifact modal to download current folder
#skipci

diff --git a/Engine/Source/Programs/Horde/HordeDashboard/src/components/artifacts/ArtifactsModal.tsx b/Engine/Source/Programs/Horde/HordeDashboard/src/components/artifacts/ArtifactsModal.tsx
--- a/Engine/Source/Programs/Horde/HordeDashboard/src/components/artifacts/ArtifactsModal.tsx
+++ b/Engine/Source/Programs/Horde/HordeDashboard/src/components/artifacts/ArtifactsModal.tsx
@@ -565,6 +565,12 @@ const DownloadButton: React.FC<{ handler: ArtifactsHandler, openArtifactInfo: ()
             zipRequest = {
                filter: filter
             }   
+         } else {
+            if (path) {
+               zipRequest = {
+                  filter: [path + "..."]
+               }   
+            }
          }
 
          try {
