Horde: Additional logging for IssueTagService.

#rnx

diff --git a/Engine/Source/Programs/Horde/Plugins/Build/HordeServer.Build/Issues/IssueTagService.cs b/Engine/Source/Programs/Horde/Plugins/Build/HordeServer.Build/Issues/IssueTagService.cs
--- a/Engine/Source/Programs/Horde/Plugins/Build/HordeServer.Build/Issues/IssueTagService.cs
+++ b/Engine/Source/Programs/Horde/Plugins/Build/HordeServer.Build/Issues/IssueTagService.cs
@@ -59,6 +59,8 @@ namespace HordeServer.Issues
 
 		async ValueTask TickAsync(CancellationToken cancellationToken)
 		{
+			_logger.LogDebug("Starting scan for changes with {Tag} tag", _buildConfig.CurrentValue.IssueFixedTag);
+
 			using CancellationTokenSource cancellationSource = CancellationTokenSource.CreateLinkedTokenSource(cancellationToken);
 			using IDisposable? listener = _buildConfig.OnChange((_, _) => cancellationSource.Cancel());
 
@@ -96,6 +98,8 @@ namespace HordeServer.Issues
 					}
 				}
 			}
+
+			_logger.LogDebug("Stopping scan for changes with {Tag} tag.", _buildConfig.CurrentValue.IssueFixedTag);
 		}
 
 		async Task TickStreamGuardedAsync(StreamConfig streamConfig, State initialState, CancellationToken cancellationToken)
@@ -109,7 +113,7 @@ namespace HordeServer.Issues
 			}
 			catch (Exception ex)
 			{
-				_logger.LogError(ex, "Exception while scanning for #horde tags: {Message}", ex.Message);
+				_logger.LogError(ex, "Exception while scanning for {Tag} tags: {Message}", _buildConfig.CurrentValue.IssueFixedTag, ex.Message);
 			}
 		}
 
