Horde: Enable debug logging for issue tag service by default.

#rnx

diff --git a/Engine/Source/Programs/Horde/HordeServer/appsettings.json b/Engine/Source/Programs/Horde/HordeServer/appsettings.json
--- a/Engine/Source/Programs/Horde/HordeServer/appsettings.json
+++ b/Engine/Source/Programs/Horde/HordeServer/appsettings.json
@@ -18,6 +18,7 @@
 				"System.Net.Http.HttpClient": "Warning",
 				"Grpc": "Warning",
 				"HordeServer.Issues.IssueService": "Debug",
+				"HordeServer.Issues.IssueTagService": "Debug",
 				"HordeServer.Storage.ObjectStores.AwsObjectStore": "Information",
 				"HordeServer.Storage.StorageService": "Debug",
 				"EpicGames.Horde.Storage": "Debug"
