Horde: Modify default config files to store plugins configs within the plugins object.

#rnx

diff --git a/Engine/Source/Programs/Horde/HordeServer/Defaults/default.global.json b/Engine/Source/Programs/Horde/HordeServer/Defaults/default.global.json
--- a/Engine/Source/Programs/Horde/HordeServer/Defaults/default.global.json
+++ b/Engine/Source/Programs/Horde/HordeServer/Defaults/default.global.json
@@ -11,115 +11,128 @@
 		"showDeviceManager": true,
 		"showTests": true
 	},
-	"compute": [
-		{
-			"id": "default",
-			"namespaceId": "horde.compute"
-		}
-	],
-	"storage": {
-		"backends": [
-			{
-				"id": "default-backend",
-				"type": "FileSystem",
-				"baseDir": "Storage" // Default base directory is C:\ProgramData\HordeServer
-			},
-			{
-				"id": "memory-backend", // Used for automated tests
-				"type": "Memory"
-			}
-		],
-		"namespaces": [
-			{
-				"id": "default",
-				"backend": "default-backend"
-			},
-			{
-				"id": "horde-artifacts",
-				"prefix": "Artifacts/",
-				"backend": "default-backend"
-			},
-			{
-				"id": "horde-perforce",
-				"prefix": "Perforce/",
-				"backend": "default-backend"
-			},
-			// The 'horde-logs' namespace is used internally by Horde to store logs from the CI
-			// and remote execution systems. It must be configured in order for Horde to function
-			// correctly.
-			{
-				"id": "horde-logs",
-				"prefix": "Logs/",
-				"backend": "default-backend"
-			},
-			{
-				"id": "horde-tools",
-				"prefix": "Tools/",
-				"backend": "default-backend"
-			},
-			{
-				"id": "memory",
-				"backend": "memory-backend",
-				"enableAliases": true
-			}
-		]
-	},
-	"pools": [
-		{
-			"name": "Win-UE5",
-			"condition": "Platform == 'Win64'",
-			"color": "Blue"
-		},
-		{
-			"name": "Win-UE5-GPU",
-			"condition": "Platform == 'Win64' && Interactive == true",
-			"color": "Blue"
+
+	// Settings for various plugins
+	"plugins": {
+
+		"compute": {
+			"clusters": [
+				{
+					"id": "default",
+					"namespaceId": "horde.compute"
+				}
+			],
+			"networks": [
+				{
+					"id": "default",
+					"cidrBlock": "0.0.0.0/0",
+					"description": "Default network (matching all IPs)",
+					"computeId": "default"
+				}
+			]
 		},
-		{
-			"name": "Mac-UE5",
-			"condition": "Platform == 'Mac'",
-			"color": "Green"
+
+		"storage": {
+			"backends": [
+				{
+					"id": "default-backend",
+					"type": "FileSystem",
+					"baseDir": "Storage" // Default base directory is C:\ProgramData\HordeServer
+				},
+				{
+					"id": "memory-backend", // Used for automated tests
+					"type": "Memory"
+				}
+			],
+			"namespaces": [
+				{
+					"id": "default",
+					"backend": "default-backend"
+				},
+				{
+					"id": "horde-artifacts",
+					"prefix": "Artifacts/",
+					"backend": "default-backend"
+				},
+				{
+					"id": "horde-perforce",
+					"prefix": "Perforce/",
+					"backend": "default-backend"
+				},
+				// The 'horde-logs' namespace is used internally by Horde to store logs from the CI
+				// and remote execution systems. It must be configured in order for Horde to function
+				// correctly.
+				{
+					"id": "horde-logs",
+					"prefix": "Logs/",
+					"backend": "default-backend"
+				},
+				{
+					"id": "horde-tools",
+					"prefix": "Tools/",
+					"backend": "default-backend"
+				},
+				{
+					"id": "memory",
+					"backend": "memory-backend",
+					"enableAliases": true
+				}
+			]
 		},
-		{
-			"name": "Linux-UE5",
-			"condition": "Platform == 'Linux'",
-			"color": "Orange"
-		}
-	],
-	"networks": [
-		{
-			"id": "default",
-			"cidrBlock": "0.0.0.0/0",
-			"description": "Default network (matching all IPs)",
-			"computeId": "default"
-		}
-	],
-	"artifactTypes": [
-		{
-			"name": "step-output",
-			"keepDays": 1
-		}
-	],
-	"devices": {
-		"platforms": [
-			{
-				"id": "android",
-				"name": "Android"
-			}
-		],
-		"pools": [
-			{
-				"id": "ue5",
-				"name": "UE5",
-				"poolType": "Automation"
+
+		"build": {
+			"artifactTypes": [
+				{
+					"name": "step-output",
+					"keepDays": 1
+				}
+			],
+			"devices": {
+				"platforms": [
+					{
+						"id": "android",
+						"name": "Android"
+					}
+				],
+				"pools": [
+					{
+						"id": "ue5",
+						"name": "UE5",
+						"poolType": "Automation"
+					},
+					{
+						"id": "remote-ue5",
+						"name": "Remote UE5",
+						"poolType": "Shared"
+					}
+				]
 			},
-			{
-				"id": "remote-ue5",
-				"name": "Remote UE5",
-				"poolType": "Shared"
-			}
-		]
+			"pools": [
+				{
+					"name": "Win-UE5",
+					"condition": "Platform == 'Win64'",
+					"color": "Blue"
+				},
+				{
+					"name": "Win-UE5-GPU",
+					"condition": "Platform == 'Win64' && Interactive == true",
+					"color": "Blue"
+				},
+				{
+					"name": "Mac-UE5",
+					"condition": "Platform == 'Mac'",
+					"color": "Green"
+				},
+				{
+					"name": "Linux-UE5",
+					"condition": "Platform == 'Linux'",
+					"color": "Orange"
+				}
+			]
+		}
 	},
+
+	// Default ACL groups
 	"acl": {
 		"entries": [
 			{
diff --git a/Engine/Source/Programs/Horde/HordeServer/Defaults/globals.json b/Engine/Source/Programs/Horde/HordeServer/Defaults/globals.json
--- a/Engine/Source/Programs/Horde/HordeServer/Defaults/globals.json
+++ b/Engine/Source/Programs/Horde/HordeServer/Defaults/globals.json
@@ -10,55 +10,59 @@
 		}
 	],
 
-	// Definitions for projects in the CI system
-	"projects": [
-		// Uncomment this to enable a default //UE5/Dev-Main-Horde stream with some simple CI jobs.
-//		{
-//			"id": "ue5",
-//			"path": "ue5.project.json"
-//		}
-	],
+	"plugins": {
+		"build": {
+			// Definitions for projects in the CI system
+			"projects": [
+				// Uncomment this to enable a default //UE5/Dev-Main-Horde stream with some simple CI jobs.
+//				{
+//					"id": "ue5",
+//					"path": "ue5.project.json"
+//				}
+			],
 
-	// Configuration for different Perforce servers. Horde allows running against multiple 'logical' Perforce servers, each of which may be accessed through different edge servers. Server
-	// selection is performed to evenly distribute load, and servers can expose their health/status via a separate health check service.
-	"perforceClusters": 
-	[
-//		{
-//			// Name for this cluster. Streams can be configured to use a particular cluster using the ClusterName property in the stream config. 
-//			"name": "Default",
+			// Configuration for different Perforce servers. Horde allows running against multiple 'logical' Perforce servers, each of which may be accessed through different edge servers. Server
+			// selection is performed to evenly distribute load, and servers can expose their health/status via a separate health check service.
+			"perforceClusters": 
+			[
+//				{
+//					// Name for this cluster. Streams can be configured to use a particular cluster using the ClusterName property in the stream config. 
+//					"name": "Default",
 //
-//			// Default account to use inside the Horde Server for querying submitted changes etc..
-//			// If this value is not specified, the first item from the credentials below will be used.
-//			"serviceAccount": null,
+//					// Default account to use inside the Horde Server for querying submitted changes etc..
+//					// If this value is not specified, the first item from the credentials below will be used.
+//					"serviceAccount": null,
 //
-//			// Whether the service account supports generating login tickets for other users, which requires a priviledged account on the server. Impersonating other users
-//			// allows Horde to submit changes on behalf of users.
-//			"canImpersonate": false,
+//					// Whether the service account supports generating login tickets for other users, which requires a priviledged account on the server. Impersonating other users
+//					// allows Horde to submit changes on behalf of users.
+//					"canImpersonate": false,
 //
-//			// List of servers to access
-//			"servers": 
-//			[
-//				{
-//					// DNS name and port for the server.
-//					// If connecting using "ssl:" prefix, ensure P4 server's fingerprint/certificate is trusted (see docs for "p4 trust")
-//					// See Horde's documentation on connecting to SSL-enabled Perforce servers.
-//					"serverAndPort": "perforce:1666"
-//				}
-//			],
+//					// List of servers to access
+//					"servers": 
+//					[
+//						{
+//							// DNS name and port for the server.
+//							// If connecting using "ssl:" prefix, ensure P4 server's fingerprint/certificate is trusted (see docs for "p4 trust")
+//							// See Horde's documentation on connecting to SSL-enabled Perforce servers.
+//							"serverAndPort": "perforce:1666"
+//						}
+//					],
 //
-//			// List of credentials for the server. The first credentials will be used by default, but can be overridden by setting the userName property in the workspace.
-//			"credentials": 
-//			[
-//				{
-//					"userName": "test.user",
+//					// List of credentials for the server. The first credentials will be used by default, but can be overridden by setting the userName property in the workspace.
+//					"credentials": 
+//					[
+//						{
+//							"userName": "test.user",
 //
-//					// Specify either the password or ticket property below (but not both)
-//					"password": "password123",
-//					// "ticket": "ticket if not using password"
+//							// Specify either the password or ticket property below (but not both)
+//							"password": "password123",
+//							// "ticket": "ticket if not using password"
+//						}
+//					]
 //				}
-//			]
-//		}
-	],
+			]
+		}
+	},
 
 	// Default deployment parameters for other tools bundled with Horde. These can be queried at runtime without auth using the default Horde server.
 	"parameters": {
