Horde: Add a name to the PluginConfigCollection schema type, to force it to be expanded like a regular object in the generated documentation.

#rnx

diff --git a/Engine/Source/Programs/Horde/HordeServer/Utilities/JsonSchemaCache.cs b/Engine/Source/Programs/Horde/HordeServer/Utilities/JsonSchemaCache.cs
--- a/Engine/Source/Programs/Horde/HordeServer/Utilities/JsonSchemaCache.cs
+++ b/Engine/Source/Programs/Horde/HordeServer/Utilities/JsonSchemaCache.cs
@@ -23,6 +23,7 @@ namespace HordeServer.Utilities
 			_factory = new JsonSchemaFactory(new XmlDocReader());
 
 			JsonSchemaObject pluginObj = new JsonSchemaObject();
+			pluginObj.Name = "PluginsConfig";
 			foreach (ILoadedPlugin plugin in pluginCollection.LoadedPlugins)
 			{
 				if (plugin.GlobalConfigType != null)
