Horde: Fix exception merging two deeply nested config objects.

#rnx

diff --git a/Engine/Source/Programs/Horde/HordeServer.Shared/Configuration/ConfigNode.cs b/Engine/Source/Programs/Horde/HordeServer.Shared/Configuration/ConfigNode.cs
--- a/Engine/Source/Programs/Horde/HordeServer.Shared/Configuration/ConfigNode.cs
+++ b/Engine/Source/Programs/Horde/HordeServer.Shared/Configuration/ConfigNode.cs
@@ -598,21 +598,38 @@ namespace HordeServer.Configuration
 		{
 			if (source is JsonObject sourceObj && target is JsonObject targetObj)
 			{
-				foreach ((string name, JsonNode? node) in sourceObj)
+				JsonObject result = new JsonObject();
+				foreach ((string name, JsonNode? targetNode) in targetObj)
 				{
-					targetObj[name] = Merge(node, targetObj[name]);
+					JsonNode? sourceNode = sourceObj[name];
+					if (sourceNode == null)
+					{
+						result[name] = targetNode?.DeepClone();
+					}
+					else
+					{
+						result[name] = Merge(sourceNode, targetNode);
+					}
 				}
-				return target;
+				foreach ((string name, JsonNode? sourceNode) in sourceObj)
+				{
+					if (!result.ContainsKey(name))
+					{
+						result[name] = sourceNode?.DeepClone();
+					}
+				}
+				return result;
 			}
 			else if (source is JsonArray sourceArr && target is JsonArray targetArr)
 			{
+				JsonArray result = (JsonArray)targetArr.DeepClone();
 				foreach (JsonNode? node in sourceArr)
 				{
-					targetArr.Add(node?.DeepClone());
+					result.Add(node?.DeepClone());
 				}
-				return target;
+				return result;
 			}
-			return source;
+			return source?.DeepClone();
 		}
 
 		/// <inheritdoc/>
