Horde: Fix generated documentation for plugin config objects.

#rnx

diff --git a/Engine/Source/Programs/Horde/HordeServer/Utilities/JsonSchemaCache.cs b/Engine/Source/Programs/Horde/HordeServer/Utilities/JsonSchemaCache.cs
--- a/Engine/Source/Programs/Horde/HordeServer/Utilities/JsonSchemaCache.cs
+++ b/Engine/Source/Programs/Horde/HordeServer/Utilities/JsonSchemaCache.cs
@@ -22,17 +22,33 @@ namespace HordeServer.Utilities
 		{
 			_factory = new JsonSchemaFactory(new XmlDocReader());
 
-			JsonSchemaObject pluginObj = new JsonSchemaObject();
-			pluginObj.Name = "PluginsConfig";
+			// Create a dummy class for the global plugin configuration, which has named properties for each plugin
+			JsonSchemaObject globalPluginsObj = new JsonSchemaObject();
+			globalPluginsObj.Name = "GlobalPluginsConfig";
 			foreach (ILoadedPlugin plugin in pluginCollection.LoadedPlugins)
 			{
 				if (plugin.GlobalConfigType != null)
 				{
 					JsonSchemaType pluginSchemaType = _factory.CreateSchemaType(plugin.GlobalConfigType);
-					pluginObj.Properties.Add(new JsonSchemaProperty(plugin.Name.ToString(), $"Configuration for the {plugin.Name} plugin", pluginSchemaType));
+					globalPluginsObj.Properties.Add(new JsonSchemaProperty(plugin.Name.ToString(), $"Configuration for the {plugin.Name} plugin", pluginSchemaType));
 				}
 			}
-			_factory.TypeCache.Add(typeof(PluginConfigCollection), pluginObj);
+			_factory.TypeCache.Add(typeof(PluginConfigCollection), globalPluginsObj);
+
+			// Same thing for the server config. This is handled by adding a custom property, since config binding for app settings is different to the global config parser.
+			JsonSchemaObject serverPluginsObj = new JsonSchemaObject();
+			serverPluginsObj.Name = "ServerPluginsConfig";
+			foreach (ILoadedPlugin plugin in pluginCollection.LoadedPlugins)
+			{
+				if (plugin.ServerConfigType != null && plugin.ServerConfigType != typeof(object))
+				{
+					JsonSchemaType pluginSchemaType = _factory.CreateSchemaType(plugin.ServerConfigType);
+					serverPluginsObj.Properties.Add(new JsonSchemaProperty(plugin.Name.ToString(), $"Configuration for the {plugin.Name} plugin", pluginSchemaType));
+				}
+			}
+
+			JsonSchemaObject serverObj = (JsonSchemaObject)_factory.CreateSchemaType(typeof(ServerSettings));
+			serverObj.Properties.Add(new JsonSchemaProperty("Plugins", "Configuration for plugins", serverPluginsObj));
 		}
 
 		/// <summary>
