Horde: Updating schema documentation

#rnx

diff --git a/Engine/Source/Programs/Horde/Docs/Config/Schema/Globals.md b/Engine/Source/Programs/Horde/Docs/Config/Schema/Globals.md
--- a/Engine/Source/Programs/Horde/Docs/Config/Schema/Globals.md
+++ b/Engine/Source/Programs/Horde/Docs/Config/Schema/Globals.md
@@ -11,7 +11,7 @@ Name | Description
 `macros` | [ConfigMacro](#configmacro)`[]`<br>Macros within the global scope
 `dashboard` | [DashboardConfig](Dashboard.md)<br>Settings for the dashboard
 `downtime` | [ScheduledDowntime](#scheduleddowntime)`[]`<br>List of scheduled downtime
-`plugins` | `object`<br>Plugin config objects
+`plugins` | [GlobalPluginsConfig](#globalpluginsconfig)<br>Plugin config objects
 `parameters` | `object`<br>General parameters for other tools. Can be queried through the api/v1/parameters endpoint.
 `acl` | [AclConfig](#aclconfig)<br>Access control list
 
@@ -52,6 +52,19 @@ Name | Description
 `Daily` | Every day
 `Weekly` | Every week
 
+## GlobalPluginsConfig
+
+Name | Description
+---- | -----------
+`analytics` | [AnalyticsConfig](#analyticsconfig)<br>Configuration for the analytics plugin
+`build` | [BuildConfig](#buildconfig)<br>Configuration for the build plugin
+`compute` | [ComputeConfig](#computeconfig)<br>Configuration for the compute plugin
+`ddc` | [EmptyPluginConfig](#emptypluginconfig)<br>Configuration for the ddc plugin
+`secrets` | [SecretsConfig](#secretsconfig)<br>Configuration for the secrets plugin
+`storage` | [StorageConfig](#storageconfig)<br>Configuration for the storage plugin
+`symbols` | [SymbolsConfig](#symbolsconfig)<br>Configuration for the symbols plugin
+`tools` | [ToolsConfig](#toolsconfig)<br>Configuration for the tools plugin
+
 ## AnalyticsConfig
 
 Config settings for analytics
diff --git a/Engine/Source/Programs/Horde/Docs/Deployment/ServerSettings.md b/Engine/Source/Programs/Horde/Docs/Deployment/ServerSettings.md
--- a/Engine/Source/Programs/Horde/Docs/Deployment/ServerSettings.md
+++ b/Engine/Source/Programs/Horde/Docs/Deployment/ServerSettings.md
@@ -22,7 +22,7 @@ Name | Description
 `databaseReadOnlyMode` | `boolean`<br>
 `shutdownMemoryThreshold` | `integer`<br>Shutdown the current server process if memory usage reaches this threshold (specified in MB)<br>Usually set to 80-90% of available memory to avoid CLR heap using all of it. If a memory leak was to occur, it's usually better to restart the process rather than to let the GC work harder and harder trying to recoup memory.<br>Should only be used when multiple server processes are running behind a load balancer and one can be safely restarted automatically by the underlying process handler (Docker, Kubernetes, AWS ECS, Supervisor etc). The shutdown behaves similar to receiving a SIGTERM and will wait for outstanding requests to finish.
 `serverPrivateCert` | `string`<br>Optional PFX certificate to use for encrypting agent SSL traffic. This can be a self-signed certificate, as long as it's trusted by agents.
-`authMethod` | [AuthMethod](#authmethod-enum)<br>Type of authentication (e.g anonymous, OIDC, built-in Horde accounts)
+`authMethod` | [AuthMethod](#authmethod-enum)<br>Type of authentication (e.g anonymous, OIDC, built-in Horde accounts) If "Horde" auth mode is used, be sure to configure "ServerUrl" as well.
 `oidcProfileName` | `string`<br>Optional profile name to report through the /api/v1/server/auth endpoint. Allows sharing auth tokens between providers configured through the same profile name in OidcToken.exe config files.
 `oidcAuthority` | `string`<br>OpenID Connect (OIDC) authority URL (required when OIDC is enabled)
 `oidcAudience` | `string`<br>Audience for validating externally issued tokens (required when OIDC is enabled)
@@ -36,7 +36,7 @@ Name | Description
 `oidcClaimEmailMapping` | `string[]`<br>List of fields in /userinfo endpoint to try map to the standard email claim (see System.Security.Claims.ClaimTypes.Email)
 `oidcClaimHordeUserMapping` | `string[]`<br>List of fields in /userinfo endpoint to try map to the Horde user claim (see HordeClaimTypes.User)
 `oidcClaimHordePerforceUserMapping` | `string[]`<br>List of fields in /userinfo endpoint to try map to the Horde Perforce user claim (see HordeClaimTypes.PerforceUser)
-`serverUrl` | `string`<br>Name of this machine
+`serverUrl` | `string`<br>Base URL this Horde server is accessible from For example https://horde.mystudio.com If not set, a default is used based on current hostname. In more advanced setups where a reverse proxy is present in front of Horde, this must manually be set.
 `jwtIssuer` | `string`<br>Name of the issuer in bearer tokens from the server
 `jwtExpiryTimeHours` | `integer`<br>Length of time before JWT tokens expire, in hours
 `adminClaimType` | `string`<br>The claim type for administrators
@@ -66,6 +66,7 @@ Name | Description
 `openBrowser` | `boolean`<br>Whether to open a browser on startup
 `featureFlags` | [FeatureFlagSettings](#featureflagsettings)<br>Experimental features to enable on the server.
 `openTelemetry` | [OpenTelemetrySettings](#opentelemetrysettings)<br>Options for OpenTelemetry
+`plugins` | [ServerPluginsConfig](#serverpluginsconfig)<br>Configuration for plugins
 
 ## RunMode (Enum)
 
@@ -117,3 +118,357 @@ Name | Description
 ---- | -----------
 `endpoint` | `string`<br>Endpoint URL. Usually differs depending on protocol used.
 `protocol` | `string`<br>Protocol for the exporter ('grpc' or 'httpprotobuf')
+
+## ServerPluginsConfig
+
+Name | Description
+---- | -----------
+`analytics` | [AnalyticsServerConfig](#analyticsserverconfig)<br>Configuration for the analytics plugin
+`build` | [StaticBuildConfig](#staticbuildconfig)<br>Configuration for the build plugin
+`compute` | [StaticComputeConfig](#staticcomputeconfig)<br>Configuration for the compute plugin
+`storage` | [StaticStorageConfig](#staticstorageconfig)<br>Configuration for the storage plugin
+`tools` | [ToolsServerConfig](#toolsserverconfig)<br>Configuration for the tools plugin
+
+## AnalyticsServerConfig
+
+Server configuration for the analytics system
+
+Name | Description
+---- | -----------
+`sinks` | [TelemetrySinkConfig](#telemetrysinkconfig)<br>Settings for the various telemetry sinks
+
+## TelemetrySinkConfig
+
+Telemetry sinks
+
+Name | Description
+---- | -----------
+`epic` | [EpicTelemetryConfig](#epictelemetryconfig)<br>Settings for the Epic telemetry sink
+`mongo` | [MongoTelemetryConfig](#mongotelemetryconfig)<br>Settings for the MongoDB telemetry sink
+
+## EpicTelemetryConfig
+
+Configuration for the telemetry sink
+
+Name | Description
+---- | -----------
+`url` | `string`<br>Base URL for the telemetry server
+`appId` | `string`<br>Application name to send in the event messages
+`enabled` | `boolean`<br>Whether to enable this sink
+
+## MongoTelemetryConfig
+
+Configuration for the telemetry sink
+
+Name | Description
+---- | -----------
+`retainDays` | `number`<br>Number of days worth of telmetry events to keep
+`enabled` | `boolean`<br>Whether to enable this sink
+
+## StaticBuildConfig
+
+Static configuration for the build plugin
+
+Name | Description
+---- | -----------
+`perforce` | [PerforceConnectionSettings](#perforceconnectionsettings)`[]`<br>Perforce connections for use by the Horde server (not agents)
+`useLocalPerforceEnv` | `boolean`<br>Whether to use the local Perforce environment
+`perforceConnectionPoolSize` | `integer`<br>Number of pooled perforce connections to keep
+`enableConformTasks` | `boolean`<br>Whether to enable the conform task source.
+`p4SwarmUrl` | `string`<br>Url of P4 Swarm installation
+`jiraUsername` | `string`<br>The Jira service account user name
+`jiraApiToken` | `string`<br>The Jira service account API token
+`jiraUrl` | `string`<br>The Uri for the Jira installation
+`sharedDeviceCheckoutDays` | `integer`<br>The number of days shared device checkouts are held
+`deviceProblemCooldownMinutes` | `integer`<br>The number of cooldown minutes for device problems
+`deviceReportChannel` | `string`<br>Channel to send device reports to
+`disableSchedules` | `boolean`<br>Whether to run scheduled jobs.
+`slackToken` | `string`<br>Bot token for interacting with Slack (xoxb-*)
+`slackSocketToken` | `string`<br>Token for opening a socket to slack (xapp-*)
+`slackAdminToken` | `string`<br>Admin user token for Slack (xoxp-*). This is only required when using the admin endpoints to invite users.
+`slackUsers` | `string`<br>Filtered list of slack users to send notifications to. Should be Slack user ids, separated by commas.
+`slackErrorPrefix` | `string`<br>Prefix to use when reporting errors
+`slackWarningPrefix` | `string`<br>Prefix to use when reporting warnings
+`configNotificationChannel` | `string`<br>Channel for sending messages related to config update failures
+`updateStreamsNotificationChannel` | `string`<br>Channel to send stream notification update failures to
+`jobNotificationChannel` | `string`<br>Slack channel to send job related notifications to. Multiple channels can be specified, separated by ;
+`agentNotificationChannel` | `string`<br>Slack channel to send agent related notifications to.
+`testDataRetainMonths` | `integer`<br>The number of months to retain test data
+`blockCacheDir` | `string`<br>Directory to store the fine-grained block cache. This caches individual exports embedded in bundles.
+`blockCacheSize` | `string`<br>Maximum size of the block cache. Accepts standard binary suffixes. Currently only allocates in multiples of 1024mb.
+`blockCacheSizeBytes` | `integer`<br>Accessor for the block cache size in bytes
+`commits` | [CommitSettings](#commitsettings)<br>Options for the commit service
+
+## PerforceConnectionSettings
+
+Perforce connection information for use by the Horde server (for reading config files, etc...)
+
+Name | Description
+---- | -----------
+`id` | `string`<br>Identifier for this server
+`serverAndPort` | `string`<br>Server and port
+`credentials` | [PerforceCredentials](#perforcecredentials)<br>Credentials for the server
+
+## PerforceCredentials
+
+Credentials for a Perforce user
+
+Name | Description
+---- | -----------
+`userName` | `string`<br>The username
+`password` | `string`<br>Password for the user
+`ticket` | `string`<br>Login ticket for the user (will be used instead of password if set)
+
+## CommitSettings
+
+Options for the commit service
+
+Name | Description
+---- | -----------
+`replicateMetadata` | `boolean`<br>Whether to mirror commit metadata to the database
+`replicateContent` | `boolean`<br>Whether to mirror commit data to storage
+`bundle` | [BundleOptions](#bundleoptions)<br>Options for how objects are packed together
+`chunking` | [ChunkingOptions](#chunkingoptions)<br>Options for how objects are sliced
+
+## BundleOptions
+
+Options for configuring a bundle serializer
+
+Name | Description
+---- | -----------
+`maxVersion` | [BundleVersion](#bundleversion-enum)<br>Maximum version number of bundles to write
+`maxBlobSize` | `integer`<br>Maximum payload size fo a blob
+`compressionFormat` | [BundleCompressionFormat](#bundlecompressionformat-enum)<br>Compression format to use
+`minCompressionPacketSize` | `integer`<br>Minimum size of a block to be compressed
+`maxWriteQueueLength` | `integer`<br>Maximum amount of data to store in memory. This includes any background writes as well as bundles being built.
+
+## BundleVersion (Enum)
+
+Bundle version number
+
+Name | Description
+---- | -----------
+`Initial` | Initial version number
+`ExportAliases` | Added the BundleExport.Alias property
+`RemoveAliases` | Back out change to include aliases. Will likely do this through an API rather than baked into the data.
+`InPlace` | Use data structures which support in-place reading and writing.
+`ImportHashes` | Add import hashes to imported nodes
+`LatestV1` | Last version using the V1 pipeline
+`PacketSequence` | Structure bundles as a sequence of self-contained packets (uses V2 code)
+`Latest` | The current version number
+`LatestV2` | Last version using the V2 pipeline
+`LatestPlusOne` | Last item in the enum. Used for
+
+## BundleCompressionFormat (Enum)
+
+Indicates the compression format in the bundle
+
+Name | Description
+---- | -----------
+`None` | Packets are uncompressed
+`LZ4` | LZ4 compression
+`Gzip` | Gzip compression
+`Oodle` | Oodle compression (Selkie)
+`Brotli` | Brotli compression
+`Zstd` | ZStandard compression
+
+## ChunkingOptions
+
+Options for creating file nodes
+
+Name | Description
+---- | -----------
+`leafOptions` | [LeafChunkedDataNodeOptions](#leafchunkeddatanodeoptions)<br>Options for creating leaf nodes
+`interiorOptions` | [InteriorChunkedDataNodeOptions](#interiorchunkeddatanodeoptions)<br>Options for creating interior nodes
+
+## LeafChunkedDataNodeOptions
+
+Options for creating a specific type of file nodes
+
+Name | Description
+---- | -----------
+`minSize` | `integer`<br>Minimum chunk size
+`maxSize` | `integer`<br>Maximum chunk size. Chunks will be split on this boundary if another match is not found.
+`targetSize` | `integer`<br>Target chunk size for content-slicing
+
+## InteriorChunkedDataNodeOptions
+
+Options for creating interior nodes
+
+Name | Description
+---- | -----------
+`minChildCount` | `integer`<br>Minimum number of children in each node
+`targetChildCount` | `integer`<br>Target number of children in each node
+`maxChildCount` | `integer`<br>Maximum number of children in each node
+`sliceThreshold` | `integer`<br>Threshold hash value for splitting interior nodes
+
+## StaticComputeConfig
+
+Static configuration for the compute plugin
+
+Name | Description
+---- | -----------
+`enableUpgradeTasks` | `boolean`<br>Whether to enable the upgrade task source.
+`withAws` | `boolean`<br>Whether to enable Amazon Web Services (AWS) specific features
+`awsRegions` | `string[]`<br>List of AWS regions for Horde to be aware of (e.g. us-east-1 or eu-central-1) Right now, this is only used for replicating CloudWatch metrics to multiple regions
+`awsAutoScalingQueueUrls` | `string[]`<br>AWS SQS queue URLs where lifecycle events from EC2 auto-scaling are received
+`fleetManagerV2` | [FleetManagerType](#fleetmanagertype-enum)<br>Default fleet manager to use (when not specified by pool)
+`fleetManagerV2Config` | `object`<br>Config for the fleet manager (serialized JSON)
+`autoEnrollAgents` | `boolean`<br>Whether to automatically enroll agents in the farm
+`defaultAgentPoolSizeStrategy` | [PoolSizeStrategy](#poolsizestrategy-enum)<br>Default agent pool sizing strategy for pools that doesn't have one explicitly configured
+`agentPoolScaleOutCooldownSeconds` | `integer`<br>Scale-out cooldown for auto-scaling agent pools (in seconds). Can be overridden by per-pool settings.
+`agentPoolScaleInCooldownSeconds` | `integer`<br>Scale-in cooldown for auto-scaling agent pools (in seconds). Can be overridden by per-pool settings.
+`computeTunnelPort` | `integer`<br>Port to listen on for tunneling compute sockets to agents
+`computeTunnelAddress` | `string`<br>What address (host:port) clients should connect to for compute socket tunneling Port may differ from  if Horde server is behind a reverse proxy/firewall
+
+## FleetManagerType (Enum)
+
+Available fleet managers
+
+Name | Description
+---- | -----------
+`Default` | Default fleet manager
+`NoOp` | No-op fleet manager.
+`Aws` | Fleet manager for handling AWS EC2 instances. Will create and/or terminate instances from scratch.
+`AwsReuse` | Fleet manager for handling AWS EC2 instances. Will start already existing but stopped instances to reuse existing EBS disks.
+`AwsRecycle` | Fleet manager for handling AWS EC2 instances. Will start already existing but stopped instances to reuse existing EBS disks.
+`AwsAsg` | Fleet manager for handling AWS EC2 instances. Uses an EC2 auto-scaling group for controlling the number of running instances.
+
+## PoolSizeStrategy (Enum)
+
+Available pool sizing strategies
+
+Name | Description
+---- | -----------
+`LeaseUtilization` | Strategy based on lease utilization
+`JobQueue` | Strategy based on size of job build queue
+`NoOp` | No-op strategy used as fallback/default behavior
+`ComputeQueueAwsMetric` | A no-op strategy that reports metrics to let an external AWS auto-scaling policy scale the fleet
+`LeaseUtilizationAwsMetric` | A no-op strategy that reports metrics to let an external AWS auto-scaling policy scale the fleet
+
+## StaticStorageConfig
+
+Static settings for the storage system
+
+Name | Description
+---- | -----------
+`bundleCacheDir` | `string`<br>Directory to use for the coarse-grained backend cache. This caches full bundles downloaded from the upstream object store.
+`bundleCacheSize` | `string`<br>Maximum size of the storage cache on disk. Accepts standard binary suffixes (kb, mb, gb, tb, etc...)
+`bundleCacheSizeBytes` | `integer`<br>Accessor for the bundle cache size in bytes
+`backends` | [BackendConfig](#backendconfig)`[]`<br>Overridden settings for storage backends. Useful for running against a production server with custom backends.
+
+## BackendConfig
+
+Common settings object for different providers
+
+Name | Description
+---- | -----------
+`id` | `string`<br>The storage backend ID
+`base` | `string`<br>Base backend to copy default settings from
+`secondary` | `string`<br>Specifies another backend to read from if an object is not found in this one. Can be used when migrating data from one backend to another.
+`type` | [StorageBackendType](#storagebackendtype-enum)<br>
+`baseDir` | `string`<br>
+`awsBucketName` | `string`<br>Name of the bucket to use
+`awsBucketPath` | `string`<br>Base path within the bucket
+`awsCredentials` | [AwsCredentialsType](#awscredentialstype-enum)<br>Type of credentials to use
+`awsRole` | `string`<br>ARN of a role to assume
+`awsProfile` | `string`<br>The AWS profile to read credentials form
+`awsRegion` | `string`<br>Region to connect to
+`azureConnectionString` | `string`<br>Connection string for Azure
+`azureContainerName` | `string`<br>Name of the container
+`relayServer` | `string`<br>
+`relayToken` | `string`<br>
+
+## StorageBackendType (Enum)
+
+Types of storage backend to use
+
+Name | Description
+---- | -----------
+`FileSystem` | Local filesystem
+`Aws` | AWS S3
+`Azure` | Azure blob store
+`Memory` | In-memory only (for testing)
+
+## AwsCredentialsType (Enum)
+
+Credentials to use for AWS
+
+Name | Description
+---- | -----------
+`Default` | Use default credentials from the AWS SDK
+`Profile` | Read credentials from the  profile in the AWS config file
+`AssumeRole` | Assume a particular role. Should specify ARN in
+`AssumeRoleWebIdentity` | Assume a particular role using the current environment variables.
+
+## ToolsServerConfig
+
+Server configuration for bundled tools
+
+Name | Description
+---- | -----------
+`bundledTools` | [BundledToolConfig](#bundledtoolconfig)`[]`<br>Tools bundled along with the server. Data for each tool can be produced using the 'bundle create' command, and should be stored in the Tools directory.
+
+## BundledToolConfig
+
+Configuration for a tool bundled alongsize the server
+
+Name | Description
+---- | -----------
+`version` | `string`<br>Version string for the current tool data
+`refName` | `string`<br>Ref name in the tools directory
+`dataDir` | `string`<br>Directory containing blob data for this tool. If empty, the tools/{id} folder next to the server will be used.
+`id` | `string`<br>Unique identifier for the tool
+`name` | `string`<br>Name of the tool
+`description` | `string`<br>Description for the tool
+`category` | `string`<br>Category for the tool. Will cause the tool to be shown in a different tab in the dashboard.
+`group` | `string`<br>Grouping key for different variations of the same tool. The dashboard will show these together.
+`platforms` | `string[]`<br>Platforms for this tool. Takes the form of a NET RID (https://learn.microsoft.com/en-us/dotnet/core/rid-catalog).
+`public` | `boolean`<br>Whether this tool should be exposed for download on a public endpoint without authentication
+`showInUgs` | `boolean`<br>Whether to show this tool for download in the UGS tools menu
+`showInDashboard` | `boolean`<br>Whether to show this tool for download in the dashboard
+`showInToolbox` | `boolean`<br>Whether to show this tool for download in Unreal Toolbox
+`metadata` | `string` `->` `string`<br>Metadata for this tool
+`namespaceId` | `string`<br>Default namespace for new deployments of this tool
+`acl` | [AclConfig](#aclconfig)<br>Permissions for the tool
+
+## AclConfig
+
+Parameters to update an ACL
+
+Name | Description
+---- | -----------
+`entries` | [AclEntryConfig](#aclentryconfig)`[]`<br>Entries to replace the existing ACL
+`profiles` | [AclProfileConfig](#aclprofileconfig)`[]`<br>Defines profiles which allow grouping sets of actions into named collections
+`inherit` | `boolean`<br>Whether to inherit permissions from the parent ACL
+`exceptions` | `string[]`<br>List of exceptions to the inherited setting
+
+## AclEntryConfig
+
+Individual entry in an ACL
+
+Name | Description
+---- | -----------
+`claim` | [AclClaimConfig](#aclclaimconfig)<br>Name of the user or group
+`actions` | `string[]`<br>Array of actions to allow
+`profiles` | `string[]`<br>List of profiles to grant
+
+## AclClaimConfig
+
+New claim to create
+
+Name | Description
+---- | -----------
+`type` | `string`<br>The claim type
+`value` | `string`<br>The claim value
+
+## AclProfileConfig
+
+Configuration for an ACL profile. This defines a preset group of actions which can be given to a user via an ACL entry.
+
+Name | Description
+---- | -----------
+`id` | `string`<br>Identifier for this profile
+`actions` | `string[]`<br>Actions to include
+`excludeActions` | `string[]`<br>Actions to exclude from the inherited actions
+`extends` | `string[]`<br>Other profiles to extend from
