Horde: Bump compute protocol version for new CPU env vars

Also make ComputeProtocol a constant and add test for allocation.

diff --git a/Engine/Source/Programs/Horde/HordeAgent/Services/CapabilitiesService.cs b/Engine/Source/Programs/Horde/HordeAgent/Services/CapabilitiesService.cs
--- a/Engine/Source/Programs/Horde/HordeAgent/Services/CapabilitiesService.cs
+++ b/Engine/Source/Programs/Horde/HordeAgent/Services/CapabilitiesService.cs
@@ -463,7 +463,7 @@ namespace HordeAgent.Services
 			}
 
 			// Add the max supported compute protocol version
-			capabilities.Properties.Add($"ComputeProtocol={(int)ComputeProtocol.Latest}");
+			capabilities.Properties.Add($"{KnownPropertyNames.ComputeProtocol}={(int)ComputeProtocol.Latest}");
 
 			// Whether the agent is packaged as a self-contained .NET app
 			// Used during the transition period over from multi-platform, non-self-contained agent packages.
diff --git a/Engine/Source/Programs/Horde/Plugins/Build/HordeServer.Build.Tests/Compute/ComputeServiceTest.cs b/Engine/Source/Programs/Horde/Plugins/Build/HordeServer.Build.Tests/Compute/ComputeServiceTest.cs
--- a/Engine/Source/Programs/Horde/Plugins/Build/HordeServer.Build.Tests/Compute/ComputeServiceTest.cs
+++ b/Engine/Source/Programs/Horde/Plugins/Build/HordeServer.Build.Tests/Compute/ComputeServiceTest.cs
@@ -2,6 +2,7 @@
 
 using System.Diagnostics.Metrics;
 using System.Net;
+using EpicGames.Horde.Agents;
 using EpicGames.Horde.Agents.Pools;
 using EpicGames.Horde.Compute;
 using HordeServer.Acls;
@@ -244,6 +245,19 @@ namespace HordeServer.Tests.Compute
 			await Assert.ThrowsExceptionAsync<NoComputeResourcesException>(() => ComputeService.TryAllocateResourceAsync(CreateParams("does-not-exist"), CancellationToken.None));
 		}
 		
+		[TestMethod]
+		public async Task Assignment_Protocol_Async()
+		{
+			await CreateComputeAgentAsync("pool1", protocol: ComputeProtocol.Initial);
+			Assert.AreEqual(ComputeProtocol.Initial, (await AllocateAsync("pool1"))!.Protocol);
+			
+			await CreateComputeAgentAsync("pool2", protocol: ComputeProtocol.Latest);
+			Assert.AreEqual(ComputeProtocol.Latest, (await AllocateAsync("pool2"))!.Protocol);
+			
+			await CreateComputeAgentAsync("pool3", protocol: ComputeProtocol.Latest);
+			Assert.AreEqual(ComputeProtocol.Initial, (await AllocateAsync("pool3", protocol: ComputeProtocol.Initial))!.Protocol);
+		}
+		
 		[TestMethod]
 		public async Task Assignment_NoResources_Async()
 		{
@@ -338,7 +352,7 @@ namespace HordeServer.Tests.Compute
 		[TestMethod]
 		public async Task Connection_Direct_IpConnection_Async()
 		{
-			ComputeResource? cr = await AllocateAsync(ConnectionMode.Direct);
+			ComputeResource? cr = await AllocateWithAgentAsync(ConnectionMode.Direct);
 			Assert.AreEqual(ConnectionMode.Direct, cr!.ConnectionMode);
 			Assert.AreEqual("11.0.0.1", cr.Ip.ToString());
 			Assert.IsNull(cr.ConnectionAddress);
@@ -347,7 +361,7 @@ namespace HordeServer.Tests.Compute
 		[TestMethod]
 		public async Task Connection_Direct_PortsAreMapped_Async()
 		{
-			ComputeResource? cr = await AllocateAsync(ConnectionMode.Direct, ports: new Dictionary<string, int> { { "myOtherPort", 13000 }, { "myPort", 12000 } });
+			ComputeResource? cr = await AllocateWithAgentAsync(ConnectionMode.Direct, ports: new Dictionary<string, int> { { "myOtherPort", 13000 }, { "myPort", 12000 } });
 			Assert.AreEqual(3, cr!.Ports.Count);
 			Assert.AreEqual(new ComputeResourcePort(5000, 5000), cr.Ports[ConnectionMetadataPort.ComputeId]);
 			Assert.AreEqual(new ComputeResourcePort(12000, 12000), cr.Ports["myPort"]);
@@ -357,7 +371,7 @@ namespace HordeServer.Tests.Compute
 		[TestMethod]
 		public async Task Connection_Tunnel_IpConnection_Async()
 		{
-			ComputeResource? cr = await AllocateAsync(ConnectionMode.Tunnel, tunnelAddress: "localhost:3344");
+			ComputeResource? cr = await AllocateWithAgentAsync(ConnectionMode.Tunnel, tunnelAddress: "localhost:3344");
 			Assert.AreEqual(ConnectionMode.Tunnel, cr!.ConnectionMode);
 			Assert.AreEqual("11.0.0.1", cr.Ip.ToString());
 			Assert.AreEqual("localhost:3344", cr.ConnectionAddress);
@@ -366,7 +380,7 @@ namespace HordeServer.Tests.Compute
 		[TestMethod]
 		public async Task Connection_Tunnel_PortsAreMapped_Async()
 		{
-			ComputeResource? cr = await AllocateAsync(ConnectionMode.Tunnel, tunnelAddress: "localhost:3344", ports: new Dictionary<string, int> { { "myOtherPort", 13000 }, { "myPort", 12000 } });
+			ComputeResource? cr = await AllocateWithAgentAsync(ConnectionMode.Tunnel, tunnelAddress: "localhost:3344", ports: new Dictionary<string, int> { { "myOtherPort", 13000 }, { "myPort", 12000 } });
 			Assert.AreEqual(3, cr!.Ports.Count);
 			Assert.AreEqual(new ComputeResourcePort(-1, 5000), cr.Ports[ConnectionMetadataPort.ComputeId]);
 			Assert.AreEqual(new ComputeResourcePort(-1, 12000), cr.Ports["myPort"]);
@@ -376,7 +390,7 @@ namespace HordeServer.Tests.Compute
 		[TestMethod]
 		public async Task Connection_Relay_IpConnection_Async()
 		{
-			ComputeResource? cr = await AllocateAsync(ConnectionMode.Relay);
+			ComputeResource? cr = await AllocateWithAgentAsync(ConnectionMode.Relay);
 			Assert.AreEqual(ConnectionMode.Relay, cr!.ConnectionMode);
 			Assert.AreEqual("11.0.0.1", cr.Ip.ToString());
 			Assert.AreEqual("192.168.1.1", cr.ConnectionAddress);
@@ -385,7 +399,7 @@ namespace HordeServer.Tests.Compute
 		[TestMethod]
 		public async Task Connection_Relay_PortsAreMapped_Async()
 		{
-			ComputeResource? cr = await AllocateAsync(ConnectionMode.Relay, ports: new Dictionary<string, int> { { "myOtherPort", 13000 }, { "myPort", 12000 } });
+			ComputeResource? cr = await AllocateWithAgentAsync(ConnectionMode.Relay, ports: new Dictionary<string, int> { { "myOtherPort", 13000 }, { "myPort", 12000 } });
 			Assert.AreEqual(ConnectionMode.Relay, cr!.ConnectionMode);
 			Assert.AreEqual(3, cr.Ports.Count);
 			Assert.AreEqual(new ComputeResourcePort(12214, 5000), cr.Ports[ConnectionMetadataPort.ComputeId]);
@@ -407,16 +421,17 @@ namespace HordeServer.Tests.Compute
 			return cs;
 		}
 
-		private async Task<ComputeResource?> AllocateAsync(
+		private async Task<ComputeResource?> AllocateWithAgentAsync(
 			ConnectionMode connectionMode = ConnectionMode.Direct,
 			Dictionary<string, int>? ports = null,
 			bool usePublicIp = false,
 			string[]? relayIps = null,
 			string? tunnelAddress = null,
+			ComputeProtocol protocol = ComputeProtocol.Latest,
 			string pool = "myPool")
 		{
 			await using ComputeService cs = await CreateComputeServiceAsync(tunnelAddress, pool, new ComputeClusterConfig { Id = _cluster1 });
-			AllocateResourceParams arp = new(_cluster1, ComputeProtocol.Latest, new Requirements() { Pool = pool })
+			AllocateResourceParams arp = new(_cluster1, protocol, new Requirements() { Pool = pool })
 			{
 				ConnectionMode = connectionMode,
 				Ports = ports ?? new Dictionary<string, int>(),
@@ -426,6 +441,32 @@ namespace HordeServer.Tests.Compute
 			await AgentRelayService.UpdateAgentHeartbeatAsync(_cluster1, "myrelay", relayIps?.Select(IPAddress.Parse) ?? defaultRelayIps);
 			return await cs.TryAllocateResourceAsync(arp, CancellationToken.None);
 		}
+		
+		private async Task<ComputeResource?> AllocateAsync(
+			string pool = "defaultTestPool",
+			ConnectionMode connectionMode = ConnectionMode.Direct,
+			Dictionary<string, int>? ports = null,
+			bool usePublicIp = false,
+			ComputeProtocol protocol = ComputeProtocol.Latest)
+		{
+			AllocateResourceParams arp = new(new ClusterId("default"), protocol, new Requirements() { Pool = pool })
+			{
+				ConnectionMode = connectionMode,
+				Ports = ports ?? new Dictionary<string, int>(),
+				UsePublicIp = usePublicIp
+			};
+			return await ComputeService.TryAllocateResourceAsync(arp, CancellationToken.None);
+		}
+		
+		private Task<IAgent> CreateComputeAgentAsync(
+			string poolId = "defaultTestPool",
+			string ip = "10.0.0.1",
+			int port = 5000,
+			ComputeProtocol protocol = ComputeProtocol.Latest)
+		{
+			List<string> props = [$"{KnownPropertyNames.ComputeIp}={ip}", $"{KnownPropertyNames.ComputePort}={port}", $"{KnownPropertyNames.ComputeProtocol}={(int)protocol}"];
+			return CreateAgentAsync(new PoolId(poolId), properties: props);
+		}
 
 		private static void AssertContainsMeasurement(List<Measurement<int>> actualMeasurements, int expectedValue, string expectedResource, string expectedClusterId, string expectedPool)
 		{
diff --git a/Engine/Source/Programs/Horde/Plugins/Compute/HordeServer.Compute/Compute/ComputeService.cs b/Engine/Source/Programs/Horde/Plugins/Compute/HordeServer.Compute/Compute/ComputeService.cs
--- a/Engine/Source/Programs/Horde/Plugins/Compute/HordeServer.Compute/Compute/ComputeService.cs
+++ b/Engine/Source/Programs/Horde/Plugins/Compute/HordeServer.Compute/Compute/ComputeService.cs
@@ -442,7 +442,7 @@ namespace HordeServer.Compute
 
 						using TelemetrySpan assignSpan = _tracer.StartActiveSpan("TryAssignAsync");
 
-						ComputeResource? resource = await TryAssignAsync(arp, agent, computeTask, leaseId);
+						ComputeResource? resource = await TryAssignAsync(arp, agent, computeTask, protocol, leaseId);
 						if (resource != null)
 						{
 							using TelemetrySpan addLeaseSpan = _tracer.StartActiveSpan("Adding lease");
@@ -712,7 +712,7 @@ namespace HordeServer.Compute
 			return await _agentScheduler.GetChildLeaseCountAsync(parentLeaseId.Value, cancellationToken);
 		}
 
-		private async Task<ComputeResource?> TryAssignAsync(AllocateResourceParams arp, IAgent agent, ComputeTask computeTask, LeaseId leaseId)
+		private async Task<ComputeResource?> TryAssignAsync(AllocateResourceParams arp, IAgent agent, ComputeTask computeTask, ComputeProtocol protocol, LeaseId leaseId)
 		{
 			using TelemetrySpan span = _tracer.StartActiveSpan($"{nameof(ComputeService)}.{nameof(TryAssignAsync)}");
 
@@ -740,7 +740,7 @@ namespace HordeServer.Compute
 					ports[portId] = new ComputeResourcePort(port, port);
 				}
 
-				return new ComputeResource(ConnectionMode.Direct, agentIp, null, ports, computeTask, agent.Properties, agent.Id, agent.Version, leaseId);
+				return new ComputeResource(ConnectionMode.Direct, agentIp, null, ports, computeTask, agent.Properties, agent.Id, agent.Version, protocol, leaseId);
 			}
 			else if (arp.ConnectionMode == ConnectionMode.Tunnel && tunnelAddress != null)
 			{
@@ -752,7 +752,7 @@ namespace HordeServer.Compute
 					ports[portId] = new ComputeResourcePort(-1, port);
 				}
 
-				return new ComputeResource(ConnectionMode.Tunnel, agentIp, tunnelAddress, ports, computeTask, agent.Properties, agent.Id, agent.Version, leaseId);
+				return new ComputeResource(ConnectionMode.Tunnel, agentIp, tunnelAddress, ports, computeTask, agent.Properties, agent.Id, agent.Version, protocol, leaseId);
 			}
 			else if (arp.ConnectionMode == ConnectionMode.Relay)
 			{
@@ -783,7 +783,7 @@ namespace HordeServer.Compute
 					}
 				}
 
-				return new ComputeResource(ConnectionMode.Relay, agentIp, relayIp.ToString(), ports, computeTask, agent.Properties, agent.Id, agent.Version, leaseId);
+				return new ComputeResource(ConnectionMode.Relay, agentIp, relayIp.ToString(), ports, computeTask, agent.Properties, agent.Id, agent.Version, protocol, leaseId);
 			}
 
 			throw new Exception("Unable to resolve a suitable connection mode for compute task");
diff --git a/Engine/Source/Programs/Horde/Plugins/Compute/HordeServer.Compute/Compute/ComputeTaskSource.cs b/Engine/Source/Programs/Horde/Plugins/Compute/HordeServer.Compute/Compute/ComputeTaskSource.cs
--- a/Engine/Source/Programs/Horde/Plugins/Compute/HordeServer.Compute/Compute/ComputeTaskSource.cs
+++ b/Engine/Source/Programs/Horde/Plugins/Compute/HordeServer.Compute/Compute/ComputeTaskSource.cs
@@ -115,6 +115,11 @@ namespace HordeServer.Compute
 		/// Agent version on the remote machine
 		/// </summary>
 		public string? AgentVersion { get; }
+		
+		/// <summary>
+		/// Compute protocol this resource is using
+		/// </summary>
+		public ComputeProtocol Protocol { get; }
 
 		/// <summary>
 		/// Lease id on the remote machine
@@ -133,6 +138,7 @@ namespace HordeServer.Compute
 			IReadOnlyList<string> properties,
 			AgentId agentId,
 			string? agentVersion,
+			ComputeProtocol protocol,
 			LeaseId leaseId)
 		{
 			ConnectionMode = connectionMode;
@@ -143,6 +149,7 @@ namespace HordeServer.Compute
 			Properties = properties;
 			AgentId = agentId;
 			AgentVersion = agentVersion;
+			Protocol = protocol;
 			LeaseId = leaseId;
 		}
 	}
diff --git a/Engine/Source/Programs/Shared/EpicGames.Horde/Agents/KnownPropertyNames.cs b/Engine/Source/Programs/Shared/EpicGames.Horde/Agents/KnownPropertyNames.cs
--- a/Engine/Source/Programs/Shared/EpicGames.Horde/Agents/KnownPropertyNames.cs
+++ b/Engine/Source/Programs/Shared/EpicGames.Horde/Agents/KnownPropertyNames.cs
@@ -66,6 +66,11 @@ public static class KnownPropertyNames
 	/// Port used for sending compute task payloads
 	/// </summary>
 	public const string ComputePort = "ComputePort";
+	
+	/// <summary>
+	/// Protocol version for compute task payloads
+	/// </summary>
+	public const string ComputeProtocol = "ComputeProtocol";
 
 	/// <summary>
 	/// AWS: Instance ID
diff --git a/Engine/Source/Programs/Shared/EpicGames.Horde/Compute/ComputeProtocol.cs b/Engine/Source/Programs/Shared/EpicGames.Horde/Compute/ComputeProtocol.cs
--- a/Engine/Source/Programs/Shared/EpicGames.Horde/Compute/ComputeProtocol.cs
+++ b/Engine/Source/Programs/Shared/EpicGames.Horde/Compute/ComputeProtocol.cs
@@ -22,11 +22,16 @@ namespace EpicGames.Horde.Compute
 		/// Initial version number
 		/// </summary>
 		Initial,
+		
+		/// <summary>
+		/// Set new env vars UE_HORDE_CPU_COUNT and UE_HORDE_CPU_MULTIPLIER
+		/// </summary>
+		NewCpuEnvVars,
 
 		/// <summary>
 		/// Constant for the latest protocol version
 		/// </summary>
-		Latest = (int)Initial
+		Latest = (int)NewCpuEnvVars
 	}
 
 	/// <summary>
