Horde: Log errors when mongodb config/log files contain non-latin characters as mongodb.exe cannot handle these paths

diff --git a/Engine/Source/Programs/Horde/HordeServer/Server/MongoService.cs b/Engine/Source/Programs/Horde/HordeServer/Server/MongoService.cs
--- a/Engine/Source/Programs/Horde/HordeServer/Server/MongoService.cs
+++ b/Engine/Source/Programs/Horde/HordeServer/Server/MongoService.cs
@@ -338,6 +338,18 @@ namespace HordeServer.Server
 			FileReference mongoLogFile = FileReference.Combine(mongoDir, "mongod.log");
 
 			FileReference configFile = FileReference.Combine(mongoDir, "mongod.conf");
+
+			// Check mongo log and config files for non-latin characters, mongodb.exe cannout read these paths
+			if (Regex.IsMatch(mongoLogFile.ToString(), "[^a-zA-Z0-9:\\\\\\/]"))
+			{
+				logger.LogError("MongoDB log file contains non-latin characters, which mongodb.exe cannot read: {MongoLogFile}", mongoLogFile.ToString());
+			}
+
+			if (Regex.IsMatch(configFile.ToString(), "[^a-zA-Z0-9:\\\\\\/]"))
+			{
+				logger.LogError("MongoDB config file contains non-latin characters, which mongodb.exe cannot read: {MongoConfigFile}", configFile.ToString());
+			}
+
 			if (!FileReference.Exists(configFile))
 			{
 				DirectoryReference.CreateDirectory(configFile.Directory);
