Horde: Mention authentication in remote compilation tutorial

#rnx

diff --git a/Engine/Source/Programs/Horde/Docs/Tutorials/RemoteCompilation.md b/Engine/Source/Programs/Horde/Docs/Tutorials/RemoteCompilation.md
--- a/Engine/Source/Programs/Horde/Docs/Tutorials/RemoteCompilation.md
+++ b/Engine/Source/Programs/Horde/Docs/Tutorials/RemoteCompilation.md
@@ -26,6 +26,7 @@ spreading the workload over multiple machines.
 * Horde Server and one or more Horde Agents (see [Getting Started: Install Horde](InstallHorde.md)).
 * A workstation with a UE project under development.
 * Network connectivity between your workstation and Horde Agents on port range 7000-7010.
+* The default `Anonymous` authentication method enabled (see below for more details)
 
 ## Steps
 
@@ -92,3 +93,40 @@ spreading the workload over multiple machines.
   pools of agents for build automation. Agents used for build automation typically have higher requirements 
   and are a more scarce resource than compute helpers.
 
+## Enabling authentication
+
+When using the anonymous authentication mode, there's no authentication and every user has full access
+to perform remote compilation. This is only recommended during a testing phase and you should switch to either
+`Horde` or `OpenIdConnect` as soon as possible. Once you do that, additional permissions must be granted to the
+so called compute cluster. By default, a cluster called `default` is already defined.
+
+Horde jobs utilizing remote compilation are automatically granted access via an injected token that's set as
+an environment variable for each job step (`UE_HORDE_TOKEN`).
+
+Below is an updated global config with a `default` cluster and `AddComputeTasks` granted for letting UBT and UBA schedule remote compilation.
+To see what claims are available to your users, open `/api/v1/user/claims` as a logged in user.
+```json
+{
+  // ...
+  "plugins": {
+    // ...
+    "compute": {
+      // ...
+      "clusters": [
+        {
+          "id": "default",
+          "namespaceid": "horde.compute",
+          "acl": {
+            "entries": [
+              {
+                "claim": { "type": "http://epicgames.com/ue/horde/user", "value": "jane.smith" },
+                "actions": ["AddComputeTasks"]
+              }
+            ]
+          }
+        }
+      ]
+    }
+  }
+}
+```
\ No newline at end of file
