Horde: Return Empty action results rather than NoContent() for controller methods that manually write responses.

#rnx

diff --git a/Engine/Source/Programs/Horde/Plugins/Build/HordeServer.Build/Streams/StreamsController.cs b/Engine/Source/Programs/Horde/Plugins/Build/HordeServer.Build/Streams/StreamsController.cs
--- a/Engine/Source/Programs/Horde/Plugins/Build/HordeServer.Build/Streams/StreamsController.cs
+++ b/Engine/Source/Programs/Horde/Plugins/Build/HordeServer.Build/Streams/StreamsController.cs
@@ -459,7 +459,7 @@ namespace HordeServer.Streams
 			await channel.FindAsync(HttpContext.Response.BodyWriter, minTime, maxTime, index, count);
 			await Response.CompleteAsync();
 
-			return NoContent();
+			return Empty;
 		}
 
 		/// <summary>
diff --git a/Engine/Source/Programs/Horde/Plugins/Compute/HordeServer.Compute/Agents/AgentsController.cs b/Engine/Source/Programs/Horde/Plugins/Compute/HordeServer.Compute/Agents/AgentsController.cs
--- a/Engine/Source/Programs/Horde/Plugins/Compute/HordeServer.Compute/Agents/AgentsController.cs
+++ b/Engine/Source/Programs/Horde/Plugins/Compute/HordeServer.Compute/Agents/AgentsController.cs
@@ -365,7 +365,7 @@ namespace HordeServer.Agents
 			await _agentService.Agents.GetLogger(agentId).FindAsync(Response.BodyWriter, minTime, maxTime, index, count);
 			await Response.CompleteAsync();
 
-			return NoContent();
+			return Empty;
 		}
 
 		/// <summary>
