Horde: Rename StaticBuildConfig to BuildServerConfig for consistency.

#rnx

diff --git a/Engine/Source/Programs/Horde/Plugins/Build/HordeServer.Build/BuildPlugin.cs b/Engine/Source/Programs/Horde/Plugins/Build/HordeServer.Build/BuildPlugin.cs
--- a/Engine/Source/Programs/Horde/Plugins/Build/HordeServer.Build/BuildPlugin.cs
+++ b/Engine/Source/Programs/Horde/Plugins/Build/HordeServer.Build/BuildPlugin.cs
@@ -40,16 +40,16 @@ namespace HordeServer
 	/// <summary>
 	/// Entry point for the storage plugin
 	/// </summary>
-	[Plugin("Build", GlobalConfigType = typeof(BuildConfig), ServerConfigType = typeof(StaticBuildConfig))]
+	[Plugin("Build", GlobalConfigType = typeof(BuildConfig), ServerConfigType = typeof(BuildServerConfig))]
 	public class BuildPlugin : IPluginStartup
 	{
 		readonly IServerInfo _serverInfo;
-		readonly StaticBuildConfig _staticConfig;
+		readonly BuildServerConfig _staticConfig;
 
 		/// <summary>
 		/// Constructor
 		/// </summary>
-		public BuildPlugin(IServerInfo serverInfo, StaticBuildConfig staticConfig)
+		public BuildPlugin(IServerInfo serverInfo, BuildServerConfig staticConfig)
 		{
 			_serverInfo = serverInfo;
 			_staticConfig = staticConfig;
diff --git a/Engine/Source/Programs/Horde/Plugins/Build/HordeServer.Build/BuildResponseFilter.cs b/Engine/Source/Programs/Horde/Plugins/Build/HordeServer.Build/BuildResponseFilter.cs
--- a/Engine/Source/Programs/Horde/Plugins/Build/HordeServer.Build/BuildResponseFilter.cs
+++ b/Engine/Source/Programs/Horde/Plugins/Build/HordeServer.Build/BuildResponseFilter.cs
@@ -9,9 +9,9 @@ namespace HordeServer
 {
 	class BuildResponseFilter : IPluginResponseFilter
 	{
-		readonly StaticBuildConfig _staticBuildConfig;
+		readonly BuildServerConfig _staticBuildConfig;
 
-		public BuildResponseFilter(IOptions<StaticBuildConfig> staticBuildConfig)
+		public BuildResponseFilter(IOptions<BuildServerConfig> staticBuildConfig)
 		{
 			_staticBuildConfig = staticBuildConfig.Value;
 		}
diff --git a/Engine/Source/Programs/Horde/Plugins/Build/HordeServer.Build/BuildServerConfig.cs b/Engine/Source/Programs/Horde/Plugins/Build/HordeServer.Build/BuildServerConfig.cs
new file mode 100644
--- /dev/null
+++ b/Engine/Source/Programs/Horde/Plugins/Build/HordeServer.Build/BuildServerConfig.cs
@@ -0,0 +1,269 @@
+﻿// Copyright Epic Games, Inc. All Rights Reserved.
+
+using System.ComponentModel;
+using EpicGames.Core;
+using EpicGames.Horde;
+using EpicGames.Horde.Storage.Bundles;
+using EpicGames.Horde.Storage.Nodes;
+using EpicGames.Perforce;
+using HordeServer.Plugins;
+
+namespace HordeServer
+{
+	/// <summary>
+	/// Static configuration for the build plugin
+	/// </summary>
+	public class BuildServerConfig
+	{
+		/// <summary>
+		/// Perforce connections for use by the Horde server (not agents)
+		/// </summary>
+		public List<PerforceConnectionSettings> Perforce { get; set; } = new List<PerforceConnectionSettings>();
+
+		/// <summary>
+		/// Whether to use the local Perforce environment
+		/// </summary>
+		public bool UseLocalPerforceEnv { get; set; }
+
+		/// <summary>
+		/// Number of pooled perforce connections to keep
+		/// </summary>
+		public int PerforceConnectionPoolSize { get; set; } = 5;
+
+		/// <summary>
+		/// Whether to enable the conform task source.
+		/// </summary>
+		public bool EnableConformTasks { get; set; } = true;
+
+		/// <summary>
+		/// Url of P4 Swarm installation
+		/// </summary>
+		public Uri? P4SwarmUrl { get; set; }
+
+		/// <summary>
+		/// The Jira service account user name
+		/// </summary>
+		public string? JiraUsername { get; set; }
+
+		/// <summary>
+		/// The Jira service account API token
+		/// </summary>
+		public string? JiraApiToken { get; set; }
+
+		/// <summary>
+		/// The Uri for the Jira installation
+		/// </summary>
+		public Uri? JiraUrl { get; set; }
+
+		/// <summary>
+		/// The number of days shared device checkouts are held
+		/// </summary>
+		public int SharedDeviceCheckoutDays { get; set; } = 3;
+
+		/// <summary>
+		/// The number of cooldown minutes for device problems
+		/// </summary>
+		public int DeviceProblemCooldownMinutes { get; set; } = 10;
+
+		/// <summary>
+		/// Channel to send device reports to
+		/// </summary>
+		public string? DeviceReportChannel { get; set; }
+
+		/// <summary>
+		/// Whether to run scheduled jobs.
+		/// </summary>
+		public bool DisableSchedules { get; set; }
+
+		/// <summary>
+		/// Bot token for interacting with Slack (xoxb-*)
+		/// </summary>
+		public string? SlackToken { get; set; }
+
+		/// <summary>
+		/// Token for opening a socket to slack (xapp-*)
+		/// </summary>
+		public string? SlackSocketToken { get; set; }
+
+		/// <summary>
+		/// Admin user token for Slack (xoxp-*). This is only required when using the admin endpoints to invite users.
+		/// </summary>
+		public string? SlackAdminToken { get; set; }
+
+		/// <summary>
+		/// Filtered list of slack users to send notifications to. Should be Slack user ids, separated by commas.
+		/// </summary>
+		public string? SlackUsers { get; set; }
+
+		/// <summary>
+		/// Prefix to use when reporting errors
+		/// </summary>
+		public string SlackErrorPrefix { get; set; } = ":horde-error: ";
+
+		/// <summary>
+		/// Prefix to use when reporting warnings
+		/// </summary>
+		public string SlackWarningPrefix { get; set; } = ":horde-warning: ";
+
+		/// <summary>
+		/// Channel for sending messages related to config update failures
+		/// </summary>
+		public string? ConfigNotificationChannel { get; set; }
+
+		/// <summary>
+		/// Channel to send stream notification update failures to
+		/// </summary>
+		public string? UpdateStreamsNotificationChannel { get; set; }
+
+		/// <summary>
+		/// Slack channel to send job related notifications to. Multiple channels can be specified, separated by ;
+		/// </summary>
+		public string? JobNotificationChannel { get; set; }
+
+		/// <summary>
+		/// Slack channel to send agent related notifications to.
+		/// </summary>
+		public string? AgentNotificationChannel { get; set; }
+
+		/// <summary>
+		/// The number of months to retain test data
+		/// </summary>
+		public int TestDataRetainMonths { get; set; } = 6;
+
+		/// <summary>
+		/// Directory to store the fine-grained block cache. This caches individual exports embedded in bundles.
+		/// </summary>
+		public string? BlockCacheDir { get; set; }
+
+		/// <summary>
+		/// Maximum size of the block cache. Accepts standard binary suffixes. Currently only allocates in multiples of 1024mb.
+		/// </summary>
+		public string BlockCacheSize { get; set; } = "4gb";
+
+		/// <summary>
+		/// Accessor for the block cache size in bytes
+		/// </summary>
+		public long BlockCacheSizeBytes => StringUtils.ParseBytesString(BlockCacheSize);
+
+		/// <summary>
+		/// Options for the commit service
+		/// </summary>
+		public CommitSettings Commits { get; set; } = new CommitSettings();
+	}
+
+	/// <summary>
+	/// Identifier for a pool
+	/// </summary>
+	/// <param name="Id">Id to construct from</param>
+	[JsonSchemaString]
+	[StringIdConverter(typeof(PerforceConnectionIdConverter))]
+	[TypeConverter(typeof(StringIdTypeConverter<PerforceConnectionId, PerforceConnectionIdConverter>))]
+	public record struct PerforceConnectionId(StringId Id)
+	{
+		/// <summary>
+		/// Constructor
+		/// </summary>
+		public PerforceConnectionId(string id) : this(new StringId(id))
+		{
+		}
+
+		/// <inheritdoc/>
+		public override string ToString() => Id.ToString();
+	}
+
+	/// <summary>
+	/// Converter to and from <see cref="StringId"/> instances.
+	/// </summary>
+	class PerforceConnectionIdConverter : StringIdConverter<PerforceConnectionId>
+	{
+		/// <inheritdoc/>
+		public override PerforceConnectionId FromStringId(StringId id) => new PerforceConnectionId(id);
+
+		/// <inheritdoc/>
+		public override StringId ToStringId(PerforceConnectionId value) => value.Id;
+	}
+
+	/// <summary>
+	/// Perforce connection information for use by the Horde server (for reading config files, etc...)
+	/// </summary>
+	public class PerforceConnectionSettings
+	{
+		/// <summary>
+		/// Identifier for the default perforce connection profile
+		/// </summary>
+		public static PerforceConnectionId Default { get; } = new PerforceConnectionId("default");
+
+		/// <summary>
+		/// Identifier for this server
+		/// </summary>
+		public PerforceConnectionId Id { get; set; } = Default;
+
+		/// <summary>
+		/// Server and port
+		/// </summary>
+		public string? ServerAndPort { get; set; }
+
+		/// <summary>
+		/// Credentials for the server
+		/// </summary>
+		public PerforceCredentials? Credentials { get; set; }
+
+		/// <summary>
+		/// Create a <see cref="PerforceSettings"/> object with these settings as overrides
+		/// </summary>
+		/// <returns>New perforce settings object</returns>
+		public PerforceSettings ToPerforceSettings()
+		{
+			PerforceSettings settings = new PerforceSettings(PerforceSettings.Default);
+			settings.PreferNativeClient = true;
+
+			if (!String.IsNullOrEmpty(ServerAndPort))
+			{
+				settings.ServerAndPort = ServerAndPort;
+			}
+			if (Credentials != null)
+			{
+				if (!String.IsNullOrEmpty(Credentials.UserName))
+				{
+					settings.UserName = Credentials.UserName;
+				}
+
+				if (!String.IsNullOrEmpty(Credentials.Ticket))
+				{
+					settings.Password = Credentials.Ticket;
+				}
+				else if (!String.IsNullOrEmpty(Credentials.Password))
+				{
+					settings.Password = Credentials.Password;
+				}
+			}
+			return settings;
+		}
+	}
+
+	/// <summary>
+	/// Options for the commit service
+	/// </summary>
+	public class CommitSettings
+	{
+		/// <summary>
+		/// Whether to mirror commit metadata to the database
+		/// </summary>
+		public bool ReplicateMetadata { get; set; } = true;
+
+		/// <summary>
+		/// Whether to mirror commit data to storage
+		/// </summary>
+		public bool ReplicateContent { get; set; } = true;
+
+		/// <summary>
+		/// Options for how objects are packed together
+		/// </summary>
+		public BundleOptions Bundle { get; set; } = new BundleOptions();
+
+		/// <summary>
+		/// Options for how objects are sliced
+		/// </summary>
+		public ChunkingOptions Chunking { get; set; } = new ChunkingOptions();
+	}
+}
diff --git a/Engine/Source/Programs/Horde/Plugins/Build/HordeServer.Build/Devices/DeviceReportService.cs b/Engine/Source/Programs/Horde/Plugins/Build/HordeServer.Build/Devices/DeviceReportService.cs
--- a/Engine/Source/Programs/Horde/Plugins/Build/HordeServer.Build/Devices/DeviceReportService.cs
+++ b/Engine/Source/Programs/Horde/Plugins/Build/HordeServer.Build/Devices/DeviceReportService.cs
@@ -134,7 +134,7 @@ namespace HordeServer.Devices
 		readonly IClock _clock;
 		readonly ITicker _ticker;
 		readonly IServerInfo _serverInfo;
-		readonly IOptions<StaticBuildConfig> _staticBuildConfig;
+		readonly IOptions<BuildServerConfig> _staticBuildConfig;
 		readonly ILogger<DeviceReportService> _logger;
 
 		readonly int _reportIntervalMinutes = 180;
@@ -142,7 +142,7 @@ namespace HordeServer.Devices
 		/// <summary>
 		/// Constructor
 		/// </summary>
-		public DeviceReportService(IMongoService mongoService, DeviceService deviceService, INotificationService notificationService, IClock clock, IServerInfo serverInfo, IOptions<StaticBuildConfig> staticBuildConfig, ILogger<DeviceReportService> logger)
+		public DeviceReportService(IMongoService mongoService, DeviceService deviceService, INotificationService notificationService, IClock clock, IServerInfo serverInfo, IOptions<BuildServerConfig> staticBuildConfig, ILogger<DeviceReportService> logger)
 		{
 			_state = new SingletonDocument<DeviceReportState>(mongoService);
 			_deviceService = deviceService;
@@ -168,7 +168,7 @@ namespace HordeServer.Devices
 
 		async ValueTask TickAsync(CancellationToken cancellationToken)
 		{
-			StaticBuildConfig staticBuildConfig = _staticBuildConfig.Value;
+			BuildServerConfig staticBuildConfig = _staticBuildConfig.Value;
 			if (String.IsNullOrEmpty(staticBuildConfig.DeviceReportChannel))
 			{
 				return;
diff --git a/Engine/Source/Programs/Horde/Plugins/Build/HordeServer.Build/Devices/DeviceService.cs b/Engine/Source/Programs/Horde/Plugins/Build/HordeServer.Build/Devices/DeviceService.cs
--- a/Engine/Source/Programs/Horde/Plugins/Build/HordeServer.Build/Devices/DeviceService.cs
+++ b/Engine/Source/Programs/Horde/Plugins/Build/HordeServer.Build/Devices/DeviceService.cs
@@ -71,7 +71,7 @@ namespace HordeServer.Devices
 		readonly ITicker _ticker;
 		readonly ITicker _telemetryTicker;
 		readonly IServerInfo _serverInfo;
-		readonly IOptionsMonitor<StaticBuildConfig> _settings;
+		readonly IOptionsMonitor<BuildServerConfig> _settings;
 		readonly IOptionsMonitor<BuildConfig> _buildConfig;
 
 		bool _runUpgrade = true;
@@ -84,7 +84,7 @@ namespace HordeServer.Devices
 		/// <summary>
 		/// Device service constructor
 		/// </summary>
-		public DeviceService(IDeviceCollection devices, IUserCollection userCollection, JobService jobService, IStreamCollection streamCollection, IServerInfo serverInfo, IOptionsMonitor<StaticBuildConfig> settings, IOptionsMonitor<BuildConfig> buildConfig, INotificationService notificationService, IClock clock, Tracer tracer, ILogger<DeviceService> logger)
+		public DeviceService(IDeviceCollection devices, IUserCollection userCollection, JobService jobService, IStreamCollection streamCollection, IServerInfo serverInfo, IOptionsMonitor<BuildServerConfig> settings, IOptionsMonitor<BuildConfig> buildConfig, INotificationService notificationService, IClock clock, Tracer tracer, ILogger<DeviceService> logger)
 		{
 			_userCollection = userCollection;
 			_devices = devices;
diff --git a/Engine/Source/Programs/Horde/Plugins/Build/HordeServer.Build/Issues/External/JiraService.cs b/Engine/Source/Programs/Horde/Plugins/Build/HordeServer.Build/Issues/External/JiraService.cs
--- a/Engine/Source/Programs/Horde/Plugins/Build/HordeServer.Build/Issues/External/JiraService.cs
+++ b/Engine/Source/Programs/Horde/Plugins/Build/HordeServer.Build/Issues/External/JiraService.cs
@@ -96,7 +96,7 @@ namespace HordeServer.Issues.External
 		/// <summary>
 		/// The server settings
 		/// </summary>
-		readonly StaticBuildConfig _settings;
+		readonly BuildServerConfig _settings;
 		readonly HttpClient _client;
 		readonly AsyncPolicy<HttpResponseMessage> _retryPolicy;
 		readonly ConcurrentDictionary<string, JiraCacheValue> _issueCache = new ConcurrentDictionary<string, JiraCacheValue>();
@@ -122,7 +122,7 @@ namespace HordeServer.Issues.External
 		/// <param name="clock"></param>
 		/// <param name="buildConfig"></param>
 		/// <param name="logger"></param>
-		public JiraService(IOptions<StaticBuildConfig> settings, IssueService issueService, IClock clock, IOptionsMonitor<BuildConfig> buildConfig, ILogger<JiraService> logger)
+		public JiraService(IOptions<BuildServerConfig> settings, IssueService issueService, IClock clock, IOptionsMonitor<BuildConfig> buildConfig, ILogger<JiraService> logger)
 		{
 			_settings = settings.Value;
 			_logger = logger;
diff --git a/Engine/Source/Programs/Horde/Plugins/Build/HordeServer.Build/Jobs/JobTaskSource.cs b/Engine/Source/Programs/Horde/Plugins/Build/HordeServer.Build/Jobs/JobTaskSource.cs
--- a/Engine/Source/Programs/Horde/Plugins/Build/HordeServer.Build/Jobs/JobTaskSource.cs
+++ b/Engine/Source/Programs/Horde/Plugins/Build/HordeServer.Build/Jobs/JobTaskSource.cs
@@ -196,7 +196,7 @@ namespace HordeServer.Jobs
 		readonly PerforceLoadBalancer _perforceLoadBalancer;
 		readonly IClock _clock;
 		readonly IServerInfo _serverInfo;
-		readonly IOptionsMonitor<StaticBuildConfig> _settings;
+		readonly IOptionsMonitor<BuildServerConfig> _settings;
 		readonly IOptionsMonitor<BuildConfig> _buildConfig;
 		readonly ILogger<JobTaskSource> _logger;
 		readonly Tracer _tracer;
@@ -233,7 +233,7 @@ namespace HordeServer.Jobs
 		/// <summary>
 		/// Constructor
 		/// </summary>
-		public JobTaskSource(IAclService aclService, Lazy<AgentService> agentService, IJobCollection jobs, IJobStepRefCollection jobStepRefs, IBisectTaskCollection bisectTasks, IGraphCollection graphs, IPoolCollection pools, PoolService poolService, IUgsMetadataCollection ugsMetadataCollection, IStreamCollection streamCollection, ILogCollection logCollection, PerforceLoadBalancer perforceLoadBalancer, IClock clock, IServerInfo serverInfo, IOptionsMonitor<StaticBuildConfig> settings, IOptionsMonitor<BuildConfig> buildConfig, Tracer tracer, ILogger<JobTaskSource> logger)
+		public JobTaskSource(IAclService aclService, Lazy<AgentService> agentService, IJobCollection jobs, IJobStepRefCollection jobStepRefs, IBisectTaskCollection bisectTasks, IGraphCollection graphs, IPoolCollection pools, PoolService poolService, IUgsMetadataCollection ugsMetadataCollection, IStreamCollection streamCollection, ILogCollection logCollection, PerforceLoadBalancer perforceLoadBalancer, IClock clock, IServerInfo serverInfo, IOptionsMonitor<BuildServerConfig> settings, IOptionsMonitor<BuildConfig> buildConfig, Tracer tracer, ILogger<JobTaskSource> logger)
 		{
 			_aclService = aclService;
 			_agentService = agentService;
diff --git a/Engine/Source/Programs/Horde/Plugins/Build/HordeServer.Build/Jobs/TestData/TestDataService.cs b/Engine/Source/Programs/Horde/Plugins/Build/HordeServer.Build/Jobs/TestData/TestDataService.cs
--- a/Engine/Source/Programs/Horde/Plugins/Build/HordeServer.Build/Jobs/TestData/TestDataService.cs
+++ b/Engine/Source/Programs/Horde/Plugins/Build/HordeServer.Build/Jobs/TestData/TestDataService.cs
@@ -16,14 +16,14 @@ namespace HordeServer.Jobs.TestData
 	{
 
 		readonly ITestDataCollection _testData;
-		readonly IOptionsMonitor<StaticBuildConfig> _settings;
+		readonly IOptionsMonitor<BuildServerConfig> _settings;
 		readonly ITicker _ticker;
 		readonly ILogger<TestDataService> _logger;
 
 		/// <summary>
 		/// Device service constructor
 		/// </summary>
-		public TestDataService(ITestDataCollection testData, IOptionsMonitor<StaticBuildConfig> settings, IClock clock, ILogger<TestDataService> logger)
+		public TestDataService(ITestDataCollection testData, IOptionsMonitor<BuildServerConfig> settings, IClock clock, ILogger<TestDataService> logger)
 		{
 			_testData = testData;
 			_settings = settings;
diff --git a/Engine/Source/Programs/Horde/Plugins/Build/HordeServer.Build/Notifications/NotificationService.cs b/Engine/Source/Programs/Horde/Plugins/Build/HordeServer.Build/Notifications/NotificationService.cs
--- a/Engine/Source/Programs/Horde/Plugins/Build/HordeServer.Build/Notifications/NotificationService.cs
+++ b/Engine/Source/Programs/Horde/Plugins/Build/HordeServer.Build/Notifications/NotificationService.cs
@@ -93,7 +93,7 @@ namespace HordeServer.Notifications
 		/// <summary>
 		/// Settings for the application.
 		/// </summary>
-		private readonly IOptionsMonitor<StaticBuildConfig> _settings;
+		private readonly IOptionsMonitor<BuildServerConfig> _settings;
 
 		/// <summary>
 		/// List of asychronous tasks currently executing
@@ -132,7 +132,7 @@ namespace HordeServer.Notifications
 		/// </summary>
 		public NotificationService(
 			IEnumerable<INotificationSink> sinks,
-			IOptionsMonitor<StaticBuildConfig> settings,
+			IOptionsMonitor<BuildServerConfig> settings,
 			ILogger<NotificationService> logger,
 			IGraphCollection graphCollection,
 			ISubscriptionCollection subscriptionCollection,
diff --git a/Engine/Source/Programs/Horde/Plugins/Build/HordeServer.Build/Notifications/Sinks/SlackNotificationSink.cs b/Engine/Source/Programs/Horde/Plugins/Build/HordeServer.Build/Notifications/Sinks/SlackNotificationSink.cs
--- a/Engine/Source/Programs/Horde/Plugins/Build/HordeServer.Build/Notifications/Sinks/SlackNotificationSink.cs
+++ b/Engine/Source/Programs/Horde/Plugins/Build/HordeServer.Build/Notifications/Sinks/SlackNotificationSink.cs
@@ -213,7 +213,7 @@ namespace HordeServer.Notifications.Sinks
 		readonly IUserCollection _userCollection;
 		readonly ILogCollection _logCollection;
 		readonly IWebHostEnvironment _environment;
-		readonly StaticBuildConfig _settings;
+		readonly BuildServerConfig _settings;
 		readonly IMongoCollection<MessageStateDocument> _messageStates;
 		readonly IMongoCollection<SlackUserDocument> _slackUsers;
 		readonly HashSet<string>? _allowUsers;
@@ -244,7 +244,7 @@ namespace HordeServer.Notifications.Sinks
 		/// <summary>
 		/// Constructor
 		/// </summary>
-		public SlackNotificationSink(IMongoService mongoService, IRedisService redisService, IssueService issueService, ICommitService commitService, IUserCollection userCollection, ILogCollection logCollection, IExternalIssueService externalIssueService, IWebHostEnvironment environment, IOptions<StaticBuildConfig> settings, IClock clock, IServerInfo serverInfo, IOptionsMonitor<BuildConfig> buildConfig, ILogger<SlackNotificationSink> logger)
+		public SlackNotificationSink(IMongoService mongoService, IRedisService redisService, IssueService issueService, ICommitService commitService, IUserCollection userCollection, ILogCollection logCollection, IExternalIssueService externalIssueService, IWebHostEnvironment environment, IOptions<BuildServerConfig> settings, IClock clock, IServerInfo serverInfo, IOptionsMonitor<BuildConfig> buildConfig, ILogger<SlackNotificationSink> logger)
 		{
 			_redisService = redisService;
 			_issueService = issueService;
diff --git a/Engine/Source/Programs/Horde/Plugins/Build/HordeServer.Build/Perforce/PerforceConfigSource.cs b/Engine/Source/Programs/Horde/Plugins/Build/HordeServer.Build/Perforce/PerforceConfigSource.cs
--- a/Engine/Source/Programs/Horde/Plugins/Build/HordeServer.Build/Perforce/PerforceConfigSource.cs
+++ b/Engine/Source/Programs/Horde/Plugins/Build/HordeServer.Build/Perforce/PerforceConfigSource.cs
@@ -48,7 +48,7 @@ namespace HordeServer.Perforce
 		/// <inheritdoc/>
 		public TimeSpan UpdateInterval => TimeSpan.FromMinutes(1.0);
 
-		readonly IOptionsMonitor<StaticBuildConfig> _settings;
+		readonly IOptionsMonitor<BuildServerConfig> _settings;
 		readonly IUserCollection _userCollection;
 		readonly IMemoryCache _cache;
 		readonly ILogger _logger;
@@ -56,7 +56,7 @@ namespace HordeServer.Perforce
 		/// <summary>
 		/// Constructor
 		/// </summary>
-		public PerforceConfigSource(IOptionsMonitor<StaticBuildConfig> settings, IUserCollection userCollection, IMemoryCache cache, ILogger<PerforceConfigSource> logger)
+		public PerforceConfigSource(IOptionsMonitor<BuildServerConfig> settings, IUserCollection userCollection, IMemoryCache cache, ILogger<PerforceConfigSource> logger)
 		{
 			_settings = settings;
 			_userCollection = userCollection;
@@ -147,7 +147,7 @@ namespace HordeServer.Perforce
 		{
 			_ = cancellationToken;
 
-			StaticBuildConfig settings = _settings.CurrentValue;
+			BuildServerConfig settings = _settings.CurrentValue;
 
 			PerforceConnectionId connectionId = new PerforceConnectionId();
 			if (!String.IsNullOrEmpty(host))
diff --git a/Engine/Source/Programs/Horde/Plugins/Build/HordeServer.Build/Perforce/PerforceService.cs b/Engine/Source/Programs/Horde/Plugins/Build/HordeServer.Build/Perforce/PerforceService.cs
--- a/Engine/Source/Programs/Horde/Plugins/Build/HordeServer.Build/Perforce/PerforceService.cs
+++ b/Engine/Source/Programs/Horde/Plugins/Build/HordeServer.Build/Perforce/PerforceService.cs
@@ -248,7 +248,7 @@ namespace HordeServer.Perforce
 		readonly string? _perforceServerOverride;
 		readonly string? _perforceUserOverride;
 
-		readonly StaticBuildConfig _settings;
+		readonly BuildServerConfig _settings;
 		readonly Dictionary<string, Dictionary<string, Credentials>> _userCredentialsByCluster = new Dictionary<string, Dictionary<string, Credentials>>(StringComparer.OrdinalIgnoreCase);
 		readonly IUserCollection _userCollection;
 		readonly MemoryCache _userCache = new MemoryCache(new MemoryCacheOptions { SizeLimit = 2000 });
@@ -259,7 +259,7 @@ namespace HordeServer.Perforce
 		/// <summary>
 		/// Constructor
 		/// </summary>
-		public PerforceService(PerforceLoadBalancer loadBalancer, IUserCollection userCollection, IOptions<StaticBuildConfig> settings, IOptionsMonitor<BuildConfig> buildConfig, Tracer tracer, ILogger<PerforceService> logger)
+		public PerforceService(PerforceLoadBalancer loadBalancer, IUserCollection userCollection, IOptions<BuildServerConfig> settings, IOptionsMonitor<BuildConfig> buildConfig, Tracer tracer, ILogger<PerforceService> logger)
 		{
 			_loadBalancer = loadBalancer;
 			_userCollection = userCollection;
diff --git a/Engine/Source/Programs/Horde/Plugins/Build/HordeServer.Build/Perforce/PerforceServiceCache.cs b/Engine/Source/Programs/Horde/Plugins/Build/HordeServer.Build/Perforce/PerforceServiceCache.cs
--- a/Engine/Source/Programs/Horde/Plugins/Build/HordeServer.Build/Perforce/PerforceServiceCache.cs
+++ b/Engine/Source/Programs/Horde/Plugins/Build/HordeServer.Build/Perforce/PerforceServiceCache.cs
@@ -236,7 +236,7 @@ namespace HordeServer.Perforce
 		/// <summary>
 		/// Constructor
 		/// </summary>
-		public PerforceServiceCache(PerforceLoadBalancer loadBalancer, IMongoService mongoService, IRedisService redisService, IDowntimeService downtimeService, IUserCollection userCollection, IClock clock, IOptions<StaticBuildConfig> settings, IOptionsMonitor<BuildConfig> buildConfig, Tracer tracer, ILogger<PerforceService> logger)
+		public PerforceServiceCache(PerforceLoadBalancer loadBalancer, IMongoService mongoService, IRedisService redisService, IDowntimeService downtimeService, IUserCollection userCollection, IClock clock, IOptions<BuildServerConfig> settings, IOptionsMonitor<BuildConfig> buildConfig, Tracer tracer, ILogger<PerforceService> logger)
 			: base(loadBalancer, userCollection, settings, buildConfig, tracer, logger)
 		{
 			_mongoService = mongoService;
diff --git a/Engine/Source/Programs/Horde/Plugins/Build/HordeServer.Build/StaticBuildConfig.cs b/Engine/Source/Programs/Horde/Plugins/Build/HordeServer.Build/StaticBuildConfig.cs
deleted file mode 100644
diff --git a/Engine/Source/Programs/Horde/Plugins/Build/HordeServer.Build/Tasks/ConformTaskSource.cs b/Engine/Source/Programs/Horde/Plugins/Build/HordeServer.Build/Tasks/ConformTaskSource.cs
--- a/Engine/Source/Programs/Horde/Plugins/Build/HordeServer.Build/Tasks/ConformTaskSource.cs
+++ b/Engine/Source/Programs/Horde/Plugins/Build/HordeServer.Build/Tasks/ConformTaskSource.cs
@@ -37,14 +37,14 @@ namespace HordeServer.Tasks
 		readonly PerforceLoadBalancer _perforceLoadBalancer;
 		readonly ILogCollection _logCollection;
 		readonly ILogger _logger;
-		readonly IOptionsMonitor<StaticBuildConfig> _settings;
+		readonly IOptionsMonitor<BuildServerConfig> _settings;
 		readonly IOptionsMonitor<BuildConfig> _buildConfig;
 		readonly ITicker _tickConformList;
 
 		/// <summary>
 		/// Constructor
 		/// </summary>
-		public ConformTaskSource(IMongoService mongoService, IAgentCollection agentCollection, PoolService poolService, ILogCollection logCollection, PerforceLoadBalancer perforceLoadBalancer, IClock clock, IOptionsMonitor<StaticBuildConfig> settings, IOptionsMonitor<BuildConfig> buildConfig, ILogger<ConformTaskSource> logger)
+		public ConformTaskSource(IMongoService mongoService, IAgentCollection agentCollection, PoolService poolService, ILogCollection logCollection, PerforceLoadBalancer perforceLoadBalancer, IClock clock, IOptionsMonitor<BuildServerConfig> settings, IOptionsMonitor<BuildConfig> buildConfig, ILogger<ConformTaskSource> logger)
 		{
 			_agentCollection = agentCollection;
 			_poolService = poolService;
