Horde: Rename StaticComputeConfig to ComputeServerConfig for consistency.

#rnx

diff --git a/Engine/Source/Programs/Horde/HordeServer.Tests.Shared/ServerTestSetup.cs b/Engine/Source/Programs/Horde/HordeServer.Tests.Shared/ServerTestSetup.cs
--- a/Engine/Source/Programs/Horde/HordeServer.Tests.Shared/ServerTestSetup.cs
+++ b/Engine/Source/Programs/Horde/HordeServer.Tests.Shared/ServerTestSetup.cs
@@ -109,7 +109,7 @@ namespace HordeServer.Tests
 			ServerSettings settings = new ServerSettings();
 			ConfigureSettings(settings);
 
-			services.Configure<StaticComputeConfig>(x => x.WithAws = true);
+			services.Configure<ComputeServerConfig>(x => x.WithAws = true);
 
 			ServerInfo serverInfo = new ServerInfo(configuration, Options.Create(settings));
 			services.AddSingleton<IServerInfo>(serverInfo);
diff --git a/Engine/Source/Programs/Horde/Plugins/Build/HordeServer.Build.Tests/Compute/ComputeServiceTest.cs b/Engine/Source/Programs/Horde/Plugins/Build/HordeServer.Build.Tests/Compute/ComputeServiceTest.cs
--- a/Engine/Source/Programs/Horde/Plugins/Build/HordeServer.Build.Tests/Compute/ComputeServiceTest.cs
+++ b/Engine/Source/Programs/Horde/Plugins/Build/HordeServer.Build.Tests/Compute/ComputeServiceTest.cs
@@ -410,9 +410,9 @@ namespace HordeServer.Tests.Compute
 		private async Task<ComputeService> CreateComputeServiceAsync(string? tunnelAddress, string pool, ComputeClusterConfig ccc)
 		{
 			IOptionsMonitor<ComputeConfig> computeConfig = ServiceProvider.GetRequiredService<IOptionsMonitor<ComputeConfig>>();
-			StaticComputeConfig ss = new() { ComputeTunnelAddress = tunnelAddress };
+			ComputeServerConfig ss = new() { ComputeTunnelAddress = tunnelAddress };
 			ComputeService cs = new(AgentCollection, ServiceProvider.GetRequiredService<IAgentScheduler>(), LogCollection, AgentService, AgentRelayService, GetRedisServiceSingleton(),
-				new TestOptionsMonitor<StaticComputeConfig>(ss), computeConfig, Clock, Tracer, Meter,
+				new TestOptionsMonitor<ComputeServerConfig>(ss), computeConfig, Clock, Tracer, Meter,
 				NullLogger<ComputeService>.Instance);
 			await CreateAgentAsync(new PoolId(pool), properties: ["ComputeIp=11.0.0.1", "ComputePort=5000"]);
 			computeConfig.CurrentValue.Clusters = [ccc];
diff --git a/Engine/Source/Programs/Horde/Plugins/Build/HordeServer.Build.Tests/Fleet/AwsAutoScalingLifecycleTest.cs b/Engine/Source/Programs/Horde/Plugins/Build/HordeServer.Build.Tests/Fleet/AwsAutoScalingLifecycleTest.cs
--- a/Engine/Source/Programs/Horde/Plugins/Build/HordeServer.Build.Tests/Fleet/AwsAutoScalingLifecycleTest.cs
+++ b/Engine/Source/Programs/Horde/Plugins/Build/HordeServer.Build.Tests/Fleet/AwsAutoScalingLifecycleTest.cs
@@ -58,8 +58,8 @@ public class AwsAutoScalingLifecycleServiceTest : BuildTestSetup
 			.Setup(x => x.CompleteLifecycleActionAsync(It.IsAny<CompleteLifecycleActionRequest>(), It.IsAny<CancellationToken>()))
 			.Returns(OnCompleteLifecycleActionAsync);
 
-		StaticComputeConfig ss = new() { AwsAutoScalingQueueUrls = new[] { _queueUrl1, _queueUrl2 } };
-		_asgLifecycleService = new AwsAutoScalingLifecycleService(AgentService, GetRedisServiceSingleton(), AgentCollection, Clock, new TestOptionsMonitor<StaticComputeConfig>(ss), ServiceProvider, Tracer, logger);
+		ComputeServerConfig ss = new() { AwsAutoScalingQueueUrls = new[] { _queueUrl1, _queueUrl2 } };
+		_asgLifecycleService = new AwsAutoScalingLifecycleService(AgentService, GetRedisServiceSingleton(), AgentCollection, Clock, new TestOptionsMonitor<ComputeServerConfig>(ss), ServiceProvider, Tracer, logger);
 		_asgLifecycleService.SetAmazonClientsTesting(_asgMock.Object, _fakeSqs);
 		await _asgLifecycleService.StartAsync(CancellationToken.None);
 	}
diff --git a/Engine/Source/Programs/Horde/Plugins/Build/HordeServer.Build.Tests/Fleet/FleetServiceTest.cs b/Engine/Source/Programs/Horde/Plugins/Build/HordeServer.Build.Tests/Fleet/FleetServiceTest.cs
--- a/Engine/Source/Programs/Horde/Plugins/Build/HordeServer.Build.Tests/Fleet/FleetServiceTest.cs
+++ b/Engine/Source/Programs/Horde/Plugins/Build/HordeServer.Build.Tests/Fleet/FleetServiceTest.cs
@@ -226,7 +226,7 @@ namespace HordeServer.Tests.Fleet
 		private FleetService GetFleetService(IFleetManager fleetManager, bool isDowntimeActive = false)
 		{
 			ILoggerFactory loggerFactory = ServiceProvider.GetRequiredService<ILoggerFactory>();
-			IOptions<StaticComputeConfig> serverSettingsOpt = ServiceProvider.GetRequiredService<IOptions<StaticComputeConfig>>();
+			IOptions<ComputeServerConfig> serverSettingsOpt = ServiceProvider.GetRequiredService<IOptions<ComputeServerConfig>>();
 			serverSettingsOpt.Value.FleetManagerV2 = FleetManagerType.AwsReuse;
 
 			IOptionsMonitor<ComputeConfig> computeConfig = ServiceProvider.GetRequiredService<IOptionsMonitor<ComputeConfig>>();
diff --git a/Engine/Source/Programs/Horde/Plugins/Compute/HordeServer.Compute.Tests/Compute/TunnelServiceTests.cs b/Engine/Source/Programs/Horde/Plugins/Compute/HordeServer.Compute.Tests/Compute/TunnelServiceTests.cs
--- a/Engine/Source/Programs/Horde/Plugins/Compute/HordeServer.Compute.Tests/Compute/TunnelServiceTests.cs
+++ b/Engine/Source/Programs/Horde/Plugins/Compute/HordeServer.Compute.Tests/Compute/TunnelServiceTests.cs
@@ -60,8 +60,8 @@ namespace HordeServer.Tests.Compute
 
 			_echoServerTask = StartEchoServerAsync(_echoPort, echoServerBufferSize, cts.Token);
 
-			StaticComputeConfig serverSettings = new() { ComputeTunnelPort = _tunnelPort };
-			using TunnelService tunnelService = new(new OptionsWrapper<StaticComputeConfig>(serverSettings), ComputeTestSetup.CreateConsoleLogger<TunnelService>());
+			ComputeServerConfig serverSettings = new() { ComputeTunnelPort = _tunnelPort };
+			using TunnelService tunnelService = new(new OptionsWrapper<ComputeServerConfig>(serverSettings), ComputeTestSetup.CreateConsoleLogger<TunnelService>());
 			tunnelService.Start(IPAddress.Any);
 			_tunnelServerTask = tunnelService.ServerTask!;
 
diff --git a/Engine/Source/Programs/Horde/Plugins/Compute/HordeServer.Compute/Agents/Enrollment/EnrollmentService.cs b/Engine/Source/Programs/Horde/Plugins/Compute/HordeServer.Compute/Agents/Enrollment/EnrollmentService.cs
--- a/Engine/Source/Programs/Horde/Plugins/Compute/HordeServer.Compute/Agents/Enrollment/EnrollmentService.cs
+++ b/Engine/Source/Programs/Horde/Plugins/Compute/HordeServer.Compute/Agents/Enrollment/EnrollmentService.cs
@@ -28,14 +28,14 @@ namespace HordeServer.Agents.Enrollment
 		readonly RedisHashKey<string, EnrollmentRequest> _requests = new("agents:registration:requests");
 		readonly RedisHashKey<string, string> _approvals = new("agents:registration:approvals");
 		readonly AsyncEvent _approvalEvent = new AsyncEvent();
-		readonly IOptionsMonitor<StaticComputeConfig> _staticComputeConfig;
+		readonly IOptionsMonitor<ComputeServerConfig> _staticComputeConfig;
 
 		RedisSubscription? _subscription;
 
 		/// <summary>
 		/// Constructor
 		/// </summary>
-		public EnrollmentService(IRedisService redisService, IClock clock, IOptionsMonitor<StaticComputeConfig> staticComputeConfig)
+		public EnrollmentService(IRedisService redisService, IClock clock, IOptionsMonitor<ComputeServerConfig> staticComputeConfig)
 		{
 			_redisService = redisService;
 			_clock = clock;
diff --git a/Engine/Source/Programs/Horde/Plugins/Compute/HordeServer.Compute/Agents/Fleet/AwsAutoScalingLifecycle.cs b/Engine/Source/Programs/Horde/Plugins/Compute/HordeServer.Compute/Agents/Fleet/AwsAutoScalingLifecycle.cs
--- a/Engine/Source/Programs/Horde/Plugins/Compute/HordeServer.Compute/Agents/Fleet/AwsAutoScalingLifecycle.cs
+++ b/Engine/Source/Programs/Horde/Plugins/Compute/HordeServer.Compute/Agents/Fleet/AwsAutoScalingLifecycle.cs
@@ -82,7 +82,7 @@ public sealed class AwsAutoScalingLifecycleService : IHostedService, IAsyncDispo
 		IRedisService redisService,
 		IAgentCollection agents,
 		IClock clock,
-		IOptionsMonitor<StaticComputeConfig> staticComputeConfig,
+		IOptionsMonitor<ComputeServerConfig> staticComputeConfig,
 		IServiceProvider serviceProvider,
 		Tracer tracer,
 		ILogger<AwsAutoScalingLifecycleService> logger)
diff --git a/Engine/Source/Programs/Horde/Plugins/Compute/HordeServer.Compute/Agents/Fleet/FleetManagerFactory.cs b/Engine/Source/Programs/Horde/Plugins/Compute/HordeServer.Compute/Agents/Fleet/FleetManagerFactory.cs
--- a/Engine/Source/Programs/Horde/Plugins/Compute/HordeServer.Compute/Agents/Fleet/FleetManagerFactory.cs
+++ b/Engine/Source/Programs/Horde/Plugins/Compute/HordeServer.Compute/Agents/Fleet/FleetManagerFactory.cs
@@ -40,7 +40,7 @@ public sealed class FleetManagerFactory : IFleetManagerFactory
 	private readonly IClock _clock;
 	private readonly Meter _meter;
 	private readonly IServiceProvider _provider;
-	private readonly IOptionsMonitor<StaticComputeConfig> _staticComputeConfig;
+	private readonly IOptionsMonitor<ComputeServerConfig> _staticComputeConfig;
 	private readonly Tracer _tracer;
 	private readonly ILoggerFactory _loggerFactory;
 
@@ -50,7 +50,7 @@ public sealed class FleetManagerFactory : IFleetManagerFactory
 	/// <summary>
 	/// Constructor
 	/// </summary>
-	public FleetManagerFactory(IAgentCollection agentCollection, IClock clock, Meter meter, IServiceProvider provider, IOptionsMonitor<StaticComputeConfig> staticComputeConfig, Tracer tracer, ILoggerFactory loggerFactory)
+	public FleetManagerFactory(IAgentCollection agentCollection, IClock clock, Meter meter, IServiceProvider provider, IOptionsMonitor<ComputeServerConfig> staticComputeConfig, Tracer tracer, ILoggerFactory loggerFactory)
 	{
 		_agentCollection = agentCollection;
 		_clock = clock;
@@ -114,7 +114,7 @@ public sealed class FleetManagerFactory : IFleetManagerFactory
 		_awsEc2 = _provider.GetService<IAmazonEC2>();
 		if (_staticComputeConfig.CurrentValue.WithAws == false || _awsEc2 == null)
 		{
-			throw new ArgumentException($"Unable to create fleet manager {type} requiring AWS specific classes. Check that setting '{nameof(StaticComputeConfig.WithAws)}' is enabled");
+			throw new ArgumentException($"Unable to create fleet manager {type} requiring AWS specific classes. Check that setting '{nameof(ComputeServerConfig.WithAws)}' is enabled");
 		}
 
 		return _awsEc2;
@@ -130,7 +130,7 @@ public sealed class FleetManagerFactory : IFleetManagerFactory
 		_awsAutoScaling = _provider.GetService<IAmazonAutoScaling>();
 		if (_staticComputeConfig.CurrentValue.WithAws == false || _awsAutoScaling == null)
 		{
-			throw new ArgumentException($"Unable to create fleet manager {type} requiring AWS specific classes. Check that setting '{nameof(StaticComputeConfig.WithAws)}' is enabled");
+			throw new ArgumentException($"Unable to create fleet manager {type} requiring AWS specific classes. Check that setting '{nameof(ComputeServerConfig.WithAws)}' is enabled");
 		}
 
 		return _awsAutoScaling;
diff --git a/Engine/Source/Programs/Horde/Plugins/Compute/HordeServer.Compute/Agents/Fleet/FleetService.cs b/Engine/Source/Programs/Horde/Plugins/Compute/HordeServer.Compute/Agents/Fleet/FleetService.cs
--- a/Engine/Source/Programs/Horde/Plugins/Compute/HordeServer.Compute/Agents/Fleet/FleetService.cs
+++ b/Engine/Source/Programs/Horde/Plugins/Compute/HordeServer.Compute/Agents/Fleet/FleetService.cs
@@ -65,7 +65,7 @@ namespace HordeServer.Agents.Fleet
 		private readonly ITicker _tickerHighFrequency;
 		private readonly TimeSpan _defaultScaleOutCooldown;
 		private readonly TimeSpan _defaultScaleInCooldown;
-		private readonly IOptions<StaticComputeConfig> _staticComputeConfig;
+		private readonly IOptions<ComputeServerConfig> _staticComputeConfig;
 		private readonly IOptionsMonitor<ComputeConfig> _computeConfig;
 		private readonly Tracer _tracer;
 		private readonly ILogger<FleetService> _logger;
@@ -81,7 +81,7 @@ namespace HordeServer.Agents.Fleet
 			IFleetManagerFactory fleetManagerFactory,
 			IClock clock,
 			IEnumerable<IPoolSizeStrategyFactory> poolSizeStrategyFactories,
-			IOptions<StaticComputeConfig> staticComputeConfig,
+			IOptions<ComputeServerConfig> staticComputeConfig,
 			IOptionsMonitor<ComputeConfig> computeConfig,
 			Tracer tracer,
 			ILogger<FleetService> logger)
diff --git a/Engine/Source/Programs/Horde/Plugins/Compute/HordeServer.Compute/Compute/ComputeService.cs b/Engine/Source/Programs/Horde/Plugins/Compute/HordeServer.Compute/Compute/ComputeService.cs
--- a/Engine/Source/Programs/Horde/Plugins/Compute/HordeServer.Compute/Compute/ComputeService.cs
+++ b/Engine/Source/Programs/Horde/Plugins/Compute/HordeServer.Compute/Compute/ComputeService.cs
@@ -181,7 +181,7 @@ namespace HordeServer.Compute
 		readonly AgentService _agentService;
 		readonly AgentRelayService _agentRelayService;
 		readonly IRedisService _redisService;
-		readonly IOptionsMonitor<StaticComputeConfig> _staticComputeConfig;
+		readonly IOptionsMonitor<ComputeServerConfig> _staticComputeConfig;
 		readonly IOptionsMonitor<ComputeConfig> _computeConfig;
 		readonly IClock _clock;
 		readonly Tracer _tracer;
@@ -204,7 +204,7 @@ namespace HordeServer.Compute
 			AgentService agentService,
 			AgentRelayService agentRelayService,
 			IRedisService redisService,
-			IOptionsMonitor<StaticComputeConfig> staticComputeConfig,
+			IOptionsMonitor<ComputeServerConfig> staticComputeConfig,
 			IOptionsMonitor<ComputeConfig> computeConfig,
 			IClock clock,
 			Tracer tracer,
diff --git a/Engine/Source/Programs/Horde/Plugins/Compute/HordeServer.Compute/Compute/TunnelService.cs b/Engine/Source/Programs/Horde/Plugins/Compute/HordeServer.Compute/Compute/TunnelService.cs
--- a/Engine/Source/Programs/Horde/Plugins/Compute/HordeServer.Compute/Compute/TunnelService.cs
+++ b/Engine/Source/Programs/Horde/Plugins/Compute/HordeServer.Compute/Compute/TunnelService.cs
@@ -18,7 +18,7 @@ namespace HordeServer.Compute
 		private const int BufferSize = 4096;
 
 		private readonly CancellationTokenSource _cancellationSource;
-		private readonly StaticComputeConfig _staticComputeConfig;
+		private readonly ComputeServerConfig _staticComputeConfig;
 		private readonly ILogger _logger;
 		private Task? _serverTask;
 
@@ -30,7 +30,7 @@ namespace HordeServer.Compute
 		/// <summary>
 		/// Constructor
 		/// </summary>
-		public TunnelService(IOptions<StaticComputeConfig> staticComputeConfig, ILogger<TunnelService> logger)
+		public TunnelService(IOptions<ComputeServerConfig> staticComputeConfig, ILogger<TunnelService> logger)
 		{
 			_cancellationSource = new CancellationTokenSource();
 			_staticComputeConfig = staticComputeConfig.Value;
diff --git a/Engine/Source/Programs/Horde/Plugins/Compute/HordeServer.Compute/ComputePlugin.cs b/Engine/Source/Programs/Horde/Plugins/Compute/HordeServer.Compute/ComputePlugin.cs
--- a/Engine/Source/Programs/Horde/Plugins/Compute/HordeServer.Compute/ComputePlugin.cs
+++ b/Engine/Source/Programs/Horde/Plugins/Compute/HordeServer.Compute/ComputePlugin.cs
@@ -31,16 +31,16 @@ namespace HordeServer
 	/// <summary>
 	/// Entry point for the storage plugin
 	/// </summary>
-	[Plugin("Compute", GlobalConfigType = typeof(ComputeConfig), ServerConfigType = typeof(StaticComputeConfig))]
+	[Plugin("Compute", GlobalConfigType = typeof(ComputeConfig), ServerConfigType = typeof(ComputeServerConfig))]
 	public class ComputePlugin : IPluginStartup
 	{
 		readonly IServerInfo _serverInfo;
-		readonly StaticComputeConfig _staticComputeConfig;
+		readonly ComputeServerConfig _staticComputeConfig;
 
 		/// <summary>
 		/// Constructor
 		/// </summary>
-		public ComputePlugin(IServerInfo serverInfo, StaticComputeConfig staticComputeConfig)
+		public ComputePlugin(IServerInfo serverInfo, ComputeServerConfig staticComputeConfig)
 		{
 			_serverInfo = serverInfo;
 			_staticComputeConfig = staticComputeConfig;
diff --git a/Engine/Source/Programs/Horde/Plugins/Compute/HordeServer.Compute/ComputeServerConfig.cs b/Engine/Source/Programs/Horde/Plugins/Compute/HordeServer.Compute/ComputeServerConfig.cs
new file mode 100644
--- /dev/null
+++ b/Engine/Source/Programs/Horde/Plugins/Compute/HordeServer.Compute/ComputeServerConfig.cs
@@ -0,0 +1,81 @@
+﻿// Copyright Epic Games, Inc. All Rights Reserved.
+
+using System.Text.Json.Nodes;
+using System.Text.Json.Serialization;
+using HordeServer.Agents.Fleet;
+using HordeServer.Agents.Pools;
+using HordeServer.Plugins;
+
+namespace HordeServer
+{
+	/// <summary>
+	/// Static configuration for the compute plugin
+	/// </summary>
+	public class ComputeServerConfig
+	{
+		/// <summary>
+		/// Whether to enable the upgrade task source.
+		/// </summary>
+		public bool EnableUpgradeTasks { get; set; } = true;
+
+		/// <summary>
+		/// Whether to enable Amazon Web Services (AWS) specific features
+		/// </summary>
+		public bool WithAws { get; set; } = false;
+		
+		/// <summary>
+		/// List of AWS regions for Horde to be aware of (e.g. us-east-1 or eu-central-1)
+		/// Right now, this is only used for replicating CloudWatch metrics to multiple regions
+		/// <see cref="Amazon.RegionEndpoint" />
+		/// </summary>
+		public string[] AwsRegions { get; set; } = [];
+		
+		/// <summary>
+		/// AWS SQS queue URLs where lifecycle events from EC2 auto-scaling are received
+		/// <see cref="AwsAutoScalingLifecycleService" />
+		/// </summary>
+		public string[] AwsAutoScalingQueueUrls { get; set; } = Array.Empty<string>();
+
+		/// <summary>
+		/// Default fleet manager to use (when not specified by pool)
+		/// </summary>
+		public FleetManagerType FleetManagerV2 { get; set; } = FleetManagerType.NoOp;
+
+		/// <summary>
+		/// Config for the fleet manager (serialized JSON)
+		/// </summary>
+		[JsonConverter(typeof(JsonObjectOrStringConverter))]
+		public JsonObject? FleetManagerV2Config { get; set; }
+
+		/// <summary>
+		/// Whether to automatically enroll agents in the farm
+		/// </summary>
+		public bool AutoEnrollAgents { get; set; }
+
+		/// <summary>
+		/// Default agent pool sizing strategy for pools that doesn't have one explicitly configured
+		/// </summary>
+		public PoolSizeStrategy DefaultAgentPoolSizeStrategy { get; set; } = PoolSizeStrategy.LeaseUtilization;
+
+		/// <summary>
+		/// Scale-out cooldown for auto-scaling agent pools (in seconds). Can be overridden by per-pool settings.
+		/// </summary>
+		public int AgentPoolScaleOutCooldownSeconds { get; set; } = 60; // 1 min
+
+		/// <summary>
+		/// Scale-in cooldown for auto-scaling agent pools (in seconds). Can be overridden by per-pool settings.
+		/// </summary>
+		public int AgentPoolScaleInCooldownSeconds { get; set; } = 1200; // 20 mins
+
+		/// <summary>
+		/// Port to listen on for tunneling compute sockets to agents
+		/// </summary>
+		public int ComputeTunnelPort { get; set; }
+
+		/// <summary>
+		/// What address (host:port) clients should connect to for compute socket tunneling
+		/// Port may differ from <see cref="ComputeTunnelPort" /> if Horde server is behind a reverse proxy/firewall
+		/// </summary>
+		public string? ComputeTunnelAddress { get; set; }
+	}
+}
diff --git a/Engine/Source/Programs/Horde/Plugins/Compute/HordeServer.Compute/StaticComputeConfig.cs b/Engine/Source/Programs/Horde/Plugins/Compute/HordeServer.Compute/StaticComputeConfig.cs
deleted file mode 100644
diff --git a/Engine/Source/Programs/Horde/Plugins/Compute/HordeServer.Compute/Tasks/UpgradeTaskSource.cs b/Engine/Source/Programs/Horde/Plugins/Compute/HordeServer.Compute/Tasks/UpgradeTaskSource.cs
--- a/Engine/Source/Programs/Horde/Plugins/Compute/HordeServer.Compute/Tasks/UpgradeTaskSource.cs
+++ b/Engine/Source/Programs/Horde/Plugins/Compute/HordeServer.Compute/Tasks/UpgradeTaskSource.cs
@@ -23,10 +23,10 @@ namespace HordeServer.Tasks
 		readonly IToolCollection _toolCollection;
 		readonly ILogCollection _logCollection;
 		readonly IOptionsMonitor<ComputeConfig> _computeConfig;
-		readonly IOptions<StaticComputeConfig> _staticComputeConfig;
+		readonly IOptions<ComputeServerConfig> _staticComputeConfig;
 		readonly IClock _clock;
 
-		public UpgradeTaskSource(IToolCollection toolCollection, ILogCollection logCollection, IOptionsMonitor<ComputeConfig> computeConfig, IOptions<StaticComputeConfig> staticComputeConfig, IClock clock)
+		public UpgradeTaskSource(IToolCollection toolCollection, ILogCollection logCollection, IOptionsMonitor<ComputeConfig> computeConfig, IOptions<ComputeServerConfig> staticComputeConfig, IClock clock)
 		{
 			_toolCollection = toolCollection;
 			_logCollection = logCollection;
