Horde: Discover plugins dynamically at startup. Plugins can be enabled/disabled via the "Enabled" property in the the server config file rather than being hard-coded.

#rnx

diff --git a/Engine/Source/Programs/Horde/HordeServer.Shared/Plugins/PluginServerConfig.cs b/Engine/Source/Programs/Horde/HordeServer.Shared/Plugins/PluginServerConfig.cs
new file mode 100644
--- /dev/null
+++ b/Engine/Source/Programs/Horde/HordeServer.Shared/Plugins/PluginServerConfig.cs
@@ -0,0 +1,21 @@
+// Copyright Epic Games, Inc. All Rights Reserved.
+
+using System;
+using System.Collections.Generic;
+using System.Linq;
+using System.Text;
+using System.Threading.Tasks;
+
+namespace HordeServer.Plugins
+{
+	/// <summary>
+	/// Base class for plugin server config objects
+	/// </summary>
+	public class PluginServerConfig
+	{
+		/// <summary>
+		/// Whether the plugin should be enabled or not
+		/// </summary>
+		public bool Enabled { get; set; }
+	}
+}
diff --git a/Engine/Source/Programs/Horde/HordeServer/Plugins/PluginCollection.cs b/Engine/Source/Programs/Horde/HordeServer/Plugins/PluginCollection.cs
--- a/Engine/Source/Programs/Horde/HordeServer/Plugins/PluginCollection.cs
+++ b/Engine/Source/Programs/Horde/HordeServer/Plugins/PluginCollection.cs
@@ -2,6 +2,7 @@
 
 using System;
 using System.Collections.Generic;
+using System.Diagnostics;
 using System.Reflection;
 using Microsoft.Extensions.Configuration;
 using Microsoft.Extensions.DependencyInjection;
@@ -32,8 +33,9 @@ namespace HordeServer.Plugins
 				=> _loadedPlugin;
 		}
 
+		[DebuggerDisplay("{Name}")]
 		class LoadedPlugin<TServerConfig, TGlobalConfig, TStartup> : ILoadedPlugin
-			where TServerConfig : class, new()
+			where TServerConfig : PluginServerConfig, new()
 			where TGlobalConfig : class, IPluginConfig, new()
 			where TStartup : class, IPluginStartup
 		{
@@ -145,17 +147,18 @@ namespace HordeServer.Plugins
 		/// <summary>
 		/// Adds a plugin with the given startup class
 		/// </summary>
-		ILoadedPlugin AddLoadedPlugin(IPluginMetadata metadata, Type startupType)
+		public ILoadedPlugin Add(Type startupType)
 		{
+			PluginAttribute attr = startupType.GetCustomAttribute<PluginAttribute>()
+				?? throw new InvalidOperationException($"Cannot add {startupType.Name} as a plugin. No {nameof(PluginAttribute)} was found.");
+
+			IPluginMetadata metadata = new PluginMetadata { Name = new PluginName(attr.Name) };
 			if (!_loadedPluginNames.Add(metadata.Name))
 			{
 				throw new InvalidOperationException($"An implementation of the {metadata.Name} plugin has already been added");
 			}
 
-			PluginAttribute attr = startupType.GetCustomAttribute<PluginAttribute>()
-				?? throw new InvalidOperationException($"Cannot add {startupType.Name} as a plugin. No {nameof(PluginAttribute)} was found.");
-
-			Type pluginType = typeof(LoadedPlugin<,,>).MakeGenericType(attr.ServerConfigType ?? typeof(object), attr.GlobalConfigType ?? typeof(EmptyPluginConfig), startupType);
+			Type pluginType = typeof(LoadedPlugin<,,>).MakeGenericType(attr.ServerConfigType ?? typeof(PluginServerConfig), attr.GlobalConfigType ?? typeof(EmptyPluginConfig), startupType);
 			ILoadedPlugin loadedPlugin = (ILoadedPlugin)Activator.CreateInstance(pluginType, metadata)!;
 			_loadedPlugins.Add(loadedPlugin);
 
@@ -167,23 +170,6 @@ namespace HordeServer.Plugins
 		/// </summary>
 		/// <typeparam name="T">Type of the startup class</typeparam>
 		public ILoadedPlugin Add<T>() where T : class, IPluginStartup
-		{
-			PluginAttribute attr = typeof(T).GetCustomAttribute<PluginAttribute>()
-				?? throw new InvalidOperationException($"Cannot add {typeof(T).Name} as a plugin. No {nameof(PluginAttribute)} was found.");
-
-			return Add<T>(new PluginMetadata{ Name = new PluginName(attr.Name) });
-		}
-
-		/// <summary>
-		/// Adds a plugin with the given startup class
-		/// </summary>
-		/// <typeparam name="T">Type of the startup class</typeparam>
-		/// <param name="metadata">Metadata for the plugin</param>
-		public ILoadedPlugin Add<T>(IPluginMetadata metadata) where T : class, IPluginStartup
-		{
-			ILoadedPlugin plugin = AddLoadedPlugin(metadata, typeof(T));
-			_plugins.Add(plugin);
-			return plugin;
-		}
+			=> Add(typeof(T));
 	}
 }
diff --git a/Engine/Source/Programs/Horde/HordeServer/ServerApp.cs b/Engine/Source/Programs/Horde/HordeServer/ServerApp.cs
--- a/Engine/Source/Programs/Horde/HordeServer/ServerApp.cs
+++ b/Engine/Source/Programs/Horde/HordeServer/ServerApp.cs
@@ -12,7 +12,6 @@ using System.Threading.Tasks;
 using EpicGames.Core;
 using HordeServer.Commands;
 using HordeServer.Plugins;
-using HordeServer.Storage;
 using HordeServer.Utilities;
 using Microsoft.Extensions.Configuration;
 using Microsoft.Extensions.Configuration.Memory;
@@ -91,9 +90,12 @@ namespace HordeServer
 
 		public static Type[] ConfigSchemas = FindSchemaTypes();
 
+		public static IPluginCollection Plugins => s_pluginCollection ?? throw new InvalidOperationException("Plugin collection has not been initialized");
+
 		private static DirectoryReference s_dataDir = DirectoryReference.Combine(GetAppDir(), "Data");
 		private static DirectoryReference s_configDir = DirectoryReference.Combine(GetAppDir(), "Defaults");
 		private static FileReference? s_serverConfigFile;
+		private static IPluginCollection? s_pluginCollection;
 
 		static Type[] FindSchemaTypes()
 		{
@@ -141,12 +143,12 @@ namespace HordeServer
 					if (dataDir != null)
 					{
 						s_dataDir = dataDir;
-					}					
+					}
 				}
 				else
 				{
 					s_dataDir = DirectoryReference.Combine(GetAppDir(), baseServerSettings.DataDir);
-				}								
+				}
 			}
 			else if (baseServerSettings.Installed && RuntimeInformation.IsOSPlatform(OSPlatform.Windows))
 			{
@@ -196,10 +198,10 @@ namespace HordeServer
 			ServerInfo serverInfo = new ServerInfo(config, Options.Create(serverSettings));
 			services.AddSingleton<IServerInfo>(serverInfo);
 
-			IPluginCollection pluginCollection = ConfigurePlugins();
-			services.AddSingleton<IPluginCollection>(pluginCollection);
+			s_pluginCollection = CreatePluginCollection(config);
+			services.AddSingleton<IPluginCollection>(s_pluginCollection);
 
-			foreach (Assembly pluginAssembly in pluginCollection.LoadedPlugins.Select(x => x.Assembly).Distinct())
+			foreach (Assembly pluginAssembly in s_pluginCollection.LoadedPlugins.Select(x => x.Assembly).Distinct())
 			{
 				services.AddCommandsFromAssembly(pluginAssembly);
 			}
@@ -212,17 +214,49 @@ namespace HordeServer
 #pragma warning restore ASP0000 // Do not call 'IServiceCollection.BuildServiceProvider' in 'ConfigureServices'
 		}
 
-		public static IPluginCollection ConfigurePlugins()
+		static IPluginCollection CreatePluginCollection(IConfiguration configuration)
 		{
+			Dictionary<string, PluginServerConfig> pluginConfigs = new Dictionary<string, PluginServerConfig>(StringComparer.OrdinalIgnoreCase);
+			configuration.GetSection("Horde").GetSection("Plugins").Bind(pluginConfigs);
+
+			FileFilter pluginFilter = new FileFilter(FileFilterType.Exclude);
+			pluginFilter.AddRule("/HordeServer.*.dll");
+
+			List<FileReference> files = pluginFilter.ApplyToDirectory(AppDir, true);
+
+			HashSet<string> missingPlugins = new HashSet<string>(StringComparer.OrdinalIgnoreCase);
+			foreach ((string name, PluginServerConfig pluginConfig) in pluginConfigs)
+			{
+				if (pluginConfig.Enabled)
+				{
+					missingPlugins.Add(name);
+				}
+			}
+
 			PluginCollection pluginCollection = new PluginCollection();
-			pluginCollection.Add<AnalyticsPlugin>();
-			pluginCollection.Add<BuildPlugin>();
-			pluginCollection.Add<ComputePlugin>();
-			pluginCollection.Add<DdcPlugin>();
-			pluginCollection.Add<SecretsPlugin>();
-			pluginCollection.Add<StoragePlugin>();
-			pluginCollection.Add<SymbolsPlugin>();
-			pluginCollection.Add<ToolsPlugin>();
+			foreach (FileReference file in files)
+			{
+				Assembly assembly = Assembly.LoadFrom(file.FullName);
+				foreach (Type type in assembly.GetExportedTypes())
+				{
+					PluginAttribute? pluginAttribute = type.GetCustomAttribute<PluginAttribute>();
+					if (pluginAttribute != null)
+					{
+						PluginServerConfig? pluginConfig;
+						if (!pluginConfigs.TryGetValue(pluginAttribute.Name, out pluginConfig) || pluginConfig.Enabled)
+						{
+							pluginCollection.Add(type);
+							missingPlugins.Remove(pluginAttribute.Name);
+						}
+					}
+				}
+			}
+
+			if (missingPlugins.Count > 0)
+			{
+				throw new InvalidOperationException($"Unable to find plugin(s) enabled in config file: {StringUtils.FormatList(missingPlugins)}");
+			}
+
 			return pluginCollection;
 		}
 
@@ -335,7 +369,7 @@ namespace HordeServer
 						{
 							regDataDir = Registry.GetValue("HKEY_LOCAL_MACHINE\\SOFTWARE\\Epic Games\\Horde\\Server", "DataDir", null) as string;
 						}
-					}					
+					}
 				}
 			}
 
diff --git a/Engine/Source/Programs/Horde/HordeServer/Startup.cs b/Engine/Source/Programs/Horde/HordeServer/Startup.cs
--- a/Engine/Source/Programs/Horde/HordeServer/Startup.cs
+++ b/Engine/Source/Programs/Horde/HordeServer/Startup.cs
@@ -274,7 +274,7 @@ namespace HordeServer
 			services.AddSingleton<IServerInfo>(serverInfo);
 
 			// Register the plugin collection
-			IPluginCollection pluginCollection = ServerApp.ConfigurePlugins();
+			IPluginCollection pluginCollection = ServerApp.Plugins;
 			services.AddSingleton<IPluginCollection>(pluginCollection);
 
 			// Register all the plugin services
diff --git a/Engine/Source/Programs/Horde/Plugins/Analytics/HordeServer.Analytics/AnalyticsServerConfig.cs b/Engine/Source/Programs/Horde/Plugins/Analytics/HordeServer.Analytics/AnalyticsServerConfig.cs
--- a/Engine/Source/Programs/Horde/Plugins/Analytics/HordeServer.Analytics/AnalyticsServerConfig.cs
+++ b/Engine/Source/Programs/Horde/Plugins/Analytics/HordeServer.Analytics/AnalyticsServerConfig.cs
@@ -1,5 +1,6 @@
 // Copyright Epic Games, Inc. All Rights Reserved.
 
+using HordeServer.Plugins;
 using HordeServer.Telemetry.Sinks;
 
 #pragma warning disable CA2227 // Change x to be read-only by removing the property setter
@@ -9,7 +10,7 @@ namespace HordeServer
 	/// <summary>
 	/// Server configuration for the analytics system
 	/// </summary>
-	public class AnalyticsServerConfig
+	public class AnalyticsServerConfig : PluginServerConfig
 	{
 		/// <summary>
 		/// Settings for the various telemetry sinks
diff --git a/Engine/Source/Programs/Horde/Plugins/Build/HordeServer.Build/BuildServerConfig.cs b/Engine/Source/Programs/Horde/Plugins/Build/HordeServer.Build/BuildServerConfig.cs
--- a/Engine/Source/Programs/Horde/Plugins/Build/HordeServer.Build/BuildServerConfig.cs
+++ b/Engine/Source/Programs/Horde/Plugins/Build/HordeServer.Build/BuildServerConfig.cs
@@ -13,7 +13,7 @@ namespace HordeServer
 	/// <summary>
 	/// Static configuration for the build plugin
 	/// </summary>
-	public class BuildServerConfig
+	public class BuildServerConfig : PluginServerConfig
 	{
 		/// <summary>
 		/// Perforce connections for use by the Horde server (not agents)
diff --git a/Engine/Source/Programs/Horde/Plugins/Compute/HordeServer.Compute/ComputeServerConfig.cs b/Engine/Source/Programs/Horde/Plugins/Compute/HordeServer.Compute/ComputeServerConfig.cs
--- a/Engine/Source/Programs/Horde/Plugins/Compute/HordeServer.Compute/ComputeServerConfig.cs
+++ b/Engine/Source/Programs/Horde/Plugins/Compute/HordeServer.Compute/ComputeServerConfig.cs
@@ -11,7 +11,7 @@ namespace HordeServer
 	/// <summary>
 	/// Static configuration for the compute plugin
 	/// </summary>
-	public class ComputeServerConfig
+	public class ComputeServerConfig : PluginServerConfig
 	{
 		/// <summary>
 		/// Whether to enable the upgrade task source.
diff --git a/Engine/Source/Programs/Horde/Plugins/Storage/HordeServer.Storage/Storage/StorageService.cs b/Engine/Source/Programs/Horde/Plugins/Storage/HordeServer.Storage/Storage/StorageService.cs
--- a/Engine/Source/Programs/Horde/Plugins/Storage/HordeServer.Storage/Storage/StorageService.cs
+++ b/Engine/Source/Programs/Horde/Plugins/Storage/HordeServer.Storage/Storage/StorageService.cs
@@ -419,7 +419,7 @@ namespace HordeServer.Storage
 		readonly IMemoryCache _memoryCache;
 		readonly IObjectStoreFactory _objectStoreFactory;
 		readonly IOptionsMonitor<StorageConfig> _storageConfig;
-		readonly IOptions<StaticStorageConfig> _staticStorageConfig;
+		readonly IOptions<StorageServerConfig> _staticStorageConfig;
 		readonly Tracer _tracer;
 		readonly ILogger _logger;
 
@@ -462,7 +462,7 @@ namespace HordeServer.Storage
 		/// <summary>
 		/// Constructor
 		/// </summary>
-		public StorageService(IMongoService mongoService, IRedisService redisService, IClock clock, BundleCache bundleCache, IMemoryCache memoryCache, IObjectStoreFactory objectStoreFactory, IOptionsMonitor<StorageConfig> storageConfig, IOptions<StaticStorageConfig> staticStorageConfig, Tracer tracer, ILogger<StorageService> logger)
+		public StorageService(IMongoService mongoService, IRedisService redisService, IClock clock, BundleCache bundleCache, IMemoryCache memoryCache, IObjectStoreFactory objectStoreFactory, IOptionsMonitor<StorageConfig> storageConfig, IOptions<StorageServerConfig> staticStorageConfig, Tracer tracer, ILogger<StorageService> logger)
 		{
 			_redisService = redisService;
 			_clock = clock;
diff --git a/Engine/Source/Programs/Horde/Plugins/Storage/HordeServer.Storage/StorageConfig.cs b/Engine/Source/Programs/Horde/Plugins/Storage/HordeServer.Storage/StorageConfig.cs
--- a/Engine/Source/Programs/Horde/Plugins/Storage/HordeServer.Storage/StorageConfig.cs
+++ b/Engine/Source/Programs/Horde/Plugins/Storage/HordeServer.Storage/StorageConfig.cs
@@ -171,7 +171,7 @@ namespace HordeServer.Storage
 	/// <summary>
 	/// Static settings for the storage system
 	/// </summary>
-	public class StaticStorageConfig
+	public class StorageServerConfig : PluginServerConfig
 	{
 		/// <summary>
 		/// Directory to use for the coarse-grained backend cache. This caches full bundles downloaded from the upstream object store.
diff --git a/Engine/Source/Programs/Horde/Plugins/Storage/HordeServer.Storage/StoragePlugin.cs b/Engine/Source/Programs/Horde/Plugins/Storage/HordeServer.Storage/StoragePlugin.cs
--- a/Engine/Source/Programs/Horde/Plugins/Storage/HordeServer.Storage/StoragePlugin.cs
+++ b/Engine/Source/Programs/Horde/Plugins/Storage/HordeServer.Storage/StoragePlugin.cs
@@ -16,16 +16,16 @@ namespace HordeServer.Storage
 	/// <summary>
 	/// Entry point for the storage plugin
 	/// </summary>
-	[Plugin("Storage", GlobalConfigType = typeof(StorageConfig), ServerConfigType = typeof(StaticStorageConfig))]
+	[Plugin("Storage", GlobalConfigType = typeof(StorageConfig), ServerConfigType = typeof(StorageServerConfig))]
 	public class StoragePlugin : IPluginStartup
 	{
 		readonly IServerInfo _serverInfo;
-		readonly StaticStorageConfig _staticConfig;
+		readonly StorageServerConfig _staticConfig;
 
 		/// <summary>
 		/// Constructor
 		/// </summary>
-		public StoragePlugin(IServerInfo serverInfo, StaticStorageConfig staticConfig)
+		public StoragePlugin(IServerInfo serverInfo, StorageServerConfig staticConfig)
 		{
 			_serverInfo = serverInfo;
 			_staticConfig = staticConfig;
diff --git a/Engine/Source/Programs/Horde/Plugins/Tools/HordeServer.Tools/ToolsServerConfig.cs b/Engine/Source/Programs/Horde/Plugins/Tools/HordeServer.Tools/ToolsServerConfig.cs
--- a/Engine/Source/Programs/Horde/Plugins/Tools/HordeServer.Tools/ToolsServerConfig.cs
+++ b/Engine/Source/Programs/Horde/Plugins/Tools/HordeServer.Tools/ToolsServerConfig.cs
@@ -2,6 +2,7 @@
 
 using System.Diagnostics.CodeAnalysis;
 using EpicGames.Horde.Tools;
+using HordeServer.Plugins;
 using HordeServer.Tools;
 
 #pragma warning disable CA2227 // Change x to be read-only by removing the property setter
@@ -11,7 +12,7 @@ namespace HordeServer
 	/// <summary>
 	/// Server configuration for bundled tools
 	/// </summary>
-	public class ToolsServerConfig
+	public class ToolsServerConfig : PluginServerConfig
 	{
 		/// <summary>
 		/// Tools bundled along with the server. Data for each tool can be produced using the 'bundle create' command, and should be stored in the Tools directory.
