Horde: Merging changes from Main up to CL 36937914.

diff --git a/Engine/Source/Programs/Horde/HordeServer.Tests.Shared/ControllerIntegrationTest.cs b/Engine/Source/Programs/Horde/HordeServer.Tests.Shared/ControllerIntegrationTest.cs
--- a/Engine/Source/Programs/Horde/HordeServer.Tests.Shared/ControllerIntegrationTest.cs
+++ b/Engine/Source/Programs/Horde/HordeServer.Tests.Shared/ControllerIntegrationTest.cs
@@ -90,6 +90,7 @@ public class FakeHordeWebApp : IAsyncDisposable
 
 	public FakeHordeWebApp(Action<IServiceCollection>? configureServices = null, Dictionary<string, string>? settings = null)
 	{
+		ServerApp.InitializePluginsForTests();
 		MongoInstance = new MongoInstance();
 		RedisInstance = new RedisInstance();
 		Factory = new TestWebApplicationFactory<Startup>(MongoInstance, RedisInstance, configureServices, settings);
diff --git a/Engine/Source/Programs/Horde/HordeServer/ServerApp.cs b/Engine/Source/Programs/Horde/HordeServer/ServerApp.cs
--- a/Engine/Source/Programs/Horde/HordeServer/ServerApp.cs
+++ b/Engine/Source/Programs/Horde/HordeServer/ServerApp.cs
@@ -214,6 +214,11 @@ namespace HordeServer
 #pragma warning restore ASP0000 // Do not call 'IServiceCollection.BuildServiceProvider' in 'ConfigureServices'
 		}
 
+		internal static void InitializePluginsForTests()
+		{
+			s_pluginCollection ??= CreatePluginCollection(new ConfigurationBuilder().Build());
+		}
+
 		static IPluginCollection CreatePluginCollection(IConfiguration configuration)
 		{
 			Dictionary<string, PluginServerConfig> pluginConfigs = new Dictionary<string, PluginServerConfig>(StringComparer.OrdinalIgnoreCase);
