Horde: Fix some static analysis warnings.

#rnx

diff --git a/Engine/Source/Programs/Horde/HordeServer.Shared/Plugins/PluginServerConfig.cs b/Engine/Source/Programs/Horde/HordeServer.Shared/Plugins/PluginServerConfig.cs
--- a/Engine/Source/Programs/Horde/HordeServer.Shared/Plugins/PluginServerConfig.cs
+++ b/Engine/Source/Programs/Horde/HordeServer.Shared/Plugins/PluginServerConfig.cs
@@ -1,11 +1,5 @@
 // Copyright Epic Games, Inc. All Rights Reserved.
 
-using System;
-using System.Collections.Generic;
-using System.Linq;
-using System.Text;
-using System.Threading.Tasks;
-
 namespace HordeServer.Plugins
 {
 	/// <summary>
