Horde: Return the invalid_grant response when refresh tokens have expired. OIDCToken checks for this and handles re-login if this code is returned, but does not for invalid_request.

#rnx

diff --git a/Engine/Source/Programs/Horde/HordeServer/Accounts/OAuthController.cs b/Engine/Source/Programs/Horde/HordeServer/Accounts/OAuthController.cs
--- a/Engine/Source/Programs/Horde/HordeServer/Accounts/OAuthController.cs
+++ b/Engine/Source/Programs/Horde/HordeServer/Accounts/OAuthController.cs
@@ -291,21 +291,21 @@ namespace HordeServer.Accounts
 			AccountId accountId;
 			if (!TryParseAccountIdFromSubject(payload, out accountId))
 			{
-				OAuthGetTokenResponse error = new OAuthGetTokenResponse { Error = "invalid_request", ErrorDescription = $"Missing account-id subject" };
+				OAuthGetTokenResponse error = new OAuthGetTokenResponse { Error = "invalid_grant", ErrorDescription = $"Missing account-id subject" };
 				return Unauthorized(error);
 			}
 
 			IAccount? account = await _accountCollection.GetAsync(accountId, cancellationToken);
 			if (account == null)
 			{
-				OAuthGetTokenResponse error = new OAuthGetTokenResponse { Error = "invalid_request", ErrorDescription = $"Invalid account-id ({accountId})" };
+				OAuthGetTokenResponse error = new OAuthGetTokenResponse { Error = "invalid_grant", ErrorDescription = $"Invalid account-id ({accountId})" };
 				return Unauthorized(error);
 			}
 
 			string? session = GetClaimOrDefault(payload, SessionClaim);
 			if (session != account.SessionKey)
 			{
-				OAuthGetTokenResponse error = new OAuthGetTokenResponse { Error = "invalid_request", ErrorDescription = $"Invalid session key ('{session}')" };
+				OAuthGetTokenResponse error = new OAuthGetTokenResponse { Error = "invalid_grant", ErrorDescription = $"Invalid session key ('{session}')" };
 				return Unauthorized(error);
 			}
 
