Horde: Add some additional null checks to help isolate error in dev.

#rnx

diff --git a/Engine/Source/Programs/Horde/HordeServer.Shared/Acls/AclConfig.cs b/Engine/Source/Programs/Horde/HordeServer.Shared/Acls/AclConfig.cs
--- a/Engine/Source/Programs/Horde/HordeServer.Shared/Acls/AclConfig.cs
+++ b/Engine/Source/Programs/Horde/HordeServer.Shared/Acls/AclConfig.cs
@@ -106,6 +106,27 @@ namespace HordeServer.Acls
 			{
 				foreach (AclEntryConfig entry in Entries)
 				{
+					if (user == null)
+					{
+						throw new NullReferenceException("User is null");
+					}
+					if (entry == null)
+					{
+						throw new NullReferenceException("Entry is null");
+					}
+					if (entry.ComputedActions == null)
+					{
+						throw new NullReferenceException("ComputedActions is null");
+					}
+					if (entry.Claim.Type == null)
+					{
+						throw new NullReferenceException("Claim.Type is null");
+					}
+					if (entry.Claim.Value == null)
+					{
+						throw new NullReferenceException("Claim.Value is null");
+					}
+
 					if (entry.ComputedActions.Contains(action) && user.HasClaim(entry.Claim.Type, entry.Claim.Value))
 					{
 						return true;
