Horde: Add additional logging for plugin loading.

#rnx

diff --git a/Engine/Source/Programs/Horde/HordeServer/ServerApp.cs b/Engine/Source/Programs/Horde/HordeServer/ServerApp.cs
--- a/Engine/Source/Programs/Horde/HordeServer/ServerApp.cs
+++ b/Engine/Source/Programs/Horde/HordeServer/ServerApp.cs
@@ -7,6 +7,7 @@ using System.IO;
 using System.Linq;
 using System.Reflection;
 using System.Runtime.InteropServices;
+using System.Text.RegularExpressions;
 using System.Threading;
 using System.Threading.Tasks;
 using EpicGames.Core;
@@ -17,6 +18,8 @@ using Microsoft.Extensions.Configuration;
 using Microsoft.Extensions.Configuration.Memory;
 using Microsoft.Extensions.DependencyInjection;
 using Microsoft.Extensions.Hosting;
+using Microsoft.Extensions.Logging;
+using Microsoft.Extensions.Logging.Abstractions;
 using Microsoft.Extensions.Options;
 using Microsoft.Win32;
 using Serilog;
@@ -24,11 +27,14 @@ using Serilog.Configuration;
 using Serilog.Exceptions;
 using Serilog.Exceptions.Core;
 using Serilog.Exceptions.Grpc.Destructurers;
+using Serilog.Extensions.Logging;
 using Serilog.Formatting.Json;
 using Serilog.Sinks.SystemConsole.Themes;
 
 namespace HordeServer
 {
+	using ILogger = Microsoft.Extensions.Logging.ILogger;
+
 	static class LoggerExtensions
 	{
 		class DatadogVersionLogEnricher : Serilog.Core.ILogEventEnricher
@@ -187,6 +193,8 @@ namespace HordeServer
 				.ReadFrom.Configuration(config)
 				.CreateLogger();
 
+			ILogger startupLogger = new SerilogLoggerFactory().CreateLogger(typeof(ServerApp).FullName ?? "ServerApp");
+
 			ServiceCollection services = new ServiceCollection();
 			services.AddCommandsFromAssembly(Assembly.GetExecutingAssembly());
 			services.AddLogging(builder => builder.AddSerilog());
@@ -198,7 +206,7 @@ namespace HordeServer
 			ServerInfo serverInfo = new ServerInfo(config, Options.Create(serverSettings));
 			services.AddSingleton<IServerInfo>(serverInfo);
 
-			s_pluginCollection = CreatePluginCollection(config);
+			s_pluginCollection = CreatePluginCollection(config, startupLogger);
 			services.AddSingleton<IPluginCollection>(s_pluginCollection);
 
 			foreach (Assembly pluginAssembly in s_pluginCollection.LoadedPlugins.Select(x => x.Assembly).Distinct())
@@ -216,18 +224,24 @@ namespace HordeServer
 
 		internal static void InitializePluginsForTests()
 		{
-			s_pluginCollection ??= CreatePluginCollection(new ConfigurationBuilder().Build());
+			s_pluginCollection ??= CreatePluginCollection(new ConfigurationBuilder().Build(), NullLogger.Instance);
 		}
 
-		static IPluginCollection CreatePluginCollection(IConfiguration configuration)
+		static IPluginCollection CreatePluginCollection(IConfiguration configuration, ILogger logger)
 		{
 			Dictionary<string, PluginServerConfig> pluginConfigs = new Dictionary<string, PluginServerConfig>(StringComparer.OrdinalIgnoreCase);
 			configuration.GetSection("Horde").GetSection("Plugins").Bind(pluginConfigs);
 
-			FileFilter pluginFilter = new FileFilter(FileFilterType.Exclude);
-			pluginFilter.AddRule("/HordeServer.*.dll");
-
-			List<FileReference> files = pluginFilter.ApplyToDirectory(AppDir, true);
+			List<FileReference> files = new List<FileReference>();
+			foreach (FileInfo fileInfo in AppDir.ToDirectoryInfo().EnumerateFiles())
+			{
+				const string Prefix = "HordeServer.";
+				const string Suffix = ".dll";
+				if (fileInfo.Name.StartsWith(Prefix, StringComparison.OrdinalIgnoreCase) && fileInfo.Name.EndsWith(Suffix, StringComparison.OrdinalIgnoreCase) && fileInfo.Name.Length > Prefix.Length + Suffix.Length)
+				{
+					files.Add(new FileReference(fileInfo));
+				}
+			}
 
 			HashSet<string> missingPlugins = new HashSet<string>(StringComparer.OrdinalIgnoreCase);
 			foreach ((string name, PluginServerConfig pluginConfig) in pluginConfigs)
@@ -241,20 +255,30 @@ namespace HordeServer
 			PluginCollection pluginCollection = new PluginCollection();
 			foreach (FileReference file in files)
 			{
-				Assembly assembly = Assembly.LoadFrom(file.FullName);
-				foreach (Type type in assembly.GetExportedTypes())
+				logger.LogInformation("Loading {File}", file);
+				try
 				{
-					PluginAttribute? pluginAttribute = type.GetCustomAttribute<PluginAttribute>();
-					if (pluginAttribute != null)
+					Assembly assembly = Assembly.LoadFrom(file.FullName);
+					foreach (Type type in assembly.GetExportedTypes())
 					{
-						PluginServerConfig? pluginConfig;
-						if (!pluginConfigs.TryGetValue(pluginAttribute.Name, out pluginConfig) || pluginConfig.Enabled)
+						PluginAttribute? pluginAttribute = type.GetCustomAttribute<PluginAttribute>();
+						if (pluginAttribute != null)
 						{
-							pluginCollection.Add(type);
-							missingPlugins.Remove(pluginAttribute.Name);
+							PluginServerConfig? pluginConfig;
+							if (!pluginConfigs.TryGetValue(pluginAttribute.Name, out pluginConfig) || pluginConfig.Enabled)
+							{
+								logger.LogDebug("Added plugin '{Plugin}'", pluginAttribute.Name);
+								pluginCollection.Add(type);
+								missingPlugins.Remove(pluginAttribute.Name);
+							}
 						}
 					}
 				}
+				catch (Exception ex)
+				{
+					logger.LogError(ex, "Error loading plugins from {File}: {Message}", file, ex.Message);
+					throw;
+				}
 			}
 
 			if (missingPlugins.Count > 0)
