Horde: Updating schema documentation

#rnx

diff --git a/Engine/Source/Programs/Horde/Docs/Config/Schema/AclActions.md b/Engine/Source/Programs/Horde/Docs/Config/Schema/AclActions.md
--- a/Engine/Source/Programs/Horde/Docs/Config/Schema/AclActions.md
+++ b/Engine/Source/Programs/Horde/Docs/Config/Schema/AclActions.md
@@ -63,291 +63,6 @@
 | `DeleteAccount` |  |
 | `ViewAccount` |  |
 
-## Artifacts
-
-| Name | Description |
-| ---- | ----------- |
-| `ReadArtifact` | Permission to read from an artifact |
-| `WriteArtifact` | Permission to write to an artifact |
-| `DeleteArtifact` | Permission to delete to an artifact |
-| `UploadArtifact` | Ability to create an artifact. Typically just for debugging; agents have this access for a particular session. |
-| `DownloadArtifact` | Ability to download an artifact |
-
-## Bisect
-
-| Name | Description |
-| ---- | ----------- |
-| `CreateBisectTask` | Ability to start new bisect tasks |
-| `UpdateBisectTask` | Ability to update a bisect task |
-| `ViewBisectTask` | Ability to view a bisect task |
-
-## Devices
-
-| Name | Description |
-| ---- | ----------- |
-| `DeviceRead` | Ability to read devices |
-| `DeviceWrite` | Ability to write devices |
-
-## Jobs
-
-| Name | Description |
-| ---- | ----------- |
-| `CreateJob` | Ability to start new jobs |
-| `UpdateJob` | Rename a job, modify its priority, etc... |
-| `DeleteJob` | Delete a job properties |
-| `ExecuteJob` | Allows updating a job metadata (name, changelist number, step properties, new groups, job states, etc...). Typically granted to agents. Not user facing. |
-| `RetryJobStep` | Ability to retry a failed job step |
-| `ViewJob` | Ability to view a job |
-
-## Notices
-
-| Name | Description |
-| ---- | ----------- |
-| `CreateNotice` |  |
-| `UpdateNotice` |  |
-| `DeleteNotice` |  |
-
-## Notifications
-
-| Name | Description |
-| ---- | ----------- |
-| `CreateSubscription` | Ability to subscribe to notifications |
-
-## Projects
-
-| Name | Description |
-| ---- | ----------- |
-| `CreateProject` | Allows the creation of new projects |
-| `DeleteProject` | Allows deletion of projects. |
-| `UpdateProject` | Modify attributes of a project (name, categories, etc...) |
-| `ViewProject` | View information about a project |
-
-## Replicators
-
-| Name | Description |
-| ---- | ----------- |
-| `UpdateReplicator` | Allows deletion of projects. |
-| `ViewReplicator` | Allows the creation of new projects |
-
-## ServiceAccounts
-
-| Name | Description |
-| ---- | ----------- |
-| `CreateAccount` |  |
-| `UpdateAccount` |  |
-| `DeleteAccount` |  |
-| `ViewAccount` |  |
-
-## Streams
-
-| Name | Description |
-| ---- | ----------- |
-| `CreateStream` | Allows the creation of new streams within a project |
-| `UpdateStream` | Allows updating a stream (agent types, templates, schedules) |
-| `DeleteStream` | Allows deleting a stream |
-| `ViewStream` | Ability to view a stream |
-| `ViewChanges` | View changes submitted to a stream. NOTE: this returns responses from the server's Perforce account, which may be a priviledged user. |
-| `ViewTemplate` | View template associated with a stream |
-
-## Accounts
-
-| Name | Description |
-| ---- | ----------- |
-| `CreateAccount` |  |
-| `UpdateAccount` |  |
-| `DeleteAccount` |  |
-| `ViewAccount` |  |
-
-## Agents
-
-| Name | Description |
-| ---- | ----------- |
-| `CreateAgent` | Ability to create an agent. This may be done explicitly, or granted to agents to allow them to self-register. |
-| `UpdateAgent` | Update an agent's name, pools, etc... |
-| `DeleteAgent` | Soft-delete an agent |
-| `ViewAgent` | View an agent |
-| `ListAgents` | List the available agents |
-
-## Artifacts
-
-| Name | Description |
-| ---- | ----------- |
-| `ReadArtifact` |  |
-| `WriteArtifact` |  |
-| `DeleteArtifact` |  |
-| `UploadArtifact` |  |
-| `DownloadArtifact` |  |
-
-## Bisect
-
-| Name | Description |
-| ---- | ----------- |
-| `CreateBisectTask` |  |
-| `UpdateBisectTask` |  |
-| `ViewBisectTask` |  |
-
-## Compute
-
-| Name | Description |
-| ---- | ----------- |
-| `AddComputeTasks` | User can add tasks to the compute cluster |
-| `GetComputeTasks` | User can get and list tasks from the compute cluster |
-
-## Devices
-
-| Name | Description |
-| ---- | ----------- |
-| `DeviceRead` |  |
-| `DeviceWrite` |  |
-
-## Jobs
-
-| Name | Description |
-| ---- | ----------- |
-| `CreateJob` |  |
-| `UpdateJob` |  |
-| `DeleteJob` |  |
-| `ExecuteJob` |  |
-| `RetryJobStep` |  |
-| `ViewJob` |  |
-
-## Leases
-
-| Name | Description |
-| ---- | ----------- |
-| `ViewLeases` | View all the leases that an agent has worked on |
-| `ViewLeaseTasks` | View the task data for a lease |
-
-## Logs
-
-| Name | Description |
-| ---- | ----------- |
-| `CreateLog` | Ability to create a log. Implicitly granted to agents. |
-| `UpdateLog` | Ability to update log metadata |
-| `ViewLog` | Ability to view a log contents |
-| `WriteLogData` | Ability to write log data |
-| `CreateEvent` | Ability to create events |
-| `ViewEvent` | Ability to view events |
-
-## Notices
-
-| Name | Description |
-| ---- | ----------- |
-| `CreateNotice` |  |
-| `UpdateNotice` |  |
-| `DeleteNotice` |  |
-
-## Notifications
-
-| Name | Description |
-| ---- | ----------- |
-| `CreateSubscription` |  |
-
-## Pools
-
-| Name | Description |
-| ---- | ----------- |
-| `CreatePool` | Create a global pool of agents |
-| `UpdatePool` | Modify an agent pool |
-| `DeletePool` | Delete an agent pool |
-| `ViewPool` | Ability to view a pool |
-| `ListPools` | View all the available agent pools |
-
-## Projects
-
-| Name | Description |
-| ---- | ----------- |
-| `CreateProject` |  |
-| `DeleteProject` |  |
-| `UpdateProject` |  |
-| `ViewProject` |  |
-
-## Replicators
-
-| Name | Description |
-| ---- | ----------- |
-| `UpdateReplicator` |  |
-| `ViewReplicator` |  |
-
-## ServiceAccounts
-
-| Name | Description |
-| ---- | ----------- |
-| `CreateAccount` |  |
-| `UpdateAccount` |  |
-| `DeleteAccount` |  |
-| `ViewAccount` |  |
-
-## Sessions
-
-| Name | Description |
-| ---- | ----------- |
-| `CreateSession` | Granted to agents to call CreateSession, which returns a bearer token identifying themselves valid to call UpdateSesssion via gRPC. |
-| `ViewSession` | Allows viewing information about an agent session |
-
-## Software
-
-| Name | Description |
-| ---- | ----------- |
-| `UploadSoftware` | Ability to upload new versions of the agent software |
-| `DownloadSoftware` | Ability to download the agent software |
-| `DeleteSoftware` | Ability to delete agent software |
-
-## Streams
-
-| Name | Description |
-| ---- | ----------- |
-| `CreateStream` |  |
-| `UpdateStream` |  |
-| `DeleteStream` |  |
-| `ViewStream` |  |
-| `ViewChanges` |  |
-| `ViewTemplate` |  |
-
-## Accounts
-
-| Name | Description |
-| ---- | ----------- |
-| `CreateAccount` |  |
-| `UpdateAccount` |  |
-| `DeleteAccount` |  |
-| `ViewAccount` |  |
-
-## Agents
-
-| Name | Description |
-| ---- | ----------- |
-| `CreateAgent` |  |
-| `UpdateAgent` |  |
-| `DeleteAgent` |  |
-| `ViewAgent` |  |
-| `ListAgents` |  |
-
-## Artifacts
-
-| Name | Description |
-| ---- | ----------- |
-| `ReadArtifact` |  |
-| `WriteArtifact` |  |
-| `DeleteArtifact` |  |
-| `UploadArtifact` |  |
-| `DownloadArtifact` |  |
-
-## Bisect
-
-| Name | Description |
-| ---- | ----------- |
-| `CreateBisectTask` |  |
-| `UpdateBisectTask` |  |
-| `ViewBisectTask` |  |
-
-## Compute
-
-| Name | Description |
-| ---- | ----------- |
-| `AddComputeTasks` |  |
-| `GetComputeTasks` |  |
-
 ## Ddc
 
 | Name | Description |
@@ -361,42 +76,6 @@
 | `DdcWriteTransactionLog` | Access to write the transaction log |
 | `DdcAdminAction` | Access to perform administrative task |
 
-## Devices
-
-| Name | Description |
-| ---- | ----------- |
-| `DeviceRead` |  |
-| `DeviceWrite` |  |
-
-## Jobs
-
-| Name | Description |
-| ---- | ----------- |
-| `CreateJob` |  |
-| `UpdateJob` |  |
-| `DeleteJob` |  |
-| `ExecuteJob` |  |
-| `RetryJobStep` |  |
-| `ViewJob` |  |
-
-## Leases
-
-| Name | Description |
-| ---- | ----------- |
-| `ViewLeases` |  |
-| `ViewLeaseTasks` |  |
-
-## Logs
-
-| Name | Description |
-| ---- | ----------- |
-| `CreateLog` |  |
-| `UpdateLog` |  |
-| `ViewLog` |  |
-| `WriteLogData` |  |
-| `CreateEvent` |  |
-| `ViewEvent` |  |
-
 ## Notices
 
 | Name | Description |
@@ -405,38 +84,6 @@
 | `UpdateNotice` |  |
 | `DeleteNotice` |  |
 
-## Notifications
-
-| Name | Description |
-| ---- | ----------- |
-| `CreateSubscription` |  |
-
-## Pools
-
-| Name | Description |
-| ---- | ----------- |
-| `CreatePool` |  |
-| `UpdatePool` |  |
-| `DeletePool` |  |
-| `ViewPool` |  |
-| `ListPools` |  |
-
-## Projects
-
-| Name | Description |
-| ---- | ----------- |
-| `CreateProject` |  |
-| `DeleteProject` |  |
-| `UpdateProject` |  |
-| `ViewProject` |  |
-
-## Replicators
-
-| Name | Description |
-| ---- | ----------- |
-| `UpdateReplicator` |  |
-| `ViewReplicator` |  |
-
 ## ServiceAccounts
 
 | Name | Description |
@@ -446,32 +93,6 @@
 | `DeleteAccount` |  |
 | `ViewAccount` |  |
 
-## Sessions
-
-| Name | Description |
-| ---- | ----------- |
-| `CreateSession` |  |
-| `ViewSession` |  |
-
-## Software
-
-| Name | Description |
-| ---- | ----------- |
-| `UploadSoftware` |  |
-| `DownloadSoftware` |  |
-| `DeleteSoftware` |  |
-
-## Streams
-
-| Name | Description |
-| ---- | ----------- |
-| `CreateStream` |  |
-| `UpdateStream` |  |
-| `DeleteStream` |  |
-| `ViewStream` |  |
-| `ViewChanges` |  |
-| `ViewTemplate` |  |
-
 ## Accounts
 
 | Name | Description |
@@ -481,41 +102,6 @@
 | `DeleteAccount` |  |
 | `ViewAccount` |  |
 
-## Agents
-
-| Name | Description |
-| ---- | ----------- |
-| `CreateAgent` |  |
-| `UpdateAgent` |  |
-| `DeleteAgent` |  |
-| `ViewAgent` |  |
-| `ListAgents` |  |
-
-## Artifacts
-
-| Name | Description |
-| ---- | ----------- |
-| `ReadArtifact` |  |
-| `WriteArtifact` |  |
-| `DeleteArtifact` |  |
-| `UploadArtifact` |  |
-| `DownloadArtifact` |  |
-
-## Bisect
-
-| Name | Description |
-| ---- | ----------- |
-| `CreateBisectTask` |  |
-| `UpdateBisectTask` |  |
-| `ViewBisectTask` |  |
-
-## Compute
-
-| Name | Description |
-| ---- | ----------- |
-| `AddComputeTasks` |  |
-| `GetComputeTasks` |  |
-
 ## Ddc
 
 | Name | Description |
@@ -529,42 +115,6 @@
 | `DdcWriteTransactionLog` |  |
 | `DdcAdminAction` |  |
 
-## Devices
-
-| Name | Description |
-| ---- | ----------- |
-| `DeviceRead` |  |
-| `DeviceWrite` |  |
-
-## Jobs
-
-| Name | Description |
-| ---- | ----------- |
-| `CreateJob` |  |
-| `UpdateJob` |  |
-| `DeleteJob` |  |
-| `ExecuteJob` |  |
-| `RetryJobStep` |  |
-| `ViewJob` |  |
-
-## Leases
-
-| Name | Description |
-| ---- | ----------- |
-| `ViewLeases` |  |
-| `ViewLeaseTasks` |  |
-
-## Logs
-
-| Name | Description |
-| ---- | ----------- |
-| `CreateLog` |  |
-| `UpdateLog` |  |
-| `ViewLog` |  |
-| `WriteLogData` |  |
-| `CreateEvent` |  |
-| `ViewEvent` |  |
-
 ## Notices
 
 | Name | Description |
@@ -573,38 +123,6 @@
 | `UpdateNotice` |  |
 | `DeleteNotice` |  |
 
-## Notifications
-
-| Name | Description |
-| ---- | ----------- |
-| `CreateSubscription` |  |
-
-## Pools
-
-| Name | Description |
-| ---- | ----------- |
-| `CreatePool` |  |
-| `UpdatePool` |  |
-| `DeletePool` |  |
-| `ViewPool` |  |
-| `ListPools` |  |
-
-## Projects
-
-| Name | Description |
-| ---- | ----------- |
-| `CreateProject` |  |
-| `DeleteProject` |  |
-| `UpdateProject` |  |
-| `ViewProject` |  |
-
-## Replicators
-
-| Name | Description |
-| ---- | ----------- |
-| `UpdateReplicator` |  |
-| `ViewReplicator` |  |
-
 ## Secrets
 
 | Name | Description |
@@ -620,32 +138,6 @@
 | `DeleteAccount` |  |
 | `ViewAccount` |  |
 
-## Sessions
-
-| Name | Description |
-| ---- | ----------- |
-| `CreateSession` |  |
-| `ViewSession` |  |
-
-## Software
-
-| Name | Description |
-| ---- | ----------- |
-| `UploadSoftware` |  |
-| `DownloadSoftware` |  |
-| `DeleteSoftware` |  |
-
-## Streams
-
-| Name | Description |
-| ---- | ----------- |
-| `CreateStream` |  |
-| `UpdateStream` |  |
-| `DeleteStream` |  |
-| `ViewStream` |  |
-| `ViewChanges` |  |
-| `ViewTemplate` |  |
-
 ## Accounts
 
 | Name | Description |
@@ -655,41 +147,6 @@
 | `DeleteAccount` |  |
 | `ViewAccount` |  |
 
-## Agents
-
-| Name | Description |
-| ---- | ----------- |
-| `CreateAgent` |  |
-| `UpdateAgent` |  |
-| `DeleteAgent` |  |
-| `ViewAgent` |  |
-| `ListAgents` |  |
-
-## Artifacts
-
-| Name | Description |
-| ---- | ----------- |
-| `ReadArtifact` |  |
-| `WriteArtifact` |  |
-| `DeleteArtifact` |  |
-| `UploadArtifact` |  |
-| `DownloadArtifact` |  |
-
-## Bisect
-
-| Name | Description |
-| ---- | ----------- |
-| `CreateBisectTask` |  |
-| `UpdateBisectTask` |  |
-| `ViewBisectTask` |  |
-
-## Compute
-
-| Name | Description |
-| ---- | ----------- |
-| `AddComputeTasks` |  |
-| `GetComputeTasks` |  |
-
 ## Ddc
 
 | Name | Description |
@@ -703,42 +160,6 @@
 | `DdcWriteTransactionLog` |  |
 | `DdcAdminAction` |  |
 
-## Devices
-
-| Name | Description |
-| ---- | ----------- |
-| `DeviceRead` |  |
-| `DeviceWrite` |  |
-
-## Jobs
-
-| Name | Description |
-| ---- | ----------- |
-| `CreateJob` |  |
-| `UpdateJob` |  |
-| `DeleteJob` |  |
-| `ExecuteJob` |  |
-| `RetryJobStep` |  |
-| `ViewJob` |  |
-
-## Leases
-
-| Name | Description |
-| ---- | ----------- |
-| `ViewLeases` |  |
-| `ViewLeaseTasks` |  |
-
-## Logs
-
-| Name | Description |
-| ---- | ----------- |
-| `CreateLog` |  |
-| `UpdateLog` |  |
-| `ViewLog` |  |
-| `WriteLogData` |  |
-| `CreateEvent` |  |
-| `ViewEvent` |  |
-
 ## Notices
 
 | Name | Description |
@@ -747,38 +168,6 @@
 | `UpdateNotice` |  |
 | `DeleteNotice` |  |
 
-## Notifications
-
-| Name | Description |
-| ---- | ----------- |
-| `CreateSubscription` |  |
-
-## Pools
-
-| Name | Description |
-| ---- | ----------- |
-| `CreatePool` |  |
-| `UpdatePool` |  |
-| `DeletePool` |  |
-| `ViewPool` |  |
-| `ListPools` |  |
-
-## Projects
-
-| Name | Description |
-| ---- | ----------- |
-| `CreateProject` |  |
-| `DeleteProject` |  |
-| `UpdateProject` |  |
-| `ViewProject` |  |
-
-## Replicators
-
-| Name | Description |
-| ---- | ----------- |
-| `UpdateReplicator` |  |
-| `ViewReplicator` |  |
-
 ## Secrets
 
 | Name | Description |
@@ -794,21 +183,6 @@
 | `DeleteAccount` |  |
 | `ViewAccount` |  |
 
-## Sessions
-
-| Name | Description |
-| ---- | ----------- |
-| `CreateSession` |  |
-| `ViewSession` |  |
-
-## Software
-
-| Name | Description |
-| ---- | ----------- |
-| `UploadSoftware` |  |
-| `DownloadSoftware` |  |
-| `DeleteSoftware` |  |
-
 ## Storage
 
 | Name | Description |
@@ -819,17 +193,6 @@
 | `WriteRefs` | Ability to write refs to the storage service |
 | `DeleteRefs` | Ability to delete refs |
 
-## Streams
-
-| Name | Description |
-| ---- | ----------- |
-| `CreateStream` |  |
-| `UpdateStream` |  |
-| `DeleteStream` |  |
-| `ViewStream` |  |
-| `ViewChanges` |  |
-| `ViewTemplate` |  |
-
 ## Accounts
 
 | Name | Description |
@@ -839,41 +202,6 @@
 | `DeleteAccount` |  |
 | `ViewAccount` |  |
 
-## Agents
-
-| Name | Description |
-| ---- | ----------- |
-| `CreateAgent` |  |
-| `UpdateAgent` |  |
-| `DeleteAgent` |  |
-| `ViewAgent` |  |
-| `ListAgents` |  |
-
-## Artifacts
-
-| Name | Description |
-| ---- | ----------- |
-| `ReadArtifact` |  |
-| `WriteArtifact` |  |
-| `DeleteArtifact` |  |
-| `UploadArtifact` |  |
-| `DownloadArtifact` |  |
-
-## Bisect
-
-| Name | Description |
-| ---- | ----------- |
-| `CreateBisectTask` |  |
-| `UpdateBisectTask` |  |
-| `ViewBisectTask` |  |
-
-## Compute
-
-| Name | Description |
-| ---- | ----------- |
-| `AddComputeTasks` |  |
-| `GetComputeTasks` |  |
-
 ## Ddc
 
 | Name | Description |
@@ -887,42 +215,6 @@
 | `DdcWriteTransactionLog` |  |
 | `DdcAdminAction` |  |
 
-## Devices
-
-| Name | Description |
-| ---- | ----------- |
-| `DeviceRead` |  |
-| `DeviceWrite` |  |
-
-## Jobs
-
-| Name | Description |
-| ---- | ----------- |
-| `CreateJob` |  |
-| `UpdateJob` |  |
-| `DeleteJob` |  |
-| `ExecuteJob` |  |
-| `RetryJobStep` |  |
-| `ViewJob` |  |
-
-## Leases
-
-| Name | Description |
-| ---- | ----------- |
-| `ViewLeases` |  |
-| `ViewLeaseTasks` |  |
-
-## Logs
-
-| Name | Description |
-| ---- | ----------- |
-| `CreateLog` |  |
-| `UpdateLog` |  |
-| `ViewLog` |  |
-| `WriteLogData` |  |
-| `CreateEvent` |  |
-| `ViewEvent` |  |
-
 ## Notices
 
 | Name | Description |
@@ -931,38 +223,6 @@
 | `UpdateNotice` |  |
 | `DeleteNotice` |  |
 
-## Notifications
-
-| Name | Description |
-| ---- | ----------- |
-| `CreateSubscription` |  |
-
-## Pools
-
-| Name | Description |
-| ---- | ----------- |
-| `CreatePool` |  |
-| `UpdatePool` |  |
-| `DeletePool` |  |
-| `ViewPool` |  |
-| `ListPools` |  |
-
-## Projects
-
-| Name | Description |
-| ---- | ----------- |
-| `CreateProject` |  |
-| `DeleteProject` |  |
-| `UpdateProject` |  |
-| `ViewProject` |  |
-
-## Replicators
-
-| Name | Description |
-| ---- | ----------- |
-| `UpdateReplicator` |  |
-| `ViewReplicator` |  |
-
 ## Secrets
 
 | Name | Description |
@@ -978,21 +238,6 @@
 | `DeleteAccount` |  |
 | `ViewAccount` |  |
 
-## Sessions
-
-| Name | Description |
-| ---- | ----------- |
-| `CreateSession` |  |
-| `ViewSession` |  |
-
-## Software
-
-| Name | Description |
-| ---- | ----------- |
-| `UploadSoftware` |  |
-| `DownloadSoftware` |  |
-| `DeleteSoftware` |  |
-
 ## Storage
 
 | Name | Description |
@@ -1003,17 +248,6 @@
 | `WriteRefs` |  |
 | `DeleteRefs` |  |
 
-## Streams
-
-| Name | Description |
-| ---- | ----------- |
-| `CreateStream` |  |
-| `UpdateStream` |  |
-| `DeleteStream` |  |
-| `ViewStream` |  |
-| `ViewChanges` |  |
-| `ViewTemplate` |  |
-
 ## Symbols
 
 | Name | Description |
@@ -1029,41 +263,6 @@
 | `DeleteAccount` |  |
 | `ViewAccount` |  |
 
-## Agents
-
-| Name | Description |
-| ---- | ----------- |
-| `CreateAgent` |  |
-| `UpdateAgent` |  |
-| `DeleteAgent` |  |
-| `ViewAgent` |  |
-| `ListAgents` |  |
-
-## Artifacts
-
-| Name | Description |
-| ---- | ----------- |
-| `ReadArtifact` |  |
-| `WriteArtifact` |  |
-| `DeleteArtifact` |  |
-| `UploadArtifact` |  |
-| `DownloadArtifact` |  |
-
-## Bisect
-
-| Name | Description |
-| ---- | ----------- |
-| `CreateBisectTask` |  |
-| `UpdateBisectTask` |  |
-| `ViewBisectTask` |  |
-
-## Compute
-
-| Name | Description |
-| ---- | ----------- |
-| `AddComputeTasks` |  |
-| `GetComputeTasks` |  |
-
 ## Ddc
 
 | Name | Description |
@@ -1077,42 +276,6 @@
 | `DdcWriteTransactionLog` |  |
 | `DdcAdminAction` |  |
 
-## Devices
-
-| Name | Description |
-| ---- | ----------- |
-| `DeviceRead` |  |
-| `DeviceWrite` |  |
-
-## Jobs
-
-| Name | Description |
-| ---- | ----------- |
-| `CreateJob` |  |
-| `UpdateJob` |  |
-| `DeleteJob` |  |
-| `ExecuteJob` |  |
-| `RetryJobStep` |  |
-| `ViewJob` |  |
-
-## Leases
-
-| Name | Description |
-| ---- | ----------- |
-| `ViewLeases` |  |
-| `ViewLeaseTasks` |  |
-
-## Logs
-
-| Name | Description |
-| ---- | ----------- |
-| `CreateLog` |  |
-| `UpdateLog` |  |
-| `ViewLog` |  |
-| `WriteLogData` |  |
-| `CreateEvent` |  |
-| `ViewEvent` |  |
-
 ## Notices
 
 | Name | Description |
@@ -1121,38 +284,6 @@
 | `UpdateNotice` |  |
 | `DeleteNotice` |  |
 
-## Notifications
-
-| Name | Description |
-| ---- | ----------- |
-| `CreateSubscription` |  |
-
-## Pools
-
-| Name | Description |
-| ---- | ----------- |
-| `CreatePool` |  |
-| `UpdatePool` |  |
-| `DeletePool` |  |
-| `ViewPool` |  |
-| `ListPools` |  |
-
-## Projects
-
-| Name | Description |
-| ---- | ----------- |
-| `CreateProject` |  |
-| `DeleteProject` |  |
-| `UpdateProject` |  |
-| `ViewProject` |  |
-
-## Replicators
-
-| Name | Description |
-| ---- | ----------- |
-| `UpdateReplicator` |  |
-| `ViewReplicator` |  |
-
 ## Secrets
 
 | Name | Description |
@@ -1168,21 +299,6 @@
 | `DeleteAccount` |  |
 | `ViewAccount` |  |
 
-## Sessions
-
-| Name | Description |
-| ---- | ----------- |
-| `CreateSession` |  |
-| `ViewSession` |  |
-
-## Software
-
-| Name | Description |
-| ---- | ----------- |
-| `UploadSoftware` |  |
-| `DownloadSoftware` |  |
-| `DeleteSoftware` |  |
-
 ## Storage
 
 | Name | Description |
@@ -1193,17 +309,6 @@
 | `WriteRefs` |  |
 | `DeleteRefs` |  |
 
-## Streams
-
-| Name | Description |
-| ---- | ----------- |
-| `CreateStream` |  |
-| `UpdateStream` |  |
-| `DeleteStream` |  |
-| `ViewStream` |  |
-| `ViewChanges` |  |
-| `ViewTemplate` |  |
-
 ## Symbols
 
 | Name | Description |
@@ -1226,41 +331,6 @@
 | `DeleteAccount` |  |
 | `ViewAccount` |  |
 
-## Agents
-
-| Name | Description |
-| ---- | ----------- |
-| `CreateAgent` |  |
-| `UpdateAgent` |  |
-| `DeleteAgent` |  |
-| `ViewAgent` |  |
-| `ListAgents` |  |
-
-## Artifacts
-
-| Name | Description |
-| ---- | ----------- |
-| `ReadArtifact` |  |
-| `WriteArtifact` |  |
-| `DeleteArtifact` |  |
-| `UploadArtifact` |  |
-| `DownloadArtifact` |  |
-
-## Bisect
-
-| Name | Description |
-| ---- | ----------- |
-| `CreateBisectTask` |  |
-| `UpdateBisectTask` |  |
-| `ViewBisectTask` |  |
-
-## Compute
-
-| Name | Description |
-| ---- | ----------- |
-| `AddComputeTasks` |  |
-| `GetComputeTasks` |  |
-
 ## Ddc
 
 | Name | Description |
@@ -1274,42 +344,6 @@
 | `DdcWriteTransactionLog` |  |
 | `DdcAdminAction` |  |
 
-## Devices
-
-| Name | Description |
-| ---- | ----------- |
-| `DeviceRead` |  |
-| `DeviceWrite` |  |
-
-## Jobs
-
-| Name | Description |
-| ---- | ----------- |
-| `CreateJob` |  |
-| `UpdateJob` |  |
-| `DeleteJob` |  |
-| `ExecuteJob` |  |
-| `RetryJobStep` |  |
-| `ViewJob` |  |
-
-## Leases
-
-| Name | Description |
-| ---- | ----------- |
-| `ViewLeases` |  |
-| `ViewLeaseTasks` |  |
-
-## Logs
-
-| Name | Description |
-| ---- | ----------- |
-| `CreateLog` |  |
-| `UpdateLog` |  |
-| `ViewLog` |  |
-| `WriteLogData` |  |
-| `CreateEvent` |  |
-| `ViewEvent` |  |
-
 ## Notices
 
 | Name | Description |
@@ -1318,38 +352,6 @@
 | `UpdateNotice` |  |
 | `DeleteNotice` |  |
 
-## Notifications
-
-| Name | Description |
-| ---- | ----------- |
-| `CreateSubscription` |  |
-
-## Pools
-
-| Name | Description |
-| ---- | ----------- |
-| `CreatePool` |  |
-| `UpdatePool` |  |
-| `DeletePool` |  |
-| `ViewPool` |  |
-| `ListPools` |  |
-
-## Projects
-
-| Name | Description |
-| ---- | ----------- |
-| `CreateProject` |  |
-| `DeleteProject` |  |
-| `UpdateProject` |  |
-| `ViewProject` |  |
-
-## Replicators
-
-| Name | Description |
-| ---- | ----------- |
-| `UpdateReplicator` |  |
-| `ViewReplicator` |  |
-
 ## Secrets
 
 | Name | Description |
@@ -1365,21 +367,6 @@
 | `DeleteAccount` |  |
 | `ViewAccount` |  |
 
-## Sessions
-
-| Name | Description |
-| ---- | ----------- |
-| `CreateSession` |  |
-| `ViewSession` |  |
-
-## Software
-
-| Name | Description |
-| ---- | ----------- |
-| `UploadSoftware` |  |
-| `DownloadSoftware` |  |
-| `DeleteSoftware` |  |
-
 ## Storage
 
 | Name | Description |
@@ -1390,17 +377,6 @@
 | `WriteRefs` |  |
 | `DeleteRefs` |  |
 
-## Streams
-
-| Name | Description |
-| ---- | ----------- |
-| `CreateStream` |  |
-| `UpdateStream` |  |
-| `DeleteStream` |  |
-| `ViewStream` |  |
-| `ViewChanges` |  |
-| `ViewTemplate` |  |
-
 ## Symbols
 
 | Name | Description |
diff --git a/Engine/Source/Programs/Horde/Docs/Config/Schema/Globals.md b/Engine/Source/Programs/Horde/Docs/Config/Schema/Globals.md
--- a/Engine/Source/Programs/Horde/Docs/Config/Schema/Globals.md
+++ b/Engine/Source/Programs/Horde/Docs/Config/Schema/Globals.md
@@ -57,8 +57,6 @@ Name | Description
 Name | Description
 ---- | -----------
 `analytics` | [AnalyticsConfig](#analyticsconfig)<br>Configuration for the analytics plugin
-`build` | [BuildConfig](#buildconfig)<br>Configuration for the build plugin
-`compute` | [ComputeConfig](#computeconfig)<br>Configuration for the compute plugin
 `ddc` | [EmptyPluginConfig](#emptypluginconfig)<br>Configuration for the ddc plugin
 `secrets` | [SecretsConfig](#secretsconfig)<br>Configuration for the secrets plugin
 `storage` | [StorageConfig](#storageconfig)<br>Configuration for the storage plugin
@@ -73,122 +71,49 @@ Name | Description
 ---- | -----------
 `stores` | [TelemetryStoreConfig](Telemetry.md)`[]`<br>Metrics to aggregate on the Horde server
 
-## BuildConfig
-
-Configuration for the build plugin
-
-Name | Description
----- | -----------
-`perforceClusters` | [PerforceCluster](#perforcecluster)`[]`<br>List of Perforce clusters
-`devices` | [DeviceConfig](#deviceconfig)<br>Device configuration
-`maxConformCount` | `integer`<br>Maximum number of conforms to run at once
-`agentShutdownIfDisabledGracePeriod` | `string`<br>Time to wait before shutting down an agent that has been disabled Used if no value is set on the actual pool.
-`artifactTypes` | [ArtifactTypeConfig](#artifacttypeconfig)`[]`<br>Configuration for different artifact types
-`projects` | [ProjectConfig](Projects.md)`[]`<br>List of projects
-`enableConformTasks` | `boolean`<br>Whether to allow conform tasks to run
-`issueFixedTag` | `string`<br>Commit tag to use for marking issues as fixed
-
-## PerforceCluster
-
-Information about a cluster of Perforce servers.
-
-Name | Description
----- | -----------
-`name` | `string`<br>Name of the cluster
-`serviceAccount` | `string`<br>Username for Horde to log in to this server. Will use the first account specified below if not overridden.
-`canImpersonate` | `boolean`<br>Whether the service account can impersonate other users
-`supportsPartitionedWorkspaces` | `boolean`<br>Whether to use partitioned workspaces on this server
-`servers` | [PerforceServer](#perforceserver)`[]`<br>List of servers
-`credentials` | [PerforceCredentials](#perforcecredentials)`[]`<br>List of server credentials
-`autoSdk` | [AutoSdkWorkspace](#autosdkworkspace)`[]`<br>List of autosdk streams
-
-## PerforceServer
-
-Information about an individual Perforce server
-
-Name | Description
----- | -----------
-`serverAndPort` | `string`<br>The server and port. The server may be a DNS entry with multiple records, in which case it will be actively load balanced. If "ssl:" prefix is used, ensure P4 server's fingerprint/certificate is trusted. See Horde's documentation on connecting to SSL-enabled Perforce servers.
-`healthCheck` | `boolean`<br>Whether to query the healthcheck address under each server
-`resolveDns` | `boolean`<br>Whether to resolve the DNS entries and load balance between different hosts
-`maxConformCount` | `integer`<br>Maximum number of simultaneous conforms on this server
-`condition` | `string`<br>Optional condition for a machine to be eligible to use this server
-`properties` | `string[]`<br>List of properties for an agent to be eligible to use this server
-
-## PerforceCredentials
-
-Credentials for a Perforce user
-
-Name | Description
----- | -----------
-`userName` | `string`<br>The username
-`password` | `string`<br>Password for the user
-`ticket` | `string`<br>Login ticket for the user (will be used instead of password if set)
-
-## AutoSdkWorkspace
-
-Path to a platform and stream to use for syncing AutoSDK
-
-Name | Description
----- | -----------
-`name` | `string`<br>Name of this workspace
-`properties` | `string[]`<br>The agent properties to check (eg. "OSFamily=Windows")
-`userName` | `string`<br>Username for logging in to the server
-`stream` | `string`<br>Stream to use
+## EmptyPluginConfig
 
-## DeviceConfig
+Empty implementation of
 
-Configuration for devices
 
-Name | Description
----- | -----------
-`platforms` | [DevicePlatformConfig](#deviceplatformconfig)`[]`<br>List of device platforms
-`pools` | [DevicePoolConfig](#devicepoolconfig)`[]`<br>List of device pools
-
-## DevicePlatformConfig
+## SecretsConfig
 
-Configuration for a device platform
+Configuration for the secrets system
 
 Name | Description
 ---- | -----------
-`id` | `string`<br>The id for this platform
-`name` | `string`<br>Name of the platform
-`models` | `string[]`<br>A list of platform models
-`legacyNames` | `string[]`<br>Legacy names which older versions of Gauntlet may be using
-`legacyPerfSpecHighModel` | `string`<br>Model name for the high perf spec, which may be requested by Gauntlet
+`secrets` | [SecretConfig](#secretconfig)`[]`<br>List of secrets
 
-## DevicePoolConfig
+## SecretConfig
 
-Configuration for a device pool
+Configuration for a secret value
 
 Name | Description
 ---- | -----------
-`id` | `string`<br>The id for this platform
-`name` | `string`<br>The name of the pool
-`poolType` | [DevicePoolType](#devicepooltype-enum)<br>The type of the pool
-`projectIds` | `string[]`<br>List of project ids associated with pool
+`id` | `string`<br>Identifier for this secret
+`data` | `string` `->` `string`<br>Key/value pairs associated with this secret
+`sources` | [ExternalSecretConfig](#externalsecretconfig)`[]`<br>Providers to source key/value pairs from
+`acl` | [AclConfig](#aclconfig)<br>Defines access to this particular secret
 
-## DevicePoolType (Enum)
+## ExternalSecretConfig
 
-The type of device pool
+Configuration for an external secret provider
 
 Name | Description
 ---- | -----------
-`Automation` | Available to CIS jobs
-`Shared` | Shared by users with remote checking and checkouts
+`provider` | `string`<br>Name of the provider to use
+`format` | [ExternalSecretFormat](#externalsecretformat-enum)<br>Format of the secret
+`key` | `string`<br>Optional key indicating the parameter to set in the resulting data array. Required if if  is .
+`path` | `string`<br>Optional value indicating what to fetch from the provider
 
-## ArtifactTypeConfig
+## ExternalSecretFormat (Enum)
 
-Configuration for an artifact
+Format describing how to parse external secret values
 
 Name | Description
 ---- | -----------
-`name` | `string`<br>Legacy 'Name' property
-`type` | `string`<br>Name of the artifact type
-`acl` | [AclConfig](#aclconfig)<br>Acl for the artifact type
-`keepCount` | `integer`<br>Number of artifacts to retain
-`keepDays` | `integer`<br>Number of days to retain artifacts of this type
-`namespaceId` | `string`<br>Storage namespace to use for this artifact types
+`Text` | Secret is a plain text value which will be stored using the external secret key
+`Json` | Secret is a JSON formatted string containing key/value pairs
 
 ## AclConfig
 
@@ -231,234 +156,6 @@ Name | Description
 `excludeActions` | `string[]`<br>Actions to exclude from the inherited actions
 `extends` | `string[]`<br>Other profiles to extend from
 
-## ComputeConfig
-
-Configuration for the compute system
-
-Name | Description
----- | -----------
-`acl` | [AclConfig](#aclconfig)<br>Inherited root acl
-`versionEnum` | [ConfigVersion](#configversion-enum)<br>Config version number
-`rates` | [AgentRateConfig](#agentrateconfig)`[]`<br>List of costs of a particular agent type
-`clusters` | [ComputeClusterConfig](#computeclusterconfig)`[]`<br>List of compute profiles
-`pools` | [PoolConfig](#poolconfig)`[]`<br>List of pools
-`software` | [AgentSoftwareConfig](#agentsoftwareconfig)`[]`<br>List of costs of a particular agent type
-`networks` | [NetworkConfig](#networkconfig)`[]`<br>List of networks
-
-## ConfigVersion (Enum)
-
-Global version number for running the server. As new features are introduced that require data migrations, this version number indicates the backwards compatibility functionality that must be enabled. When adding a new version here, also add a message to ConfigService.CreateSnapshotAsync describing the steps that need to be taken to upgrade the deployment.
-
-Name | Description
----- | -----------
-`None` | Not specified
-`Initial` | Initial version number
-`PoolsInConfigFiles` | Ability to add/remove pools via the REST API is removed. Pools should be configured through globals.json instead.
-`Latest` | Latest version number
-`LatestPlusOne` | One after the l`ast defined version number
-
-## AgentRateConfig
-
-Describes the monetary cost of agents matching a particular criteria
-
-Name | Description
----- | -----------
-`condition` | `string`<br>Condition string
-`rate` | `number`<br>Rate for this agent
-
-## ComputeClusterConfig
-
-Profile for executing compute requests
-
-Name | Description
----- | -----------
-`id` | `string`<br>Name of the partition
-`namespaceId` | `string`<br>Name of the namespace to use
-`requestBucketId` | `string`<br>Name of the input bucket
-`responseBucketId` | `string`<br>Name of the output bucket
-`condition` | `string`<br>Filter for agents to include
-`acl` | [AclConfig](#aclconfig)<br>Access control list
-
-## PoolConfig
-
-Mutable configuration for a pool
-
-Name | Description
----- | -----------
-`id` | `string`<br>Unique id for this pool
-`base` | `string`<br>Base pool config to copy settings from
-`name` | `string`<br>Name of the pool
-`condition` | `string`<br>Condition for agents to automatically be included in this pool
-`properties` | `string` `->` `string`<br>Arbitrary properties related to this pool
-`color` | [PoolColor](#poolcolor-enum)<br>Color to use for this pool on the dashboard
-`enableAutoscaling` | `boolean`<br>Whether to enable autoscaling for this pool
-`minAgents` | `integer`<br>The minimum number of agents to keep in the pool
-`numReserveAgents` | `integer`<br>The minimum number of idle agents to hold in reserve
-`conformInterval` | `string`<br>Interval between conforms. If zero, the pool will not conform on a schedule.
-`scaleOutCooldown` | `string`<br>Cooldown time between scale-out events
-`scaleInCooldown` | `string`<br>Cooldown time between scale-in events
-`shutdownIfDisabledGracePeriod` | `string`<br>Time to wait before shutting down an agent that has been disabled
-`sizeStrategy` | [PoolSizeStrategy](#poolsizestrategy-enum)<br>
-`sizeStrategies` | [PoolSizeStrategyInfo](#poolsizestrategyinfo)`[]`<br>List of pool sizing strategies for this pool. The first strategy with a matching condition will be picked.
-`fleetManagers` | [FleetManagerInfo](#fleetmanagerinfo)`[]`<br>List of fleet managers for this pool. The first strategy with a matching condition will be picked. If empty or no conditions match, a default fleet manager will be used.
-`leaseUtilizationSettings` | [LeaseUtilizationSettings](#leaseutilizationsettings)<br>Settings for lease utilization pool sizing strategy (if used)
-`jobQueueSettings` | [JobQueueSettings](#jobqueuesettings)<br>Settings for job queue pool sizing strategy (if used)
-`computeQueueAwsMetricSettings` | [ComputeQueueAwsMetricSettings](#computequeueawsmetricsettings)<br>Settings for job queue pool sizing strategy (if used)
-
-## PoolColor (Enum)
-
-Color to use for labels of this pool
-
-Name | Description
----- | -----------
-`Default` | 
-`Blue` | 
-`Orange` | 
-`Green` | 
-`Gray` | 
-
-## PoolSizeStrategy (Enum)
-
-Available pool sizing strategies
-
-Name | Description
----- | -----------
-`LeaseUtilization` | Strategy based on lease utilization
-`JobQueue` | Strategy based on size of job build queue
-`NoOp` | No-op strategy used as fallback/default behavior
-`ComputeQueueAwsMetric` | A no-op strategy that reports metrics to let an external AWS auto-scaling policy scale the fleet
-`LeaseUtilizationAwsMetric` | A no-op strategy that reports metrics to let an external AWS auto-scaling policy scale the fleet
-
-## PoolSizeStrategyInfo
-
-Metadata for configuring and picking a pool sizing strategy
-
-Name | Description
----- | -----------
-`type` | [PoolSizeStrategy](#poolsizestrategy-enum)<br>Strategy implementation to use
-`condition` | `string`<br>Condition if this strategy should be enabled (right now, using date/time as a distinguishing factor)
-`config` | `object`<br>Configuration for the strategy, serialized as JSON
-`extraAgentCount` | `integer`<br>Integer to add after pool size has been calculated. Can also be negative.
-
-## FleetManagerInfo
-
-Metadata for configuring and picking a fleet manager
-
-Name | Description
----- | -----------
-`type` | [FleetManagerType](#fleetmanagertype-enum)<br>Fleet manager type implementation to use
-`condition` | `string`<br>Condition if this strategy should be enabled (right now, using date/time as a distinguishing factor)
-`config` | `object`<br>Configuration for the strategy, serialized as JSON
-
-## FleetManagerType (Enum)
-
-Available fleet managers
-
-Name | Description
----- | -----------
-`Default` | Default fleet manager
-`NoOp` | No-op fleet manager.
-`Aws` | Fleet manager for handling AWS EC2 instances. Will create and/or terminate instances from scratch.
-`AwsReuse` | Fleet manager for handling AWS EC2 instances. Will start already existing but stopped instances to reuse existing EBS disks.
-`AwsRecycle` | Fleet manager for handling AWS EC2 instances. Will start already existing but stopped instances to reuse existing EBS disks.
-`AwsAsg` | Fleet manager for handling AWS EC2 instances. Uses an EC2 auto-scaling group for controlling the number of running instances.
-
-## LeaseUtilizationSettings
-
-Lease utilization sizing settings for a pool
-
-Name | Description
----- | -----------
-`sampleTimeSec` | `integer`<br>Time period for each sample
-`numSamples` | `integer`<br>Number of samples to collect for calculating lease utilization
-`numSamplesForResult` | `integer`<br>Min number of samples for a valid result
-`minAgents` | `integer`<br>The minimum number of agents to keep in the pool
-`numReserveAgents` | `integer`<br>The minimum number of idle agents to hold in reserve
-
-## JobQueueSettings
-
-Job queue sizing settings for a pool
-
-Name | Description
----- | -----------
-`scaleOutFactor` | `number`<br>Factor translating queue size to additional agents to grow the pool with The result is always rounded up to nearest integer. Example: if there are 20 jobs in queue, a factor 0.25 will result in 5 new agents being added (20 * 0.25)
-`scaleInFactor` | `number`<br>Factor by which to shrink the pool size with when queue is empty The result is always rounded up to nearest integer. Example: when the queue size is zero, a default value of 0.9 will shrink the pool by 10% (current agent count * 0.9)
-`samplePeriodMin` | `integer`<br>How far back in time to look for job batches (that potentially are in the queue)
-`readyTimeThresholdSec` | `integer`<br>Time spent in ready state before considered truly waiting for an agent<br>A job batch can be in ready state before getting picked up and executed. This threshold will help ensure only batches that have been waiting longer than this value will be considered.
-
-## ComputeQueueAwsMetricSettings
-
-Settings for
-
-Name | Description
----- | -----------
-`computeClusterId` | `string`<br>Compute cluster ID to observe
-`namespace` | `string`<br>AWS CloudWatch namespace to write metrics in
-
-## AgentSoftwareConfig
-
-Selects different agent software versions by evaluating a condition
-
-Name | Description
----- | -----------
-`toolId` | `string`<br>Tool identifier
-`condition` | `string`<br>Condition for using this channel
-
-## NetworkConfig
-
-Describes a network The ID describes any logical grouping, such as region, availability zone, rack or office location.
-
-Name | Description
----- | -----------
-`id` | `string`<br>ID for this network
-`cidrBlock` | `string`<br>CIDR block
-`description` | `string`<br>Human-readable description
-`computeId` | `string`<br>Compute ID for this network (used when allocating compute resources)
-
-## EmptyPluginConfig
-
-Empty implementation of
-
-
-## SecretsConfig
-
-Configuration for the secrets system
-
-Name | Description
----- | -----------
-`secrets` | [SecretConfig](#secretconfig)`[]`<br>List of secrets
-
-## SecretConfig
-
-Configuration for a secret value
-
-Name | Description
----- | -----------
-`id` | `string`<br>Identifier for this secret
-`data` | `string` `->` `string`<br>Key/value pairs associated with this secret
-`sources` | [ExternalSecretConfig](#externalsecretconfig)`[]`<br>Providers to source key/value pairs from
-`acl` | [AclConfig](#aclconfig)<br>Defines access to this particular secret
-
-## ExternalSecretConfig
-
-Configuration for an external secret provider
-
-Name | Description
----- | -----------
-`provider` | `string`<br>Name of the provider to use
-`format` | [ExternalSecretFormat](#externalsecretformat-enum)<br>Format of the secret
-`key` | `string`<br>Optional key indicating the parameter to set in the resulting data array. Required if if  is .
-`path` | `string`<br>Optional value indicating what to fetch from the provider
-
-## ExternalSecretFormat (Enum)
-
-Format describing how to parse external secret values
-
-Name | Description
----- | -----------
-`Text` | Secret is a plain text value which will be stored using the external secret key
-`Json` | Secret is a JSON formatted string containing key/value pairs
-
 ## StorageConfig
 
 Configuration for storage
diff --git a/Engine/Source/Programs/Horde/Docs/Deployment/ServerSettings.md b/Engine/Source/Programs/Horde/Docs/Deployment/ServerSettings.md
--- a/Engine/Source/Programs/Horde/Docs/Deployment/ServerSettings.md
+++ b/Engine/Source/Programs/Horde/Docs/Deployment/ServerSettings.md
@@ -124,9 +124,10 @@ Name | Description
 Name | Description
 ---- | -----------
 `analytics` | [AnalyticsServerConfig](#analyticsserverconfig)<br>Configuration for the analytics plugin
-`build` | [StaticBuildConfig](#staticbuildconfig)<br>Configuration for the build plugin
-`compute` | [StaticComputeConfig](#staticcomputeconfig)<br>Configuration for the compute plugin
-`storage` | [StaticStorageConfig](#staticstorageconfig)<br>Configuration for the storage plugin
+`ddc` | [PluginServerConfig](#pluginserverconfig)<br>Configuration for the ddc plugin
+`secrets` | [PluginServerConfig](#pluginserverconfig)<br>Configuration for the secrets plugin
+`storage` | [StorageServerConfig](#storageserverconfig)<br>Configuration for the storage plugin
+`symbols` | [PluginServerConfig](#pluginserverconfig)<br>Configuration for the symbols plugin
 `tools` | [ToolsServerConfig](#toolsserverconfig)<br>Configuration for the tools plugin
 
 ## AnalyticsServerConfig
@@ -136,6 +137,7 @@ Server configuration for the analytics system
 Name | Description
 ---- | -----------
 `sinks` | [TelemetrySinkConfig](#telemetrysinkconfig)<br>Settings for the various telemetry sinks
+`enabled` | `boolean`<br>Whether the plugin should be enabled or not
 
 ## TelemetrySinkConfig
 
@@ -165,188 +167,15 @@ Name | Description
 `retainDays` | `number`<br>Number of days worth of telmetry events to keep
 `enabled` | `boolean`<br>Whether to enable this sink
 
-## StaticBuildConfig
+## PluginServerConfig
 
-Static configuration for the build plugin
+Base class for plugin server config objects
 
 Name | Description
 ---- | -----------
-`perforce` | [PerforceConnectionSettings](#perforceconnectionsettings)`[]`<br>Perforce connections for use by the Horde server (not agents)
-`useLocalPerforceEnv` | `boolean`<br>Whether to use the local Perforce environment
-`perforceConnectionPoolSize` | `integer`<br>Number of pooled perforce connections to keep
-`enableConformTasks` | `boolean`<br>Whether to enable the conform task source.
-`p4SwarmUrl` | `string`<br>Url of P4 Swarm installation
-`jiraUsername` | `string`<br>The Jira service account user name
-`jiraApiToken` | `string`<br>The Jira service account API token
-`jiraUrl` | `string`<br>The Uri for the Jira installation
-`sharedDeviceCheckoutDays` | `integer`<br>The number of days shared device checkouts are held
-`deviceProblemCooldownMinutes` | `integer`<br>The number of cooldown minutes for device problems
-`deviceReportChannel` | `string`<br>Channel to send device reports to
-`disableSchedules` | `boolean`<br>Whether to run scheduled jobs.
-`slackToken` | `string`<br>Bot token for interacting with Slack (xoxb-*)
-`slackSocketToken` | `string`<br>Token for opening a socket to slack (xapp-*)
-`slackAdminToken` | `string`<br>Admin user token for Slack (xoxp-*). This is only required when using the admin endpoints to invite users.
-`slackUsers` | `string`<br>Filtered list of slack users to send notifications to. Should be Slack user ids, separated by commas.
-`slackErrorPrefix` | `string`<br>Prefix to use when reporting errors
-`slackWarningPrefix` | `string`<br>Prefix to use when reporting warnings
-`configNotificationChannel` | `string`<br>Channel for sending messages related to config update failures
-`updateStreamsNotificationChannel` | `string`<br>Channel to send stream notification update failures to
-`jobNotificationChannel` | `string`<br>Slack channel to send job related notifications to. Multiple channels can be specified, separated by ;
-`agentNotificationChannel` | `string`<br>Slack channel to send agent related notifications to.
-`testDataRetainMonths` | `integer`<br>The number of months to retain test data
-`blockCacheDir` | `string`<br>Directory to store the fine-grained block cache. This caches individual exports embedded in bundles.
-`blockCacheSize` | `string`<br>Maximum size of the block cache. Accepts standard binary suffixes. Currently only allocates in multiples of 1024mb.
-`blockCacheSizeBytes` | `integer`<br>Accessor for the block cache size in bytes
-`commits` | [CommitSettings](#commitsettings)<br>Options for the commit service
-
-## PerforceConnectionSettings
-
-Perforce connection information for use by the Horde server (for reading config files, etc...)
+`enabled` | `boolean`<br>Whether the plugin should be enabled or not
 
-Name | Description
----- | -----------
-`id` | `string`<br>Identifier for this server
-`serverAndPort` | `string`<br>Server and port
-`credentials` | [PerforceCredentials](#perforcecredentials)<br>Credentials for the server
-
-## PerforceCredentials
-
-Credentials for a Perforce user
-
-Name | Description
----- | -----------
-`userName` | `string`<br>The username
-`password` | `string`<br>Password for the user
-`ticket` | `string`<br>Login ticket for the user (will be used instead of password if set)
-
-## CommitSettings
-
-Options for the commit service
-
-Name | Description
----- | -----------
-`replicateMetadata` | `boolean`<br>Whether to mirror commit metadata to the database
-`replicateContent` | `boolean`<br>Whether to mirror commit data to storage
-`bundle` | [BundleOptions](#bundleoptions)<br>Options for how objects are packed together
-`chunking` | [ChunkingOptions](#chunkingoptions)<br>Options for how objects are sliced
-
-## BundleOptions
-
-Options for configuring a bundle serializer
-
-Name | Description
----- | -----------
-`maxVersion` | [BundleVersion](#bundleversion-enum)<br>Maximum version number of bundles to write
-`maxBlobSize` | `integer`<br>Maximum payload size fo a blob
-`compressionFormat` | [BundleCompressionFormat](#bundlecompressionformat-enum)<br>Compression format to use
-`minCompressionPacketSize` | `integer`<br>Minimum size of a block to be compressed
-`maxWriteQueueLength` | `integer`<br>Maximum amount of data to store in memory. This includes any background writes as well as bundles being built.
-
-## BundleVersion (Enum)
-
-Bundle version number
-
-Name | Description
----- | -----------
-`Initial` | Initial version number
-`ExportAliases` | Added the BundleExport.Alias property
-`RemoveAliases` | Back out change to include aliases. Will likely do this through an API rather than baked into the data.
-`InPlace` | Use data structures which support in-place reading and writing.
-`ImportHashes` | Add import hashes to imported nodes
-`LatestV1` | Last version using the V1 pipeline
-`PacketSequence` | Structure bundles as a sequence of self-contained packets (uses V2 code)
-`Latest` | The current version number
-`LatestV2` | Last version using the V2 pipeline
-`LatestPlusOne` | Last item in the enum. Used for
-
-## BundleCompressionFormat (Enum)
-
-Indicates the compression format in the bundle
-
-Name | Description
----- | -----------
-`None` | Packets are uncompressed
-`LZ4` | LZ4 compression
-`Gzip` | Gzip compression
-`Oodle` | Oodle compression (Selkie)
-`Brotli` | Brotli compression
-`Zstd` | ZStandard compression
-
-## ChunkingOptions
-
-Options for creating file nodes
-
-Name | Description
----- | -----------
-`leafOptions` | [LeafChunkedDataNodeOptions](#leafchunkeddatanodeoptions)<br>Options for creating leaf nodes
-`interiorOptions` | [InteriorChunkedDataNodeOptions](#interiorchunkeddatanodeoptions)<br>Options for creating interior nodes
-
-## LeafChunkedDataNodeOptions
-
-Options for creating a specific type of file nodes
-
-Name | Description
----- | -----------
-`minSize` | `integer`<br>Minimum chunk size
-`maxSize` | `integer`<br>Maximum chunk size. Chunks will be split on this boundary if another match is not found.
-`targetSize` | `integer`<br>Target chunk size for content-slicing
-
-## InteriorChunkedDataNodeOptions
-
-Options for creating interior nodes
-
-Name | Description
----- | -----------
-`minChildCount` | `integer`<br>Minimum number of children in each node
-`targetChildCount` | `integer`<br>Target number of children in each node
-`maxChildCount` | `integer`<br>Maximum number of children in each node
-`sliceThreshold` | `integer`<br>Threshold hash value for splitting interior nodes
-
-## StaticComputeConfig
-
-Static configuration for the compute plugin
-
-Name | Description
----- | -----------
-`enableUpgradeTasks` | `boolean`<br>Whether to enable the upgrade task source.
-`withAws` | `boolean`<br>Whether to enable Amazon Web Services (AWS) specific features
-`awsRegions` | `string[]`<br>List of AWS regions for Horde to be aware of (e.g. us-east-1 or eu-central-1) Right now, this is only used for replicating CloudWatch metrics to multiple regions
-`awsAutoScalingQueueUrls` | `string[]`<br>AWS SQS queue URLs where lifecycle events from EC2 auto-scaling are received
-`fleetManagerV2` | [FleetManagerType](#fleetmanagertype-enum)<br>Default fleet manager to use (when not specified by pool)
-`fleetManagerV2Config` | `object`<br>Config for the fleet manager (serialized JSON)
-`autoEnrollAgents` | `boolean`<br>Whether to automatically enroll agents in the farm
-`defaultAgentPoolSizeStrategy` | [PoolSizeStrategy](#poolsizestrategy-enum)<br>Default agent pool sizing strategy for pools that doesn't have one explicitly configured
-`agentPoolScaleOutCooldownSeconds` | `integer`<br>Scale-out cooldown for auto-scaling agent pools (in seconds). Can be overridden by per-pool settings.
-`agentPoolScaleInCooldownSeconds` | `integer`<br>Scale-in cooldown for auto-scaling agent pools (in seconds). Can be overridden by per-pool settings.
-`computeTunnelPort` | `integer`<br>Port to listen on for tunneling compute sockets to agents
-`computeTunnelAddress` | `string`<br>What address (host:port) clients should connect to for compute socket tunneling Port may differ from  if Horde server is behind a reverse proxy/firewall
-
-## FleetManagerType (Enum)
-
-Available fleet managers
-
-Name | Description
----- | -----------
-`Default` | Default fleet manager
-`NoOp` | No-op fleet manager.
-`Aws` | Fleet manager for handling AWS EC2 instances. Will create and/or terminate instances from scratch.
-`AwsReuse` | Fleet manager for handling AWS EC2 instances. Will start already existing but stopped instances to reuse existing EBS disks.
-`AwsRecycle` | Fleet manager for handling AWS EC2 instances. Will start already existing but stopped instances to reuse existing EBS disks.
-`AwsAsg` | Fleet manager for handling AWS EC2 instances. Uses an EC2 auto-scaling group for controlling the number of running instances.
-
-## PoolSizeStrategy (Enum)
-
-Available pool sizing strategies
-
-Name | Description
----- | -----------
-`LeaseUtilization` | Strategy based on lease utilization
-`JobQueue` | Strategy based on size of job build queue
-`NoOp` | No-op strategy used as fallback/default behavior
-`ComputeQueueAwsMetric` | A no-op strategy that reports metrics to let an external AWS auto-scaling policy scale the fleet
-`LeaseUtilizationAwsMetric` | A no-op strategy that reports metrics to let an external AWS auto-scaling policy scale the fleet
-
-## StaticStorageConfig
+## StorageServerConfig
 
 Static settings for the storage system
 
@@ -356,6 +185,7 @@ Name | Description
 `bundleCacheSize` | `string`<br>Maximum size of the storage cache on disk. Accepts standard binary suffixes (kb, mb, gb, tb, etc...)
 `bundleCacheSizeBytes` | `integer`<br>Accessor for the bundle cache size in bytes
 `backends` | [BackendConfig](#backendconfig)`[]`<br>Overridden settings for storage backends. Useful for running against a production server with custom backends.
+`enabled` | `boolean`<br>Whether the plugin should be enabled or not
 
 ## BackendConfig
 
@@ -408,6 +238,7 @@ Server configuration for bundled tools
 Name | Description
 ---- | -----------
 `bundledTools` | [BundledToolConfig](#bundledtoolconfig)`[]`<br>Tools bundled along with the server. Data for each tool can be produced using the 'bundle create' command, and should be stored in the Tools directory.
+`enabled` | `boolean`<br>Whether the plugin should be enabled or not
 
 ## BundledToolConfig
 
