Horde: Upgrade System.Text.Json to 8.0.5 to address security warnings

diff --git a/Engine/Source/Programs/Horde/Drivers/JobDriver/JobDriver.csproj b/Engine/Source/Programs/Horde/Drivers/JobDriver/JobDriver.csproj
--- a/Engine/Source/Programs/Horde/Drivers/JobDriver/JobDriver.csproj
+++ b/Engine/Source/Programs/Horde/Drivers/JobDriver/JobDriver.csproj
@@ -92,7 +92,7 @@
     <PackageReference Include="Serilog.Sinks.File" Version="5.0.0" />
     <PackageReference Include="System.Diagnostics.PerformanceCounter" Version="8.0.0" />
     <PackageReference Include="System.ServiceProcess.ServiceController" Version="8.0.0" />
-    <PackageReference Include="System.Text.Json" Version="8.0.4" />
+    <PackageReference Include="System.Text.Json" Version="8.0.5" />
   </ItemGroup>
   
   <ItemGroup>
diff --git a/Engine/Source/Programs/Horde/HordeAgent/HordeAgent.csproj b/Engine/Source/Programs/Horde/HordeAgent/HordeAgent.csproj
--- a/Engine/Source/Programs/Horde/HordeAgent/HordeAgent.csproj
+++ b/Engine/Source/Programs/Horde/HordeAgent/HordeAgent.csproj
@@ -73,7 +73,7 @@
     <PackageReference Include="Serilog.Sinks.File" Version="5.0.0" />
     <PackageReference Include="System.Diagnostics.PerformanceCounter" Version="8.0.0" />
     <PackageReference Include="System.ServiceProcess.ServiceController" Version="8.0.0" />
-    <PackageReference Include="System.Text.Json" Version="8.0.4" />
+    <PackageReference Include="System.Text.Json" Version="8.0.5" />
   </ItemGroup>
 
   <ItemGroup>
diff --git a/Engine/Source/Programs/Horde/HordeCmd/HordeCmd.csproj b/Engine/Source/Programs/Horde/HordeCmd/HordeCmd.csproj
--- a/Engine/Source/Programs/Horde/HordeCmd/HordeCmd.csproj
+++ b/Engine/Source/Programs/Horde/HordeCmd/HordeCmd.csproj
@@ -32,7 +32,7 @@
     <PackageReference Include="Serilog.Settings.Configuration" Version="8.0.0" />
     <PackageReference Include="Serilog.Sinks.Console" Version="5.0.1" />
     <PackageReference Include="Serilog.Sinks.File" Version="5.0.0" />
-    <PackageReference Include="System.Text.Json" Version="8.0.4" />
+    <PackageReference Include="System.Text.Json" Version="8.0.5" />
   </ItemGroup>
 
   <ItemGroup>
diff --git a/Engine/Source/Programs/Horde/HordeServer/HordeServer.csproj b/Engine/Source/Programs/Horde/HordeServer/HordeServer.csproj
--- a/Engine/Source/Programs/Horde/HordeServer/HordeServer.csproj
+++ b/Engine/Source/Programs/Horde/HordeServer/HordeServer.csproj
@@ -79,7 +79,7 @@
     <PackageReference Include="StackExchange.Redis" Version="2.7.4" />
     <PackageReference Include="Swashbuckle.AspNetCore" Version="6.5.0" />
     <PackageReference Include="System.Linq.Async" Version="6.0.1" />
-    <PackageReference Include="System.Text.Json" Version="8.0.4" />
+    <PackageReference Include="System.Text.Json" Version="8.0.5" />
     <PackageReference Include="T-Digest.NET" Version="1.0.3" />
     <PackageReference Include="TimeZoneConverter" Version="6.1.0" />
     <PackageReference Include="ZstdSharp.Port" Version="0.8.1" />
diff --git a/Engine/Source/Programs/Horde/Tools/ReleaseNotes/ReleaseNotes.csproj b/Engine/Source/Programs/Horde/Tools/ReleaseNotes/ReleaseNotes.csproj
--- a/Engine/Source/Programs/Horde/Tools/ReleaseNotes/ReleaseNotes.csproj
+++ b/Engine/Source/Programs/Horde/Tools/ReleaseNotes/ReleaseNotes.csproj
@@ -10,7 +10,7 @@
 
   <ItemGroup>
     <PackageReference Include="Microsoft.Extensions.Logging.Console" Version="8.0.0" />
-    <PackageReference Include="System.Text.Json" Version="8.0.4" />
+    <PackageReference Include="System.Text.Json" Version="8.0.5" />
   </ItemGroup>
 
   <ItemGroup>
diff --git a/Engine/Source/Programs/Shared/EpicGames.Perforce.Fixture/EpicGames.Perforce.Fixture.csproj b/Engine/Source/Programs/Shared/EpicGames.Perforce.Fixture/EpicGames.Perforce.Fixture.csproj
--- a/Engine/Source/Programs/Shared/EpicGames.Perforce.Fixture/EpicGames.Perforce.Fixture.csproj
+++ b/Engine/Source/Programs/Shared/EpicGames.Perforce.Fixture/EpicGames.Perforce.Fixture.csproj
@@ -42,7 +42,7 @@
       <PrivateAssets>all</PrivateAssets>
       <IncludeAssets>runtime; build; native; contentfiles; analyzers; buildtransitive</IncludeAssets>
     </PackageReference>
-    <PackageReference Include="System.Text.Json" Version="8.0.4" />
+    <PackageReference Include="System.Text.Json" Version="8.0.5" />
   </ItemGroup>
 
   <ItemGroup>
