Horde: Accoount for a dot in the mongodb conf/log paths
#rnx

diff --git a/Engine/Source/Programs/Horde/HordeServer/Server/MongoService.cs b/Engine/Source/Programs/Horde/HordeServer/Server/MongoService.cs
--- a/Engine/Source/Programs/Horde/HordeServer/Server/MongoService.cs
+++ b/Engine/Source/Programs/Horde/HordeServer/Server/MongoService.cs
@@ -340,12 +340,12 @@ namespace HordeServer.Server
 			FileReference configFile = FileReference.Combine(mongoDir, "mongod.conf");
 
 			// Check mongo log and config files for non-latin characters, mongodb.exe cannout read these paths
-			if (Regex.IsMatch(mongoLogFile.ToString(), "[^a-zA-Z0-9:\\\\\\/]"))
+			if (Regex.IsMatch(mongoLogFile.ToString(), "[^a-zA-Z0-9:\\\\\\/.]"))
 			{
 				logger.LogError("MongoDB log file contains non-latin characters, which mongodb.exe cannot read: {MongoLogFile}", mongoLogFile.ToString());
 			}
 
-			if (Regex.IsMatch(configFile.ToString(), "[^a-zA-Z0-9:\\\\\\/]"))
+			if (Regex.IsMatch(configFile.ToString(), "[^a-zA-Z0-9:\\\\\\/.]"))
 			{
 				logger.LogError("MongoDB config file contains non-latin characters, which mongodb.exe cannot read: {MongoConfigFile}", configFile.ToString());
 			}
