Horde: Fix compiler warnings

diff --git a/Engine/Source/Programs/Horde/HordeServer/ServerApp.cs b/Engine/Source/Programs/Horde/HordeServer/ServerApp.cs
--- a/Engine/Source/Programs/Horde/HordeServer/ServerApp.cs
+++ b/Engine/Source/Programs/Horde/HordeServer/ServerApp.cs
@@ -7,7 +7,6 @@ using System.IO;
 using System.Linq;
 using System.Reflection;
 using System.Runtime.InteropServices;
-using System.Text.RegularExpressions;
 using System.Threading;
 using System.Threading.Tasks;
 using EpicGames.Core;
@@ -193,7 +192,9 @@ namespace HordeServer
 				.ReadFrom.Configuration(config)
 				.CreateLogger();
 
+#pragma warning disable CA2000 // Dispose objects before losing scope
 			ILogger startupLogger = new SerilogLoggerFactory().CreateLogger(typeof(ServerApp).FullName ?? "ServerApp");
+#pragma warning restore CA2000 // Dispose objects before losing scope
 
 			ServiceCollection services = new ServiceCollection();
 			services.AddCommandsFromAssembly(Assembly.GetExecutingAssembly());
