Horde: Add a top-level exception handler.

#rnx

diff --git a/Engine/Source/Programs/Horde/HordeServer/ServerApp.cs b/Engine/Source/Programs/Horde/HordeServer/ServerApp.cs
--- a/Engine/Source/Programs/Horde/HordeServer/ServerApp.cs
+++ b/Engine/Source/Programs/Horde/HordeServer/ServerApp.cs
@@ -196,31 +196,39 @@ namespace HordeServer
 			ILogger startupLogger = new SerilogLoggerFactory().CreateLogger(typeof(ServerApp).FullName ?? "ServerApp");
 #pragma warning restore CA2000 // Dispose objects before losing scope
 
-			ServiceCollection services = new ServiceCollection();
-			services.AddCommandsFromAssembly(Assembly.GetExecutingAssembly());
-			services.AddLogging(builder => builder.AddSerilog());
-			services.AddSingleton<IConfiguration>(config);
-			services.AddSingleton<ServerSettings>(serverSettings);
-			services.Configure<ServerSettings>(x => Startup.BindServerSettings(config, x));
-			services.AddTransient<IServerStartup, Startup>();
+			try
+			{
+				ServiceCollection services = new ServiceCollection();
+				services.AddCommandsFromAssembly(Assembly.GetExecutingAssembly());
+				services.AddLogging(builder => builder.AddSerilog());
+				services.AddSingleton<IConfiguration>(config);
+				services.AddSingleton<ServerSettings>(serverSettings);
+				services.Configure<ServerSettings>(x => Startup.BindServerSettings(config, x));
+				services.AddTransient<IServerStartup, Startup>();
 
-			ServerInfo serverInfo = new ServerInfo(config, Options.Create(serverSettings));
-			services.AddSingleton<IServerInfo>(serverInfo);
+				ServerInfo serverInfo = new ServerInfo(config, Options.Create(serverSettings));
+				services.AddSingleton<IServerInfo>(serverInfo);
 
-			s_pluginCollection = CreatePluginCollection(config, startupLogger);
-			services.AddSingleton<IPluginCollection>(s_pluginCollection);
+				s_pluginCollection = CreatePluginCollection(config, startupLogger);
+				services.AddSingleton<IPluginCollection>(s_pluginCollection);
 
-			foreach (Assembly pluginAssembly in s_pluginCollection.LoadedPlugins.Select(x => x.Assembly).Distinct())
-			{
-				services.AddCommandsFromAssembly(pluginAssembly);
-			}
+				foreach (Assembly pluginAssembly in s_pluginCollection.LoadedPlugins.Select(x => x.Assembly).Distinct())
+				{
+					services.AddCommandsFromAssembly(pluginAssembly);
+				}
 
 #pragma warning disable ASP0000 // Do not call 'IServiceCollection.BuildServiceProvider' in 'ConfigureServices'
-			await using (ServiceProvider serviceProvider = services.BuildServiceProvider())
+				await using (ServiceProvider serviceProvider = services.BuildServiceProvider())
+				{
+					return await CommandHost.RunAsync(arguments, serviceProvider, typeof(ServerCommand));
+				}
+#pragma warning restore ASP0000 // Do not call 'IServiceCollection.BuildServiceProvider' in 'ConfigureServices'
+			}
+			catch (Exception ex)
 			{
-				return await CommandHost.RunAsync(arguments, serviceProvider, typeof(ServerCommand));
+				startupLogger.LogError(ex, "Uncaught exception: {Message}", ex.Message);
+				throw;
 			}
-#pragma warning restore ASP0000 // Do not call 'IServiceCollection.BuildServiceProvider' in 'ConfigureServices'
 		}
 
 		internal static void InitializePluginsForTests()
