Horde: Enable all plugins by default. Fixes error on startup if plugins have an existing config section.

#rnx

diff --git a/Engine/Source/Programs/Horde/HordeServer.Shared/Plugins/PluginServerConfig.cs b/Engine/Source/Programs/Horde/HordeServer.Shared/Plugins/PluginServerConfig.cs
--- a/Engine/Source/Programs/Horde/HordeServer.Shared/Plugins/PluginServerConfig.cs
+++ b/Engine/Source/Programs/Horde/HordeServer.Shared/Plugins/PluginServerConfig.cs
@@ -10,6 +10,6 @@ namespace HordeServer.Plugins
 		/// <summary>
 		/// Whether the plugin should be enabled or not
 		/// </summary>
-		public bool Enabled { get; set; }
+		public bool Enabled { get; set; } = true;
 	}
 }
