Horde: Add RetryStepAfterBatchFailAsync test for investigation
#skipci

diff --git a/Engine/Source/Programs/Horde/Plugins/Build/HordeServer.Build.Tests/Jobs/JobCollectionTests.cs b/Engine/Source/Programs/Horde/Plugins/Build/HordeServer.Build.Tests/Jobs/JobCollectionTests.cs
--- a/Engine/Source/Programs/Horde/Plugins/Build/HordeServer.Build.Tests/Jobs/JobCollectionTests.cs
+++ b/Engine/Source/Programs/Horde/Plugins/Build/HordeServer.Build.Tests/Jobs/JobCollectionTests.cs
@@ -1,5 +1,6 @@
 // Copyright Epic Games, Inc. All Rights Reserved.
 
+using System.Security.Claims;
 using EpicGames.Core;
 using EpicGames.Horde.Agents;
 using EpicGames.Horde.Agents.Leases;
@@ -20,7 +21,10 @@ using HordeServer.Jobs.Templates;
 using HordeServer.Logs;
 using HordeServer.Projects;
 using HordeServer.Streams;
+using HordeServer.Users;
 using HordeServer.Utilities;
+using Microsoft.AspNetCore.Http;
+using Microsoft.AspNetCore.Mvc;
 using Moq;
 
 namespace HordeServer.Tests.Jobs
@@ -28,15 +32,47 @@ namespace HordeServer.Tests.Jobs
 	[TestClass]
 	public class JobCollectionTests : BuildTestSetup
 	{
-		private readonly StreamId _streamId = new ("ue4-main");
-		private readonly TemplateId _templateId1 = new ("template1");
+		private readonly StreamId _streamId = new("ue4-main");
+		private readonly TemplateId _templateId1 = new("template1");
 		private readonly ITemplate _template1 = new TemplateStub(ContentHash.MD5("graphHash"), "template1");
-		
+
+		// override JobController with valid user, so we can use REST API in tests
+		private JobsController? _jobsController;
+		private new JobsController JobsController
+		{
+			get
+			{
+				if (_jobsController == null)
+				{
+					IUser user = UserCollection.FindOrAddUserByLoginAsync("TestUser").Result;
+					_jobsController = base.JobsController;
+					ControllerContext controllerContext = new ControllerContext();
+					controllerContext.HttpContext = new DefaultHttpContext();
+					controllerContext.HttpContext.User = new ClaimsPrincipal(new ClaimsIdentity(
+						new List<Claim>
+						{
+							HordeClaims.AdminClaim.ToClaim(),
+							new Claim(ClaimTypes.Name, "TestUser"),
+							new Claim(HordeClaimTypes.UserId, user.Id.ToString())
+						}
+						, "TestAuthType"));
+					_jobsController.ControllerContext = controllerContext;
+
+				}
+				return _jobsController;
+			}
+		}
+
+		static T ResultToValue<T>(ActionResult<T> result) where T : class
+		{
+			return result.Value! as T;
+		}
+
 		public JobCollectionTests()
 		{
 			UpdateConfig(globalConfig =>
 			{
-				ProjectConfig projectConfig = new () { Id = new ProjectId("ue4") };
+				ProjectConfig projectConfig = new() { Id = new ProjectId("ue4") };
 				projectConfig.Streams.Add(new StreamConfig { Id = _streamId });
 				globalConfig.Plugins.GetBuildConfig().Projects.Add(projectConfig);
 			});
@@ -695,6 +731,39 @@ namespace HordeServer.Tests.Jobs
 			Assert.AreEqual(JobStepOutcome.Success, step.Outcome);
 		}
 
+		[TestMethod]
+		public async Task RetryStepAfterBatchFailAsync()
+		{
+			Mock<ITemplate> templateMock = new Mock<ITemplate>(MockBehavior.Strict);
+			templateMock.SetupGet(x => x.InitialAgentType).Returns((string?)null);
+
+			List<NewGroup> newGroups = new List<NewGroup>();
+
+			NewGroup group = AddGroup(newGroups);
+			AddNode(group, "Step 1", null);
+			AddNode(group, "Step 2", new[] { "Step 1" });
+			AddNode(group, "Step 3", new[] { "Step 2" });
+
+			IGraph graph = await GraphCollection.AppendAsync(null, newGroups, null, null);
+
+			CreateJobOptions options = new CreateJobOptions();
+			options.Arguments.Add("-Target=Step 3");
+
+			IJob job = await JobCollection.AddAsync(JobIdUtils.GenerateNewId(), new StreamId("ue4-main"), new TemplateId("test-build"), ContentHash.SHA1("hello"), graph, "Test job", CommitIdWithOrder.FromPerforceChange(123), CommitIdWithOrder.FromPerforceChange(123), options);
+
+			IJobStepBatch initialBatch = job.Batches[0];
+
+			job = await RunBatchAsync(job, 0);
+			job = Deref(await job.TryUpdateBatchAsync(initialBatch.Id, null, JobStepBatchState.Complete, JobStepBatchError.Cancelled));
+
+			UpdateStepResponse response = ResultToValue(await JobsController.UpdateStepAsync(job.Id, initialBatch.Id, initialBatch.Steps[0].Id, new UpdateStepRequest { Retry = true }));
+
+			job = Deref(await JobCollection.GetAsync(job.Id));
+
+			// Expecting 2 batches here, and new steps and stuff, commented out for Ben to look at, as it fails
+			// Assert.AreEqual(2, job.Batches.Count);
+		}
+
 		[TestMethod]
 		public async Task UpdateBatchesDuringJobAsync()
 		{
@@ -844,32 +913,32 @@ namespace HordeServer.Tests.Jobs
 			Assert.AreEqual("Engine/Source", graph.Artifacts[0].BasePath);
 			Assert.AreEqual("fileset", graph.Artifacts[0].OutputName);
 		}
-		
+
 		[TestMethod]
 		public async Task FindJobsAsync()
 		{
 			IGraph graph = await GraphCollection.AddAsync(_template1, null);
-			
+
 			List<IJob> jobs = [];
 			for (int i = 0; i < 10; i++)
 			{
 				jobs.Add(await AddJobAsync(graph, $"{i}"));
 			}
-			
+
 			jobs[2] = await RunBatchAsync(jobs[2], 0);
 			await RunStepAsync(jobs[2], 0, 0, JobStepOutcome.Success); // Setup Build
-			
+
 			{
 				IReadOnlyList<IJob> jobsFound = await JobCollection.FindAsync(new FindJobOptions());
 				Assert.AreEqual(jobs.Count, jobsFound.Count);
 			}
-			
+
 			{
 				IReadOnlyList<IJob> jobsFound = await JobCollection.FindAsync(new FindJobOptions() { Outcome = [JobStepOutcome.Success] });
 				Assert.AreEqual(1, jobsFound.Count);
 			}
 		}
-		
+
 		private async Task<IJob> AddJobAsync(IGraph graph, string commitId)
 		{
 			return await JobCollection.AddAsync(JobIdUtils.GenerateNewId(), _streamId, _templateId1, _template1.Hash, graph, "jobName", new CommitId(commitId), null, new CreateJobOptions());
