Horde: Disable warnings for fixes in the wrong stream by default. This can be toggled using the ShowMergeWarnings workflow option.

#rnx

diff --git a/Engine/Source/Programs/Horde/Plugins/Build/HordeServer.Build/Issues/WorkflowConfig.cs b/Engine/Source/Programs/Horde/Plugins/Build/HordeServer.Build/Issues/WorkflowConfig.cs
--- a/Engine/Source/Programs/Horde/Plugins/Build/HordeServer.Build/Issues/WorkflowConfig.cs
+++ b/Engine/Source/Programs/Horde/Plugins/Build/HordeServer.Build/Issues/WorkflowConfig.cs
@@ -103,6 +103,11 @@ namespace HordeServer.Issues
 		/// </summary>
 		public bool SkipWhenEmpty { get; set; }
 
+		/// <summary>
+		/// Whether to show warnings about merging changes into the origin stream.
+		/// </summary>
+		public bool ShowMergeWarnings { get; set; }
+
 		/// <summary>
 		/// Additional node annotations implicit in this workflow
 		/// </summary>
diff --git a/Engine/Source/Programs/Horde/Plugins/Build/HordeServer.Build/Notifications/Sinks/SlackNotificationSink.cs b/Engine/Source/Programs/Horde/Plugins/Build/HordeServer.Build/Notifications/Sinks/SlackNotificationSink.cs
--- a/Engine/Source/Programs/Horde/Plugins/Build/HordeServer.Build/Notifications/Sinks/SlackNotificationSink.cs
+++ b/Engine/Source/Programs/Horde/Plugins/Build/HordeServer.Build/Notifications/Sinks/SlackNotificationSink.cs
@@ -1310,7 +1310,7 @@ namespace HordeServer.Notifications.Sinks
 						await PostSingleMessageToThreadAsync(triageChannel, fixFailedEventId, threadId, fixFailedMessage, cancellationToken);
 					}
 
-					if (fixFailedSpan == null)
+					if (fixFailedSpan == null && workflow.ShowMergeWarnings)
 					{
 						foreach (IIssueStream stream in issue.Streams)
 						{
