Horde: Disable analyzers in debug builds to improve build times.

#rnx

diff --git a/Engine/Source/Programs/Horde/HordeAgent.Tests/HordeAgent.Tests.csproj b/Engine/Source/Programs/Horde/HordeAgent.Tests/HordeAgent.Tests.csproj
--- a/Engine/Source/Programs/Horde/HordeAgent.Tests/HordeAgent.Tests.csproj
+++ b/Engine/Source/Programs/Horde/HordeAgent.Tests/HordeAgent.Tests.csproj
@@ -8,6 +8,7 @@
     <DebugType>pdbonly</DebugType>
     <Configurations>Debug;Release;Development;Analyze</Configurations>
     <NoWarn>1701;1702;1591</NoWarn>
+    <RunAnalyzersDuringBuild>False</RunAnalyzersDuringBuild>
   </PropertyGroup>
 
   <PropertyGroup Condition=" '$(Configuration)|$(Platform)' == 'Analyze|AnyCPU' ">
diff --git a/Engine/Source/Programs/Horde/HordeAgent/HordeAgent.csproj b/Engine/Source/Programs/Horde/HordeAgent/HordeAgent.csproj
--- a/Engine/Source/Programs/Horde/HordeAgent/HordeAgent.csproj
+++ b/Engine/Source/Programs/Horde/HordeAgent/HordeAgent.csproj
@@ -13,6 +13,7 @@
     <DebugType>pdbonly</DebugType>
     <Configurations>Debug;Release;Analyze</Configurations>
     <NoWarn>1701;1702;NU3003</NoWarn>
+    <RunAnalyzersDuringBuild>False</RunAnalyzersDuringBuild>
   </PropertyGroup>
 
   <PropertyGroup Condition="'$(Configuration)|$(Platform)'=='Debug|AnyCPU'">
diff --git a/Engine/Source/Programs/Horde/HordeCmd/HordeCmd.csproj b/Engine/Source/Programs/Horde/HordeCmd/HordeCmd.csproj
--- a/Engine/Source/Programs/Horde/HordeCmd/HordeCmd.csproj
+++ b/Engine/Source/Programs/Horde/HordeCmd/HordeCmd.csproj
@@ -7,6 +7,7 @@
     <ImplicitUsings>enable</ImplicitUsings>
     <Nullable>enable</Nullable>
     <Configurations>Debug;Release;Analyze</Configurations>
+    <RunAnalyzersDuringBuild>False</RunAnalyzersDuringBuild>
   </PropertyGroup>
 
   <PropertyGroup Condition=" '$(Configuration)|$(Platform)' == 'Analyze|AnyCPU' ">
diff --git a/Engine/Source/Programs/Horde/HordeServer.Tests/HordeServer.Tests.csproj b/Engine/Source/Programs/Horde/HordeServer.Tests/HordeServer.Tests.csproj
--- a/Engine/Source/Programs/Horde/HordeServer.Tests/HordeServer.Tests.csproj
+++ b/Engine/Source/Programs/Horde/HordeServer.Tests/HordeServer.Tests.csproj
@@ -8,6 +8,7 @@
     <Configurations>Debug;Release;Development;Analyze</Configurations>
     <NoWarn>1701;1702;1591</NoWarn>
     <GenerateProgramFile>false</GenerateProgramFile>
+    <RunAnalyzersDuringBuild>False</RunAnalyzersDuringBuild>
   </PropertyGroup>
 
   <PropertyGroup Condition=" '$(Configuration)|$(Platform)' == 'Analyze|AnyCPU' ">
