Horde: Increase the upload grace period to 12 hours.

#rnx

diff --git a/Engine/Source/Programs/Horde/Plugins/Storage/HordeServer.Storage/Storage/StorageService.cs b/Engine/Source/Programs/Horde/Plugins/Storage/HordeServer.Storage/Storage/StorageService.cs
--- a/Engine/Source/Programs/Horde/Plugins/Storage/HordeServer.Storage/Storage/StorageService.cs
+++ b/Engine/Source/Programs/Horde/Plugins/Storage/HordeServer.Storage/Storage/StorageService.cs
@@ -681,6 +681,7 @@ namespace HordeServer.Storage
 			}
 
 			await _blobCollection.InsertOneAsync(blobInfo, new InsertOneOptions { }, cancellationToken);
+			_logger.LogWarning("Created blob {BlobId} at {Path} ({NumImports} imports)", blobInfo.Id, blobInfo.Path, blobInfo.Imports.Count);
 		}
 
 		async ValueTask CheckBlobExistsAsync(NamespaceId namespaceId, BlobLocator locator, CancellationToken cancellationToken)
@@ -735,7 +736,7 @@ namespace HordeServer.Storage
 
 			GcState gcState = await _gcState.GetAsync(cancellationToken);
 
-			DateTime ingestTimeUtc = _clock.UtcNow - TimeSpan.FromMinutes(30.0);
+			DateTime ingestTimeUtc = _clock.UtcNow - TimeSpan.FromHours(12.0);
 
 			// Get the current state of the storage system
 			State state = CreateState(_storageConfig.CurrentValue);
