Horde: Add an error response path for exceptions thrown while gathering blobs for Unsync.

#rnx

diff --git a/Engine/Source/Programs/Horde/Plugins/Build/HordeServer.Build/Artifacts/ArtifactsController.cs b/Engine/Source/Programs/Horde/Plugins/Build/HordeServer.Build/Artifacts/ArtifactsController.cs
--- a/Engine/Source/Programs/Horde/Plugins/Build/HordeServer.Build/Artifacts/ArtifactsController.cs
+++ b/Engine/Source/Programs/Horde/Plugins/Build/HordeServer.Build/Artifacts/ArtifactsController.cs
@@ -5,6 +5,7 @@ using System.Buffers.Binary;
 using System.Diagnostics;
 using System.Net;
 using System.Net.Mime;
+using System.Text;
 using System.Text.RegularExpressions;
 using EpicGames.Core;
 using EpicGames.Horde.Acls;
@@ -721,127 +722,139 @@ namespace HordeServer.Artifacts
 			response.StatusCode = (int)HttpStatusCode.OK;
 
 			await response.StartAsync(cancellationToken);
-
-			// Create a pipeline for downloading new blocks
-			await using BlobPipeline<IoHash> pipeline = new BlobPipeline<IoHash>();
-
-			// Cache of blocks that have already been extracted and compressed
-			long cachedResponseLength = 0;
-			double cachedResponseTime = 0.0;
-			List<(IoHash, IBlockCacheValue)> cachedBlocks = new List<(IoHash, IBlockCacheValue)>();
 			try
 			{
-				// Find all the blob refs that we need to fetch
-				foreach (string block in request.Blocks)
+				// Create a pipeline for downloading new blocks
+				await using BlobPipeline<IoHash> pipeline = new BlobPipeline<IoHash>();
+
+				// Cache of blocks that have already been extracted and compressed
+				long cachedResponseLength = 0;
+				double cachedResponseTime = 0.0;
+				List<(IoHash, IBlockCacheValue)> cachedBlocks = new List<(IoHash, IBlockCacheValue)>();
+				try
 				{
-					if (!IoHash.TryParse(block, out IoHash hash))
+					// Find all the blob refs that we need to fetch
+					foreach (string block in request.Blocks)
 					{
-						return BadRequest($"Invalid IoHash value: {block}");
+						if (!IoHash.TryParse(block, out IoHash hash))
+						{
+							return BadRequest($"Invalid IoHash value: {block}");
+						}
+
+						IHashedBlobRef? blobRef = await _unsyncCache.ReadBlobRefAsync(artifact, hash, cancellationToken);
+						if (blobRef == null)
+						{
+							return NotFound($"Hash '{hash}' is not part of artifact {artifact.Id}");
+						}
+
+						IBlockCacheValue? cacheValue = _blockCache.Get(GetUnsyncBlockKey(hash, compress));
+						if (cacheValue != null)
+						{
+							cachedBlocks.Add((hash, cacheValue));
+						}
+						else
+						{
+							pipeline.Add(new BlobRequest<IoHash>(blobRef, hash));
+						}
 					}
+					pipeline.FinishAdding();
 
-					IHashedBlobRef? blobRef = await _unsyncCache.ReadBlobRefAsync(artifact, hash, cancellationToken);
-					if (blobRef == null)
+					// Write all the cached blocks to the response
+					Stopwatch cachedResponseTimer = Stopwatch.StartNew();
+					foreach ((IoHash hash, IBlockCacheValue value) in cachedBlocks)
 					{
-						return NotFound($"Hash '{hash}' is not part of artifact {artifact.Id}");
-					}
+						int length = (int)value.Data.Length;
 
-					IBlockCacheValue? cacheValue = _blockCache.Get(GetUnsyncBlockKey(hash, compress));
-					if (cacheValue != null)
-					{
-						cachedBlocks.Add((hash, cacheValue));
+						Memory<byte> buffer = response.BodyWriter.GetMemory(length);
+						value.Data.CopyTo(buffer.Span);
+						response.BodyWriter.Advance(length);
+						cachedResponseLength += length;
+
+						await response.BodyWriter.FlushAsync(cancellationToken);
 					}
-					else
+					cachedResponseTime = cachedResponseTimer.Elapsed.TotalSeconds;
+				}
+				finally
+				{
+					foreach ((_, IBlockCacheValue value) in cachedBlocks)
 					{
-						pipeline.Add(new BlobRequest<IoHash>(blobRef, hash));
+						value.Dispose();
 					}
 				}
-				pipeline.FinishAdding();
 
-				// Write all the cached blocks to the response
-				Stopwatch cachedResponseTimer = Stopwatch.StartNew();
-				foreach ((IoHash hash, IBlockCacheValue value) in cachedBlocks)
+				// Read all the uncached blocks
+				long uncachedResponseLength = 0;
+				double uncachedResponseTime = 0.0;
+				if (cachedBlocks.Count < request.Blocks.Count)
 				{
-					int length = (int)value.Data.Length;
+					Stopwatch uncachedResponseTimer = Stopwatch.StartNew();
 
-					Memory<byte> buffer = response.BodyWriter.GetMemory(length);
-					value.Data.CopyTo(buffer.Span);
-					response.BodyWriter.Advance(length);
-					cachedResponseLength += length;
+					ArrayMemoryWriter blockWriter = new ArrayMemoryWriter(128 * 1024);
+					await foreach (BlobResponse<IoHash> blobResponse in pipeline.ReadAllAsync(cancellationToken))
+					{
+						using BlobData blobData = blobResponse.BlobData;
+						blockWriter.Clear();
+
+						// Skip the header to start with. We'll write it once we know the compressed size
+						blockWriter.Advance(BlockHeaderSize);
+
+						// Write the payload
+						if (compress)
+						{
+							BundleData.Compress(BundleCompressionFormat.Zstd, blobData.Data, blockWriter);
+						}
+						else
+						{
+							blockWriter.WriteFixedLengthBytes(blobData.Data.Span);
+						}
+
+						// Write the header
+						WriteBlockHeader(blockWriter.WrittenSpan, blockWriter.WrittenSpan.Length - BlockHeaderSize, blobResponse.UserData, blobData.Data.Length);
+
+						// Add it to the cache for future requests
+						string cacheKey = GetUnsyncBlockKey(blobResponse.UserData, compress);
+						_blockCache.Add(cacheKey, blockWriter.WrittenMemory);
+
+						// Copy it to the response output
+						await response.BodyWriter.WriteAsync(blockWriter.WrittenMemory, cancellationToken);
+						await response.BodyWriter.FlushAsync(cancellationToken);
+						uncachedResponseLength += blockWriter.WrittenMemory.Length;
+					}
 
-					await response.BodyWriter.FlushAsync(cancellationToken);
+					uncachedResponseTime = uncachedResponseTimer.Elapsed.TotalSeconds;
 				}
-				cachedResponseTime = cachedResponseTimer.Elapsed.TotalSeconds;
-			}
-			finally
-			{
-				foreach ((_, IBlockCacheValue value) in cachedBlocks)
-				{
-					value.Dispose();
-				}
-			}
 
-			// Read all the uncached blocks
-			long uncachedResponseLength = 0;
-			double uncachedResponseTime = 0.0;
-			if (cachedBlocks.Count < request.Blocks.Count)
-			{
-				Stopwatch uncachedResponseTimer = Stopwatch.StartNew();
+				long totalResponseLength = cachedResponseLength + uncachedResponseLength;
+				double totalResponseTime = totalResponseTimer.Elapsed.TotalSeconds;
+				_logger.LogInformation("Unsync request for {NumBlocks} blocks ({Type}, {TotalSize:n1}mb, {TotalSpeed:n1}mb/s); {NumCachedBlocks} cached ({CachedSize:n1}mb, {CachedSpeed:n1}mb/s), {NumUncachedBlocks} uncached ({UncachedSize:n1}mb, {UncachedSpeed:n1}mb/s)",
+					request.Blocks.Count,
+					compress ? "compressed" : "not compressed",
+					totalResponseLength / (1024.0 * 1024.0),
+					(totalResponseTime > 0.0) ? (totalResponseLength / (1024 * 1024 * totalResponseTime)) : 0.0,
 
-				ArrayMemoryWriter blockWriter = new ArrayMemoryWriter(128 * 1024);
-				await foreach (BlobResponse<IoHash> blobResponse in pipeline.ReadAllAsync(cancellationToken))
-				{
-					using BlobData blobData = blobResponse.BlobData;
-					blockWriter.Clear();
+					cachedBlocks.Count,
+					cachedResponseLength / (1024.0 * 1024.0),
+					(cachedResponseTime > 0.0) ? (cachedResponseLength / (1024 * 1024 * cachedResponseTime)) : 0.0,
 
-					// Skip the header to start with. We'll write it once we know the compressed size
-					blockWriter.Advance(BlockHeaderSize);
-
-					// Write the payload
-					if (compress)
-					{
-						BundleData.Compress(BundleCompressionFormat.Zstd, blobData.Data, blockWriter);
-					}
-					else
-					{
-						blockWriter.WriteFixedLengthBytes(blobData.Data.Span);
-					}
-
-					// Write the header
-					WriteBlockHeader(blockWriter.WrittenSpan, blockWriter.WrittenSpan.Length - BlockHeaderSize, blobResponse.UserData, blobData.Data.Length);
-
-					// Add it to the cache for future requests
-					string cacheKey = GetUnsyncBlockKey(blobResponse.UserData, compress);
-					_blockCache.Add(cacheKey, blockWriter.WrittenMemory);
+					request.Blocks.Count - cachedBlocks.Count,
+					uncachedResponseLength / (1024.0 * 1024.0),
+					(uncachedResponseTime > 0.0) ? (uncachedResponseLength / (1024 * 1024 * uncachedResponseTime)) : 0.0
+				);
+			}
+			catch (Exception ex)
+			{
+				_logger.LogError(ex, "Exception while extracting unsync blobs: {Message}", ex.Message);
 
-					// Copy it to the response output
-					await response.BodyWriter.WriteAsync(blockWriter.WrittenMemory, cancellationToken);
-					await response.BodyWriter.FlushAsync(cancellationToken);
-					uncachedResponseLength += blockWriter.WrittenMemory.Length;
-				}
+				ArrayMemoryWriter writer = new ArrayMemoryWriter(1024);
+				WriteErrorHeader(writer, $"Error retrieving blobs: {ex.Message}");
 
-				uncachedResponseTime = uncachedResponseTimer.Elapsed.TotalSeconds;
+				await response.BodyWriter.WriteAsync(writer.WrittenMemory, cancellationToken);
+				await response.BodyWriter.FlushAsync(cancellationToken);
 			}
 
 			// Finish the response
 			await response.CompleteAsync();
 
-			long totalResponseLength = cachedResponseLength + uncachedResponseLength;
-			double totalResponseTime = totalResponseTimer.Elapsed.TotalSeconds;
-			_logger.LogInformation("Unsync request for {NumBlocks} blocks ({Type}, {TotalSize:n1}mb, {TotalSpeed:n1}mb/s); {NumCachedBlocks} cached ({CachedSize:n1}mb, {CachedSpeed:n1}mb/s), {NumUncachedBlocks} uncached ({UncachedSize:n1}mb, {UncachedSpeed:n1}mb/s)",
-				request.Blocks.Count,
-				compress ? "compressed" : "not compressed",
-				totalResponseLength / (1024.0 * 1024.0),
-				(totalResponseTime > 0.0) ? (totalResponseLength / (1024 * 1024 * totalResponseTime)) : 0.0,
-
-				cachedBlocks.Count,
-				cachedResponseLength / (1024.0 * 1024.0),
-				(cachedResponseTime > 0.0) ? (cachedResponseLength / (1024 * 1024 * cachedResponseTime)) : 0.0,
-
-				request.Blocks.Count - cachedBlocks.Count,
-				uncachedResponseLength / (1024.0 * 1024.0),
-				(uncachedResponseTime > 0.0) ? (uncachedResponseLength / (1024 * 1024 * uncachedResponseTime)) : 0.0
-			);
-
 			return Empty;
 		}
 
@@ -864,6 +877,14 @@ namespace HordeServer.Artifacts
 			decompressedHash.CopyTo(data);
 		}
 
+		static void WriteErrorHeader(ArrayMemoryWriter writer, string message)
+		{
+			Utf8String messageUtf8 = new Utf8String(message);
+			writer.WriteUInt64(0x_4C5C_2AAB_A992_DEAD);
+			writer.WriteUInt32((uint)(messageUtf8.Length + 1));
+			writer.WriteNullTerminatedUtf8String(messageUtf8);
+		}
+
 		/// <summary>
 		/// Downloads an individual file from an artifact
 		/// </summary>
