Horde: Updating schema documentation

#rnx

diff --git a/Engine/Source/Programs/Horde/Docs/Config/Schema/AclActions.md b/Engine/Source/Programs/Horde/Docs/Config/Schema/AclActions.md
--- a/Engine/Source/Programs/Horde/Docs/Config/Schema/AclActions.md
+++ b/Engine/Source/Programs/Horde/Docs/Config/Schema/AclActions.md
@@ -37,41 +37,40 @@
 | `DeleteAccount` |  |
 | `ViewAccount` |  |
 
-## Artifacts
+## Agents
 
 | Name | Description |
 | ---- | ----------- |
-| `ReadArtifact` | Permission to read from an artifact |
-| `WriteArtifact` | Permission to write to an artifact |
-| `DeleteArtifact` | Permission to delete to an artifact |
-| `UploadArtifact` | Ability to create an artifact. Typically just for debugging; agents have this access for a particular session. |
-| `DownloadArtifact` | Ability to download an artifact |
+| `CreateAgent` | Ability to create an agent. This may be done explicitly, or granted to agents to allow them to self-register. |
+| `UpdateAgent` | Update an agent's name, pools, etc... |
+| `DeleteAgent` | Soft-delete an agent |
+| `ViewAgent` | View an agent |
+| `ListAgents` | List the available agents |
 
-## Bisect
+## Compute
 
 | Name | Description |
 | ---- | ----------- |
-| `CreateBisectTask` | Ability to start new bisect tasks |
-| `UpdateBisectTask` | Ability to update a bisect task |
-| `ViewBisectTask` | Ability to view a bisect task |
+| `AddComputeTasks` | User can add tasks to the compute cluster |
+| `GetComputeTasks` | User can get and list tasks from the compute cluster |
 
-## Devices
+## Leases
 
 | Name | Description |
 | ---- | ----------- |
-| `DeviceRead` | Ability to read devices |
-| `DeviceWrite` | Ability to write devices |
+| `ViewLeases` | View all the leases that an agent has worked on |
+| `ViewLeaseTasks` | View the task data for a lease |
 
-## Jobs
+## Logs
 
 | Name | Description |
 | ---- | ----------- |
-| `CreateJob` | Ability to start new jobs |
-| `UpdateJob` | Rename a job, modify its priority, etc... |
-| `DeleteJob` | Delete a job properties |
-| `ExecuteJob` | Allows updating a job metadata (name, changelist number, step properties, new groups, job states, etc...). Typically granted to agents. Not user facing. |
-| `RetryJobStep` | Ability to retry a failed job step |
-| `ViewJob` | Ability to view a job |
+| `CreateLog` | Ability to create a log. Implicitly granted to agents. |
+| `UpdateLog` | Ability to update log metadata |
+| `ViewLog` | Ability to view a log contents |
+| `WriteLogData` | Ability to write log data |
+| `CreateEvent` | Ability to create events |
+| `ViewEvent` | Ability to view events |
 
 ## Notices
 
@@ -81,27 +80,15 @@
 | `UpdateNotice` |  |
 | `DeleteNotice` |  |
 
-## Notifications
+## Pools
 
 | Name | Description |
 | ---- | ----------- |
-| `CreateSubscription` | Ability to subscribe to notifications |
+| `CreatePool` | Create a global pool of agents |
+| `UpdatePool` | Modify an agent pool |
+| `DeletePool` | Delete an agent pool |
+| `ViewPool` | Ability to view a pool |
+| `ListPools` | View all the available agent pools |
 
-## Projects
-
-| Name | Description |
-| ---- | ----------- |
-| `CreateProject` | Allows the creation of new projects |
-| `DeleteProject` | Allows deletion of projects. |
-| `UpdateProject` | Modify attributes of a project (name, categories, etc...) |
-| `ViewProject` | View information about a project |
-
-## Replicators
-
-| Name | Description |
-| ---- | ----------- |
-| `UpdateReplicator` | Allows deletion of projects. |
-| `ViewReplicator` | Allows the creation of new projects |
-
 ## ServiceAccounts
 
@@ -112,16 +99,20 @@
 | `DeleteAccount` |  |
 | `ViewAccount` |  |
 
-## Streams
+## Sessions
 
 | Name | Description |
 | ---- | ----------- |
-| `CreateStream` | Allows the creation of new streams within a project |
-| `UpdateStream` | Allows updating a stream (agent types, templates, schedules) |
-| `DeleteStream` | Allows deleting a stream |
-| `ViewStream` | Ability to view a stream |
-| `ViewChanges` | View changes submitted to a stream. NOTE: this returns responses from the server's Perforce account, which may be a priviledged user. |
-| `ViewTemplate` | View template associated with a stream |
+| `CreateSession` | Granted to agents to call CreateSession, which returns a bearer token identifying themselves valid to call UpdateSesssion via gRPC. |
+| `ViewSession` | Allows viewing information about an agent session |
 
+## Software
+
+| Name | Description |
+| ---- | ----------- |
+| `UploadSoftware` | Ability to upload new versions of the agent software |
+| `DownloadSoftware` | Ability to download the agent software |
+| `DeleteSoftware` | Ability to delete agent software |
+
 ## Accounts
 
@@ -132,41 +123,40 @@
 | `DeleteAccount` |  |
 | `ViewAccount` |  |
 
-## Artifacts
+## Agents
 
 | Name | Description |
 | ---- | ----------- |
-| `ReadArtifact` |  |
-| `WriteArtifact` |  |
-| `DeleteArtifact` |  |
-| `UploadArtifact` |  |
-| `DownloadArtifact` |  |
+| `CreateAgent` |  |
+| `UpdateAgent` |  |
+| `DeleteAgent` |  |
+| `ViewAgent` |  |
+| `ListAgents` |  |
 
-## Bisect
+## Compute
 
 | Name | Description |
 | ---- | ----------- |
-| `CreateBisectTask` |  |
-| `UpdateBisectTask` |  |
-| `ViewBisectTask` |  |
+| `AddComputeTasks` |  |
+| `GetComputeTasks` |  |
 
-## Devices
+## Leases
 
 | Name | Description |
 | ---- | ----------- |
-| `DeviceRead` |  |
-| `DeviceWrite` |  |
+| `ViewLeases` |  |
+| `ViewLeaseTasks` |  |
 
-## Jobs
+## Logs
 
 | Name | Description |
 | ---- | ----------- |
-| `CreateJob` |  |
-| `UpdateJob` |  |
-| `DeleteJob` |  |
-| `ExecuteJob` |  |
-| `RetryJobStep` |  |
-| `ViewJob` |  |
+| `CreateLog` |  |
+| `UpdateLog` |  |
+| `ViewLog` |  |
+| `WriteLogData` |  |
+| `CreateEvent` |  |
+| `ViewEvent` |  |
 
 ## Notices
 
@@ -176,27 +166,21 @@
 | `UpdateNotice` |  |
 | `DeleteNotice` |  |
 
-## Notifications
+## Pools
 
 | Name | Description |
 | ---- | ----------- |
-| `CreateSubscription` |  |
+| `CreatePool` |  |
+| `UpdatePool` |  |
+| `DeletePool` |  |
+| `ViewPool` |  |
+| `ListPools` |  |
 
-## Projects
+## Secrets
 
 | Name | Description |
 | ---- | ----------- |
-| `CreateProject` |  |
-| `DeleteProject` |  |
-| `UpdateProject` |  |
-| `ViewProject` |  |
+| `ViewSecret` | View a credential |
 
-## Replicators
-
-| Name | Description |
-| ---- | ----------- |
-| `UpdateReplicator` |  |
-| `ViewReplicator` |  |
-
 ## ServiceAccounts
 
@@ -207,16 +191,20 @@
 | `DeleteAccount` |  |
 | `ViewAccount` |  |
 
-## Streams
+## Sessions
 
 | Name | Description |
 | ---- | ----------- |
-| `CreateStream` |  |
-| `UpdateStream` |  |
-| `DeleteStream` |  |
-| `ViewStream` |  |
-| `ViewChanges` |  |
-| `ViewTemplate` |  |
+| `CreateSession` |  |
+| `ViewSession` |  |
 
+## Software
+
+| Name | Description |
+| ---- | ----------- |
+| `UploadSoftware` |  |
+| `DownloadSoftware` |  |
+| `DeleteSoftware` |  |
+
 ## Accounts
 
@@ -231,72 +219,36 @@
 
 | Name | Description |
 | ---- | ----------- |
-| `CreateAgent` | Ability to create an agent. This may be done explicitly, or granted to agents to allow them to self-register. |
-| `UpdateAgent` | Update an agent's name, pools, etc... |
-| `DeleteAgent` | Soft-delete an agent |
-| `ViewAgent` | View an agent |
-| `ListAgents` | List the available agents |
+| `CreateAgent` |  |
+| `UpdateAgent` |  |
+| `DeleteAgent` |  |
+| `ViewAgent` |  |
+| `ListAgents` |  |
 
-## Artifacts
-
-| Name | Description |
-| ---- | ----------- |
-| `ReadArtifact` |  |
-| `WriteArtifact` |  |
-| `DeleteArtifact` |  |
-| `UploadArtifact` |  |
-| `DownloadArtifact` |  |
-
-## Bisect
-
-| Name | Description |
-| ---- | ----------- |
-| `CreateBisectTask` |  |
-| `UpdateBisectTask` |  |
-| `ViewBisectTask` |  |
-
 ## Compute
 
 | Name | Description |
 | ---- | ----------- |
-| `AddComputeTasks` | User can add tasks to the compute cluster |
-| `GetComputeTasks` | User can get and list tasks from the compute cluster |
+| `AddComputeTasks` |  |
+| `GetComputeTasks` |  |
 
-## Devices
-
-| Name | Description |
-| ---- | ----------- |
-| `DeviceRead` |  |
-| `DeviceWrite` |  |
-
-## Jobs
-
-| Name | Description |
-| ---- | ----------- |
-| `CreateJob` |  |
-| `UpdateJob` |  |
-| `DeleteJob` |  |
-| `ExecuteJob` |  |
-| `RetryJobStep` |  |
-| `ViewJob` |  |
-
 ## Leases
 
 | Name | Description |
 | ---- | ----------- |
-| `ViewLeases` | View all the leases that an agent has worked on |
-| `ViewLeaseTasks` | View the task data for a lease |
+| `ViewLeases` |  |
+| `ViewLeaseTasks` |  |
 
 ## Logs
 
 | Name | Description |
 | ---- | ----------- |
-| `CreateLog` | Ability to create a log. Implicitly granted to agents. |
-| `UpdateLog` | Ability to update log metadata |
-| `ViewLog` | Ability to view a log contents |
-| `WriteLogData` | Ability to write log data |
-| `CreateEvent` | Ability to create events |
-| `ViewEvent` | Ability to view events |
+| `CreateLog` |  |
+| `UpdateLog` |  |
+| `ViewLog` |  |
+| `WriteLogData` |  |
+| `CreateEvent` |  |
+| `ViewEvent` |  |
 
 ## Notices
 
@@ -306,37 +258,21 @@
 | `UpdateNotice` |  |
 | `DeleteNotice` |  |
 
-## Notifications
-
-| Name | Description |
-| ---- | ----------- |
-| `CreateSubscription` |  |
-
 ## Pools
 
 | Name | Description |
 | ---- | ----------- |
-| `CreatePool` | Create a global pool of agents |
-| `UpdatePool` | Modify an agent pool |
-| `DeletePool` | Delete an agent pool |
-| `ViewPool` | Ability to view a pool |
-| `ListPools` | View all the available agent pools |
+| `CreatePool` |  |
+| `UpdatePool` |  |
+| `DeletePool` |  |
+| `ViewPool` |  |
+| `ListPools` |  |
 
-## Projects
+## Secrets
 
 | Name | Description |
 | ---- | ----------- |
-| `CreateProject` |  |
-| `DeleteProject` |  |
-| `UpdateProject` |  |
-| `ViewProject` |  |
+| `ViewSecret` |  |
 
-## Replicators
-
-| Name | Description |
-| ---- | ----------- |
-| `UpdateReplicator` |  |
-| `ViewReplicator` |  |
-
 ## ServiceAccounts
 
@@ -351,27 +287,16 @@
 
 | Name | Description |
 | ---- | ----------- |
-| `CreateSession` | Granted to agents to call CreateSession, which returns a bearer token identifying themselves valid to call UpdateSesssion via gRPC. |
-| `ViewSession` | Allows viewing information about an agent session |
+| `CreateSession` |  |
+| `ViewSession` |  |
 
 ## Software
 
 | Name | Description |
 | ---- | ----------- |
-| `UploadSoftware` | Ability to upload new versions of the agent software |
-| `DownloadSoftware` | Ability to download the agent software |
-| `DeleteSoftware` | Ability to delete agent software |
+| `UploadSoftware` |  |
+| `DownloadSoftware` |  |
+| `DeleteSoftware` |  |
 
-## Streams
-
-| Name | Description |
-| ---- | ----------- |
-| `CreateStream` |  |
-| `UpdateStream` |  |
-| `DeleteStream` |  |
-| `ViewStream` |  |
-| `ViewChanges` |  |
-| `ViewTemplate` |  |
-
 ## Accounts
 
@@ -396,19 +321,19 @@
 
 | Name | Description |
 | ---- | ----------- |
-| `ReadArtifact` |  |
-| `WriteArtifact` |  |
-| `DeleteArtifact` |  |
-| `UploadArtifact` |  |
-| `DownloadArtifact` |  |
+| `ReadArtifact` | Permission to read from an artifact |
+| `WriteArtifact` | Permission to write to an artifact |
+| `DeleteArtifact` | Permission to delete to an artifact |
+| `UploadArtifact` | Ability to create an artifact. Typically just for debugging; agents have this access for a particular session. |
+| `DownloadArtifact` | Ability to download an artifact |
 
 ## Bisect
 
 | Name | Description |
 | ---- | ----------- |
-| `CreateBisectTask` |  |
-| `UpdateBisectTask` |  |
-| `ViewBisectTask` |  |
+| `CreateBisectTask` | Ability to start new bisect tasks |
+| `UpdateBisectTask` | Ability to update a bisect task |
+| `ViewBisectTask` | Ability to view a bisect task |
 
 ## Compute
 
@@ -417,36 +342,23 @@
 | `AddComputeTasks` |  |
 | `GetComputeTasks` |  |
 
-## Ddc
-
-| Name | Description |
-| ---- | ----------- |
-| `DdcReadObject` | General read access to refs / blobs and so on |
-| `DdcWriteObject` | General write access to upload refs / blobs etc |
-| `DdcDeleteObject` | Access to delete blobs / refs etc |
-| `DdcDeleteBucket` | Access to delete a particular bucket |
-| `DdcDeleteNamespace` | Access to delete a whole namespace |
-| `DdcReadTransactionLog` | Access to read the transaction log |
-| `DdcWriteTransactionLog` | Access to write the transaction log |
-| `DdcAdminAction` | Access to perform administrative task |
-
 ## Devices
 
 | Name | Description |
 | ---- | ----------- |
-| `DeviceRead` |  |
-| `DeviceWrite` |  |
+| `DeviceRead` | Ability to read devices |
+| `DeviceWrite` | Ability to write devices |
 
 ## Jobs
 
 | Name | Description |
 | ---- | ----------- |
-| `CreateJob` |  |
-| `UpdateJob` |  |
-| `DeleteJob` |  |
-| `ExecuteJob` |  |
-| `RetryJobStep` |  |
-| `ViewJob` |  |
+| `CreateJob` | Ability to start new jobs |
+| `UpdateJob` | Rename a job, modify its priority, etc... |
+| `DeleteJob` | Delete a job properties |
+| `ExecuteJob` | Allows updating a job metadata (name, changelist number, step properties, new groups, job states, etc...). Typically granted to agents. Not user facing. |
+| `RetryJobStep` | Ability to retry a failed job step |
+| `ViewJob` | Ability to view a job |
 
 ## Leases
 
@@ -478,7 +390,7 @@
 
 | Name | Description |
 | ---- | ----------- |
-| `CreateSubscription` |  |
+| `CreateSubscription` | Ability to subscribe to notifications |
 
 ## Pools
 
@@ -494,17 +406,23 @@
 
 | Name | Description |
 | ---- | ----------- |
-| `CreateProject` |  |
-| `DeleteProject` |  |
-| `UpdateProject` |  |
-| `ViewProject` |  |
+| `CreateProject` | Allows the creation of new projects |
+| `DeleteProject` | Allows deletion of projects. |
+| `UpdateProject` | Modify attributes of a project (name, categories, etc...) |
+| `ViewProject` | View information about a project |
 
 ## Replicators
 
 | Name | Description |
 | ---- | ----------- |
-| `UpdateReplicator` |  |
-| `ViewReplicator` |  |
+| `UpdateReplicator` | Allows deletion of projects. |
+| `ViewReplicator` | Allows the creation of new projects |
 
+## Secrets
+
+| Name | Description |
+| ---- | ----------- |
+| `ViewSecret` |  |
+
 ## ServiceAccounts
 
@@ -534,12 +452,12 @@
 
 | Name | Description |
 | ---- | ----------- |
-| `CreateStream` |  |
-| `UpdateStream` |  |
-| `DeleteStream` |  |
-| `ViewStream` |  |
-| `ViewChanges` |  |
-| `ViewTemplate` |  |
+| `CreateStream` | Allows the creation of new streams within a project |
+| `UpdateStream` | Allows updating a stream (agent types, templates, schedules) |
+| `DeleteStream` | Allows deleting a stream |
+| `ViewStream` | Ability to view a stream |
+| `ViewChanges` | View changes submitted to a stream. NOTE: this returns responses from the server's Perforce account, which may be a priviledged user. |
+| `ViewTemplate` | View template associated with a stream |
 
 ## Accounts
 
@@ -585,19 +503,6 @@
 | `AddComputeTasks` |  |
 | `GetComputeTasks` |  |
 
-## Ddc
-
-| Name | Description |
-| ---- | ----------- |
-| `DdcReadObject` |  |
-| `DdcWriteObject` |  |
-| `DdcDeleteObject` |  |
-| `DdcDeleteBucket` |  |
-| `DdcDeleteNamespace` |  |
-| `DdcReadTransactionLog` |  |
-| `DdcWriteTransactionLog` |  |
-| `DdcAdminAction` |  |
-
 ## Devices
 
 | Name | Description |
@@ -678,7 +583,7 @@
 
 | Name | Description |
 | ---- | ----------- |
-| `ViewSecret` | View a credential |
+| `ViewSecret` |  |
 
 ## ServiceAccounts
 
@@ -704,6 +609,16 @@
 | `DownloadSoftware` |  |
 | `DeleteSoftware` |  |
 
+## Storage
+
+| Name | Description |
+| ---- | ----------- |
+| `ReadBlobs` | Ability to read blobs from the storage service |
+| `WriteBlobs` | Ability to write blobs to the storage service |
+| `ReadRefs` | Ability to read refs from the storage service |
+| `WriteRefs` | Ability to write refs to the storage service |
+| `DeleteRefs` | Ability to delete refs |
+
 ## Streams
 
 | Name | Description |
@@ -759,19 +674,6 @@
 | `AddComputeTasks` |  |
 | `GetComputeTasks` |  |
 
-## Ddc
-
-| Name | Description |
-| ---- | ----------- |
-| `DdcReadObject` |  |
-| `DdcWriteObject` |  |
-| `DdcDeleteObject` |  |
-| `DdcDeleteBucket` |  |
-| `DdcDeleteNamespace` |  |
-| `DdcReadTransactionLog` |  |
-| `DdcWriteTransactionLog` |  |
-| `DdcAdminAction` |  |
-
 ## Devices
 
 | Name | Description |
@@ -882,11 +784,11 @@
 
 | Name | Description |
 | ---- | ----------- |
-| `ReadBlobs` | Ability to read blobs from the storage service |
-| `WriteBlobs` | Ability to write blobs to the storage service |
-| `ReadRefs` | Ability to read refs from the storage service |
-| `WriteRefs` | Ability to write refs to the storage service |
-| `DeleteRefs` | Ability to delete refs |
+| `ReadBlobs` |  |
+| `WriteBlobs` |  |
+| `ReadRefs` |  |
+| `WriteRefs` |  |
+| `DeleteRefs` |  |
 
 ## Streams
 
@@ -899,6 +801,12 @@
 | `ViewChanges` |  |
 | `ViewTemplate` |  |
 
+## Symbols
+
+| Name | Description |
+| ---- | ----------- |
+| `ReadSymbols` | Ability to download symbols |
+
 ## Accounts
 
 | Name | Description |
@@ -943,19 +851,6 @@
 | `AddComputeTasks` |  |
 | `GetComputeTasks` |  |
 
-## Ddc
-
-| Name | Description |
-| ---- | ----------- |
-| `DdcReadObject` |  |
-| `DdcWriteObject` |  |
-| `DdcDeleteObject` |  |
-| `DdcDeleteBucket` |  |
-| `DdcDeleteNamespace` |  |
-| `DdcReadTransactionLog` |  |
-| `DdcWriteTransactionLog` |  |
-| `DdcAdminAction` |  |
-
 ## Devices
 
 | Name | Description |
@@ -1083,6 +978,12 @@
 | `ViewChanges` |  |
 | `ViewTemplate` |  |
 
+## Symbols
+
+| Name | Description |
+| ---- | ----------- |
+| `ReadSymbols` |  |
+
 ## Tools
 
 | Name | Description |
@@ -1138,14 +1039,14 @@
 
 | Name | Description |
 | ---- | ----------- |
-| `DdcReadObject` |  |
-| `DdcWriteObject` |  |
-| `DdcDeleteObject` |  |
-| `DdcDeleteBucket` |  |
-| `DdcDeleteNamespace` |  |
-| `DdcReadTransactionLog` |  |
-| `DdcWriteTransactionLog` |  |
-| `DdcAdminAction` |  |
+| `DdcReadObject` | General read access to refs / blobs and so on |
+| `DdcWriteObject` | General write access to upload refs / blobs etc |
+| `DdcDeleteObject` | Access to delete blobs / refs etc |
+| `DdcDeleteBucket` | Access to delete a particular bucket |
+| `DdcDeleteNamespace` | Access to delete a whole namespace |
+| `DdcReadTransactionLog` | Access to read the transaction log |
+| `DdcWriteTransactionLog` | Access to write the transaction log |
+| `DdcAdminAction` | Access to perform administrative task |
 
 ## Devices
 
@@ -1278,7 +1179,7 @@
 
 | Name | Description |
 | ---- | ----------- |
-| `ReadSymbols` | Ability to download symbols |
+| `ReadSymbols` |  |
 
 ## Tools
 
diff --git a/Engine/Source/Programs/Horde/Docs/Config/Schema/Globals.md b/Engine/Source/Programs/Horde/Docs/Config/Schema/Globals.md
--- a/Engine/Source/Programs/Horde/Docs/Config/Schema/Globals.md
+++ b/Engine/Source/Programs/Horde/Docs/Config/Schema/Globals.md
@@ -56,131 +56,28 @@
 
 Name | Description
 ---- | -----------
-`build` | [BuildConfig](#buildconfig)<br>Configuration for the build plugin
-`analytics` | [AnalyticsConfig](#analyticsconfig)<br>Configuration for the analytics plugin
 `compute` | [ComputeConfig](#computeconfig)<br>Configuration for the compute plugin
-`ddc` | [EmptyPluginConfig](#emptypluginconfig)<br>Configuration for the ddc plugin
 `secrets` | [SecretsConfig](#secretsconfig)<br>Configuration for the secrets plugin
+`analytics` | [AnalyticsConfig](#analyticsconfig)<br>Configuration for the analytics plugin
+`build` | [BuildConfig](#buildconfig)<br>Configuration for the build plugin
 `storage` | [StorageConfig](#storageconfig)<br>Configuration for the storage plugin
-`tools` | [ToolsConfig](#toolsconfig)<br>Configuration for the tools plugin
 `symbols` | [SymbolsConfig](#symbolsconfig)<br>Configuration for the symbols plugin
+`tools` | [ToolsConfig](#toolsconfig)<br>Configuration for the tools plugin
+`ddc` | [EmptyPluginConfig](#emptypluginconfig)<br>Configuration for the ddc plugin
 
-## BuildConfig
+## ComputeConfig
 
-Configuration for the build plugin
+Configuration for the compute system
 
 Name | Description
 ---- | -----------
-`perforceClusters` | [PerforceCluster](#perforcecluster)`[]`<br>List of Perforce clusters
-`devices` | [DeviceConfig](#deviceconfig)<br>Device configuration
-`maxConformCount` | `integer`<br>Maximum number of conforms to run at once
-`agentShutdownIfDisabledGracePeriod` | `string`<br>Time to wait before shutting down an agent that has been disabled Used if no value is set on the actual pool.
-`artifactTypes` | [ArtifactTypeConfig](#artifacttypeconfig)`[]`<br>Configuration for different artifact types
-`projects` | [ProjectConfig](Projects.md)`[]`<br>List of projects
-`enableConformTasks` | `boolean`<br>Whether to allow conform tasks to run
-`issueFixedTag` | `string`<br>Commit tag to use for marking issues as fixed
+`acl` | [AclConfig](#aclconfig)<br>Inherited root acl
+`versionEnum` | [ConfigVersion](#configversion-enum)<br>Config version number
+`rates` | [AgentRateConfig](#agentrateconfig)`[]`<br>List of costs of a particular agent type
+`clusters` | [ComputeClusterConfig](#computeclusterconfig)`[]`<br>List of compute profiles
+`pools` | [PoolConfig](#poolconfig)`[]`<br>List of pools
+`software` | [AgentSoftwareConfig](#agentsoftwareconfig)`[]`<br>List of costs of a particular agent type
+`networks` | [NetworkConfig](#networkconfig)`[]`<br>List of networks
 
-## PerforceCluster
-
-Information about a cluster of Perforce servers.
-
-Name | Description
----- | -----------
-`name` | `string`<br>Name of the cluster
-`serviceAccount` | `string`<br>Username for Horde to log in to this server. Will use the first account specified below if not overridden.
-`canImpersonate` | `boolean`<br>Whether the service account can impersonate other users
-`supportsPartitionedWorkspaces` | `boolean`<br>Whether to use partitioned workspaces on this server
-`servers` | [PerforceServer](#perforceserver)`[]`<br>List of servers
-`credentials` | [PerforceCredentials](#perforcecredentials)`[]`<br>List of server credentials
-`autoSdk` | [AutoSdkWorkspace](#autosdkworkspace)`[]`<br>List of autosdk streams
-
-## PerforceServer
-
-Information about an individual Perforce server
-
-Name | Description
----- | -----------
-`serverAndPort` | `string`<br>The server and port. The server may be a DNS entry with multiple records, in which case it will be actively load balanced. If "ssl:" prefix is used, ensure P4 server's fingerprint/certificate is trusted. See Horde's documentation on connecting to SSL-enabled Perforce servers.
-`healthCheck` | `boolean`<br>Whether to query the healthcheck address under each server
-`resolveDns` | `boolean`<br>Whether to resolve the DNS entries and load balance between different hosts
-`maxConformCount` | `integer`<br>Maximum number of simultaneous conforms on this server
-`condition` | `string`<br>Optional condition for a machine to be eligible to use this server
-`properties` | `string[]`<br>List of properties for an agent to be eligible to use this server
-
-## PerforceCredentials
-
-Credentials for a Perforce user
-
-Name | Description
----- | -----------
-`userName` | `string`<br>The username
-`password` | `string`<br>Password for the user
-`ticket` | `string`<br>Login ticket for the user (will be used instead of password if set)
-
-## AutoSdkWorkspace
-
-Path to a platform and stream to use for syncing AutoSDK
-
-Name | Description
----- | -----------
-`name` | `string`<br>Name of this workspace
-`properties` | `string[]`<br>The agent properties to check (eg. "OSFamily=Windows")
-`userName` | `string`<br>Username for logging in to the server
-`stream` | `string`<br>Stream to use
-
-## DeviceConfig
-
-Configuration for devices
-
-Name | Description
----- | -----------
-`platforms` | [DevicePlatformConfig](#deviceplatformconfig)`[]`<br>List of device platforms
-`pools` | [DevicePoolConfig](#devicepoolconfig)`[]`<br>List of device pools
-
-## DevicePlatformConfig
-
-Configuration for a device platform
-
-Name | Description
----- | -----------
-`id` | `string`<br>The id for this platform
-`name` | `string`<br>Name of the platform
-`models` | `string[]`<br>A list of platform models
-`legacyNames` | `string[]`<br>Legacy names which older versions of Gauntlet may be using
-`legacyPerfSpecHighModel` | `string`<br>Model name for the high perf spec, which may be requested by Gauntlet
-
-## DevicePoolConfig
-
-Configuration for a device pool
-
-Name | Description
----- | -----------
-`id` | `string`<br>The id for this platform
-`name` | `string`<br>The name of the pool
-`poolType` | [DevicePoolType](#devicepooltype-enum)<br>The type of the pool
-`projectIds` | `string[]`<br>List of project ids associated with pool
-
-## DevicePoolType (Enum)
-
-The type of device pool
-
-Name | Description
----- | -----------
-`Automation` | Available to CIS jobs
-`Shared` | Shared by users with remote checking and checkouts
-
-## ArtifactTypeConfig
-
-Configuration for an artifact
-
-Name | Description
----- | -----------
-`name` | `string`<br>Legacy 'Name' property
-`type` | `string`<br>Name of the artifact type
-`acl` | [AclConfig](#aclconfig)<br>Acl for the artifact type
-`keepCount` | `integer`<br>Number of artifacts to retain
-`keepDays` | `integer`<br>Number of days to retain artifacts of this type
-`namespaceId` | `string`<br>Storage namespace to use for this artifact types
-
 ## AclConfig
 
@@ -223,28 +120,6 @@
 `excludeActions` | `string[]`<br>Actions to exclude from the inherited actions
 `extends` | `string[]`<br>Other profiles to extend from
 
-## AnalyticsConfig
-
-Config settings for analytics
-
-Name | Description
----- | -----------
-`stores` | [TelemetryStoreConfig](Telemetry.md)`[]`<br>Metrics to aggregate on the Horde server
-
-## ComputeConfig
-
-Configuration for the compute system
-
-Name | Description
----- | -----------
-`acl` | [AclConfig](#aclconfig)<br>Inherited root acl
-`versionEnum` | [ConfigVersion](#configversion-enum)<br>Config version number
-`rates` | [AgentRateConfig](#agentrateconfig)`[]`<br>List of costs of a particular agent type
-`clusters` | [ComputeClusterConfig](#computeclusterconfig)`[]`<br>List of compute profiles
-`pools` | [PoolConfig](#poolconfig)`[]`<br>List of pools
-`software` | [AgentSoftwareConfig](#agentsoftwareconfig)`[]`<br>List of costs of a particular agent type
-`networks` | [NetworkConfig](#networkconfig)`[]`<br>List of networks
-
 ## ConfigVersion (Enum)
 
 Global version number for running the server. As new features are introduced that require data migrations, this version number indicates the backwards compatibility functionality that must be enabled. When adding a new version here, also add a message to ConfigService.CreateSnapshotAsync describing the steps that need to be taken to upgrade the deployment.
@@ -415,11 +290,6 @@
 `description` | `string`<br>Human-readable description
 `computeId` | `string`<br>Compute ID for this network (used when allocating compute resources)
 
-## EmptyPluginConfig
-
-Empty implementation of
-
-
 ## SecretsConfig
 
 Configuration for the secrets system
@@ -459,6 +329,131 @@
 `Text` | Secret is a plain text value which will be stored using the external secret key
 `Json` | Secret is a JSON formatted string containing key/value pairs
 
+## AnalyticsConfig
+
+Config settings for analytics
+
+Name | Description
+---- | -----------
+`stores` | [TelemetryStoreConfig](Telemetry.md)`[]`<br>Metrics to aggregate on the Horde server
+
+## BuildConfig
+
+Configuration for the build plugin
+
+Name | Description
+---- | -----------
+`perforceClusters` | [PerforceCluster](#perforcecluster)`[]`<br>List of Perforce clusters
+`devices` | [DeviceConfig](#deviceconfig)<br>Device configuration
+`maxConformCount` | `integer`<br>Maximum number of conforms to run at once
+`agentShutdownIfDisabledGracePeriod` | `string`<br>Time to wait before shutting down an agent that has been disabled Used if no value is set on the actual pool.
+`artifactTypes` | [ArtifactTypeConfig](#artifacttypeconfig)`[]`<br>Configuration for different artifact types
+`projects` | [ProjectConfig](Projects.md)`[]`<br>List of projects
+`enableConformTasks` | `boolean`<br>Whether to allow conform tasks to run
+`issueFixedTag` | `string`<br>Commit tag to use for marking issues as fixed
+
+## PerforceCluster
+
+Information about a cluster of Perforce servers.
+
+Name | Description
+---- | -----------
+`name` | `string`<br>Name of the cluster
+`serviceAccount` | `string`<br>Username for Horde to log in to this server. Will use the first account specified below if not overridden.
+`canImpersonate` | `boolean`<br>Whether the service account can impersonate other users
+`supportsPartitionedWorkspaces` | `boolean`<br>Whether to use partitioned workspaces on this server
+`servers` | [PerforceServer](#perforceserver)`[]`<br>List of servers
+`credentials` | [PerforceCredentials](#perforcecredentials)`[]`<br>List of server credentials
+`autoSdk` | [AutoSdkWorkspace](#autosdkworkspace)`[]`<br>List of autosdk streams
+
+## PerforceServer
+
+Information about an individual Perforce server
+
+Name | Description
+---- | -----------
+`serverAndPort` | `string`<br>The server and port. The server may be a DNS entry with multiple records, in which case it will be actively load balanced. If "ssl:" prefix is used, ensure P4 server's fingerprint/certificate is trusted. See Horde's documentation on connecting to SSL-enabled Perforce servers.
+`healthCheck` | `boolean`<br>Whether to query the healthcheck address under each server
+`resolveDns` | `boolean`<br>Whether to resolve the DNS entries and load balance between different hosts
+`maxConformCount` | `integer`<br>Maximum number of simultaneous conforms on this server
+`condition` | `string`<br>Optional condition for a machine to be eligible to use this server
+`properties` | `string[]`<br>List of properties for an agent to be eligible to use this server
+
+## PerforceCredentials
+
+Credentials for a Perforce user
+
+Name | Description
+---- | -----------
+`userName` | `string`<br>The username
+`password` | `string`<br>Password for the user
+`ticket` | `string`<br>Login ticket for the user (will be used instead of password if set)
+
+## AutoSdkWorkspace
+
+Path to a platform and stream to use for syncing AutoSDK
+
+Name | Description
+---- | -----------
+`name` | `string`<br>Name of this workspace
+`properties` | `string[]`<br>The agent properties to check (eg. "OSFamily=Windows")
+`userName` | `string`<br>Username for logging in to the server
+`stream` | `string`<br>Stream to use
+
+## DeviceConfig
+
+Configuration for devices
+
+Name | Description
+---- | -----------
+`platforms` | [DevicePlatformConfig](#deviceplatformconfig)`[]`<br>List of device platforms
+`pools` | [DevicePoolConfig](#devicepoolconfig)`[]`<br>List of device pools
+
+## DevicePlatformConfig
+
+Configuration for a device platform
+
+Name | Description
+---- | -----------
+`id` | `string`<br>The id for this platform
+`name` | `string`<br>Name of the platform
+`models` | `string[]`<br>A list of platform models
+`legacyNames` | `string[]`<br>Legacy names which older versions of Gauntlet may be using
+`legacyPerfSpecHighModel` | `string`<br>Model name for the high perf spec, which may be requested by Gauntlet
+
+## DevicePoolConfig
+
+Configuration for a device pool
+
+Name | Description
+---- | -----------
+`id` | `string`<br>The id for this platform
+`name` | `string`<br>The name of the pool
+`poolType` | [DevicePoolType](#devicepooltype-enum)<br>The type of the pool
+`projectIds` | `string[]`<br>List of project ids associated with pool
+
+## DevicePoolType (Enum)
+
+The type of device pool
+
+Name | Description
+---- | -----------
+`Automation` | Available to CIS jobs
+`Shared` | Shared by users with remote checking and checkouts
+
+## ArtifactTypeConfig
+
+Configuration for an artifact
+
+Name | Description
+---- | -----------
+`name` | `string`<br>Legacy 'Name' property
+`type` | `string`<br>Name of the artifact type
+`acl` | [AclConfig](#aclconfig)<br>Acl for the artifact type
+`keepCount` | `integer`<br>Number of artifacts to retain
+`keepDays` | `integer`<br>Number of days to retain artifacts of this type
+`namespaceId` | `string`<br>Storage namespace to use for this artifact types
+
 ## StorageConfig
 
 Configuration for storage
@@ -528,6 +523,25 @@
 `enableAliases` | `boolean`<br>Support querying exports by their aliases
 `acl` | [AclConfig](#aclconfig)<br>Access list for this namespace
 
+## SymbolsConfig
+
+Configuration for the tools system
+
+Name | Description
+---- | -----------
+`stores` | [SymbolStoreConfig](#symbolstoreconfig)`[]`<br>List of symbol stores
+
+## SymbolStoreConfig
+
+Configuration for a symbol store
+
+Name | Description
+---- | -----------
+`id` | `string`<br>Identifier for this store
+`namespaceId` | `string`<br>Configuration for the symbol store backend
+`public` | `boolean`<br>Whether to make this store available without auth
+`acl` | [AclConfig](#aclconfig)<br>Access to the symbol store
+
 ## ToolsConfig
 
 Configuration for the tools system
@@ -556,21 +570,7 @@
 `namespaceId` | `string`<br>Default namespace for new deployments of this tool
 `acl` | [AclConfig](#aclconfig)<br>Permissions for the tool
 
-## SymbolsConfig
+## EmptyPluginConfig
 
-Configuration for the tools system
+Empty implementation of
 
-Name | Description
----- | -----------
-`stores` | [SymbolStoreConfig](#symbolstoreconfig)`[]`<br>List of symbol stores
-
-## SymbolStoreConfig
-
-Configuration for a symbol store
-
-Name | Description
----- | -----------
-`id` | `string`<br>Identifier for this store
-`namespaceId` | `string`<br>Configuration for the symbol store backend
-`public` | `boolean`<br>Whether to make this store available without auth
-`acl` | [AclConfig](#aclconfig)<br>Access to the symbol store
diff --git a/Engine/Source/Programs/Horde/Docs/Deployment/ServerSettings.md b/Engine/Source/Programs/Horde/Docs/Deployment/ServerSettings.md
--- a/Engine/Source/Programs/Horde/Docs/Deployment/ServerSettings.md
+++ b/Engine/Source/Programs/Horde/Docs/Deployment/ServerSettings.md
@@ -123,14 +123,104 @@
 
 Name | Description
 ---- | -----------
-`build` | [BuildServerConfig](#buildserverconfig)<br>Configuration for the build plugin
-`analytics` | [AnalyticsServerConfig](#analyticsserverconfig)<br>Configuration for the analytics plugin
 `compute` | [ComputeServerConfig](#computeserverconfig)<br>Configuration for the compute plugin
-`ddc` | [PluginServerConfig](#pluginserverconfig)<br>Configuration for the ddc plugin
 `secrets` | [PluginServerConfig](#pluginserverconfig)<br>Configuration for the secrets plugin
+`analytics` | [AnalyticsServerConfig](#analyticsserverconfig)<br>Configuration for the analytics plugin
+`build` | [BuildServerConfig](#buildserverconfig)<br>Configuration for the build plugin
 `storage` | [StorageServerConfig](#storageserverconfig)<br>Configuration for the storage plugin
-`tools` | [ToolsServerConfig](#toolsserverconfig)<br>Configuration for the tools plugin
 `symbols` | [PluginServerConfig](#pluginserverconfig)<br>Configuration for the symbols plugin
+`tools` | [ToolsServerConfig](#toolsserverconfig)<br>Configuration for the tools plugin
+`ddc` | [PluginServerConfig](#pluginserverconfig)<br>Configuration for the ddc plugin
 
+## ComputeServerConfig
+
+Static configuration for the compute plugin
+
+Name | Description
+---- | -----------
+`enableUpgradeTasks` | `boolean`<br>Whether to enable the upgrade task source.
+`withAws` | `boolean`<br>Whether to enable Amazon Web Services (AWS) specific features
+`awsRegions` | `string[]`<br>List of AWS regions for Horde to be aware of (e.g. us-east-1 or eu-central-1) Right now, this is only used for replicating CloudWatch metrics to multiple regions
+`awsAutoScalingQueueUrls` | `string[]`<br>AWS SQS queue URLs where lifecycle events from EC2 auto-scaling are received
+`fleetManagerV2` | [FleetManagerType](#fleetmanagertype-enum)<br>Default fleet manager to use (when not specified by pool)
+`fleetManagerV2Config` | `object`<br>Config for the fleet manager (serialized JSON)
+`autoEnrollAgents` | `boolean`<br>Whether to automatically enroll agents in the farm
+`defaultAgentPoolSizeStrategy` | [PoolSizeStrategy](#poolsizestrategy-enum)<br>Default agent pool sizing strategy for pools that doesn't have one explicitly configured
+`agentPoolScaleOutCooldownSeconds` | `integer`<br>Scale-out cooldown for auto-scaling agent pools (in seconds). Can be overridden by per-pool settings.
+`agentPoolScaleInCooldownSeconds` | `integer`<br>Scale-in cooldown for auto-scaling agent pools (in seconds). Can be overridden by per-pool settings.
+`computeTunnelPort` | `integer`<br>Port to listen on for tunneling compute sockets to agents
+`computeTunnelAddress` | `string`<br>What address (host:port) clients should connect to for compute socket tunneling Port may differ from  if Horde server is behind a reverse proxy/firewall
+`enabled` | `boolean`<br>Whether the plugin should be enabled or not
+
+## FleetManagerType (Enum)
+
+Available fleet managers
+
+Name | Description
+---- | -----------
+`Default` | Default fleet manager
+`NoOp` | No-op fleet manager.
+`Aws` | Fleet manager for handling AWS EC2 instances. Will create and/or terminate instances from scratch.
+`AwsReuse` | Fleet manager for handling AWS EC2 instances. Will start already existing but stopped instances to reuse existing EBS disks.
+`AwsRecycle` | Fleet manager for handling AWS EC2 instances. Will start already existing but stopped instances to reuse existing EBS disks.
+`AwsAsg` | Fleet manager for handling AWS EC2 instances. Uses an EC2 auto-scaling group for controlling the number of running instances.
+
+## PoolSizeStrategy (Enum)
+
+Available pool sizing strategies
+
+Name | Description
+---- | -----------
+`LeaseUtilization` | Strategy based on lease utilization
+`JobQueue` | Strategy based on size of job build queue
+`NoOp` | No-op strategy used as fallback/default behavior
+`ComputeQueueAwsMetric` | A no-op strategy that reports metrics to let an external AWS auto-scaling policy scale the fleet
+`LeaseUtilizationAwsMetric` | A no-op strategy that reports metrics to let an external AWS auto-scaling policy scale the fleet
+
+## PluginServerConfig
+
+Base class for plugin server config objects
+
+Name | Description
+---- | -----------
+`enabled` | `boolean`<br>Whether the plugin should be enabled or not
+
+## AnalyticsServerConfig
+
+Server configuration for the analytics system
+
+Name | Description
+---- | -----------
+`sinks` | [TelemetrySinkConfig](#telemetrysinkconfig)<br>Settings for the various telemetry sinks
+`enabled` | `boolean`<br>Whether the plugin should be enabled or not
+
+## TelemetrySinkConfig
+
+Telemetry sinks
+
+Name | Description
+---- | -----------
+`epic` | [EpicTelemetryConfig](#epictelemetryconfig)<br>Settings for the Epic telemetry sink
+`mongo` | [MongoTelemetryConfig](#mongotelemetryconfig)<br>Settings for the MongoDB telemetry sink
+
+## EpicTelemetryConfig
+
+Configuration for the telemetry sink
+
+Name | Description
+---- | -----------
+`url` | `string`<br>Base URL for the telemetry server
+`appId` | `string`<br>Application name to send in the event messages
+`enabled` | `boolean`<br>Whether to enable this sink
+
+## MongoTelemetryConfig
+
+Configuration for the telemetry sink
+
+Name | Description
+---- | -----------
+`retainDays` | `number`<br>Number of days worth of telmetry events to keep
+`enabled` | `boolean`<br>Whether to enable this sink
+
 ## BuildServerConfig
 
@@ -270,96 +360,6 @@
 `maxChildCount` | `integer`<br>Maximum number of children in each node
 `sliceThreshold` | `integer`<br>Threshold hash value for splitting interior nodes
 
-## AnalyticsServerConfig
-
-Server configuration for the analytics system
-
-Name | Description
----- | -----------
-`sinks` | [TelemetrySinkConfig](#telemetrysinkconfig)<br>Settings for the various telemetry sinks
-`enabled` | `boolean`<br>Whether the plugin should be enabled or not
-
-## TelemetrySinkConfig
-
-Telemetry sinks
-
-Name | Description
----- | -----------
-`epic` | [EpicTelemetryConfig](#epictelemetryconfig)<br>Settings for the Epic telemetry sink
-`mongo` | [MongoTelemetryConfig](#mongotelemetryconfig)<br>Settings for the MongoDB telemetry sink
-
-## EpicTelemetryConfig
-
-Configuration for the telemetry sink
-
-Name | Description
----- | -----------
-`url` | `string`<br>Base URL for the telemetry server
-`appId` | `string`<br>Application name to send in the event messages
-`enabled` | `boolean`<br>Whether to enable this sink
-
-## MongoTelemetryConfig
-
-Configuration for the telemetry sink
-
-Name | Description
----- | -----------
-`retainDays` | `number`<br>Number of days worth of telmetry events to keep
-`enabled` | `boolean`<br>Whether to enable this sink
-
-## ComputeServerConfig
-
-Static configuration for the compute plugin
-
-Name | Description
----- | -----------
-`enableUpgradeTasks` | `boolean`<br>Whether to enable the upgrade task source.
-`withAws` | `boolean`<br>Whether to enable Amazon Web Services (AWS) specific features
-`awsRegions` | `string[]`<br>List of AWS regions for Horde to be aware of (e.g. us-east-1 or eu-central-1) Right now, this is only used for replicating CloudWatch metrics to multiple regions
-`awsAutoScalingQueueUrls` | `string[]`<br>AWS SQS queue URLs where lifecycle events from EC2 auto-scaling are received
-`fleetManagerV2` | [FleetManagerType](#fleetmanagertype-enum)<br>Default fleet manager to use (when not specified by pool)
-`fleetManagerV2Config` | `object`<br>Config for the fleet manager (serialized JSON)
-`autoEnrollAgents` | `boolean`<br>Whether to automatically enroll agents in the farm
-`defaultAgentPoolSizeStrategy` | [PoolSizeStrategy](#poolsizestrategy-enum)<br>Default agent pool sizing strategy for pools that doesn't have one explicitly configured
-`agentPoolScaleOutCooldownSeconds` | `integer`<br>Scale-out cooldown for auto-scaling agent pools (in seconds). Can be overridden by per-pool settings.
-`agentPoolScaleInCooldownSeconds` | `integer`<br>Scale-in cooldown for auto-scaling agent pools (in seconds). Can be overridden by per-pool settings.
-`computeTunnelPort` | `integer`<br>Port to listen on for tunneling compute sockets to agents
-`computeTunnelAddress` | `string`<br>What address (host:port) clients should connect to for compute socket tunneling Port may differ from  if Horde server is behind a reverse proxy/firewall
-`enabled` | `boolean`<br>Whether the plugin should be enabled or not
-
-## FleetManagerType (Enum)
-
-Available fleet managers
-
-Name | Description
----- | -----------
-`Default` | Default fleet manager
-`NoOp` | No-op fleet manager.
-`Aws` | Fleet manager for handling AWS EC2 instances. Will create and/or terminate instances from scratch.
-`AwsReuse` | Fleet manager for handling AWS EC2 instances. Will start already existing but stopped instances to reuse existing EBS disks.
-`AwsRecycle` | Fleet manager for handling AWS EC2 instances. Will start already existing but stopped instances to reuse existing EBS disks.
-`AwsAsg` | Fleet manager for handling AWS EC2 instances. Uses an EC2 auto-scaling group for controlling the number of running instances.
-
-## PoolSizeStrategy (Enum)
-
-Available pool sizing strategies
-
-Name | Description
----- | -----------
-`LeaseUtilization` | Strategy based on lease utilization
-`JobQueue` | Strategy based on size of job build queue
-`NoOp` | No-op strategy used as fallback/default behavior
-`ComputeQueueAwsMetric` | A no-op strategy that reports metrics to let an external AWS auto-scaling policy scale the fleet
-`LeaseUtilizationAwsMetric` | A no-op strategy that reports metrics to let an external AWS auto-scaling policy scale the fleet
-
-## PluginServerConfig
-
-Base class for plugin server config objects
-
-Name | Description
----- | -----------
-`enabled` | `boolean`<br>Whether the plugin should be enabled or not
-
 ## StorageServerConfig
 
 Static settings for the storage system
