Horde: Fix some debug info being written at warning verbosity.

#rnx

diff --git a/Engine/Source/Programs/Horde/Plugins/Storage/HordeServer.Storage/Storage/StorageService.cs b/Engine/Source/Programs/Horde/Plugins/Storage/HordeServer.Storage/Storage/StorageService.cs
--- a/Engine/Source/Programs/Horde/Plugins/Storage/HordeServer.Storage/Storage/StorageService.cs
+++ b/Engine/Source/Programs/Horde/Plugins/Storage/HordeServer.Storage/Storage/StorageService.cs
@@ -681,7 +681,7 @@ namespace HordeServer.Storage
 			}
 
 			await _blobCollection.InsertOneAsync(blobInfo, new InsertOneOptions { }, cancellationToken);
-			_logger.LogWarning("Created blob {BlobId} at {Path} ({NumImports} imports)", blobInfo.Id, blobInfo.Path, blobInfo.Imports.Count);
+			_logger.LogDebug("Created blob {BlobId} at {Path} ({NumImports} imports)", blobInfo.Id, blobInfo.Path, blobInfo.Imports.Count);
 		}
 
 		async ValueTask CheckBlobExistsAsync(NamespaceId namespaceId, BlobLocator locator, CancellationToken cancellationToken)
