Horde: Fix static analysis warning

diff --git a/Engine/Source/Programs/Horde/Plugins/Build/HordeServer.Build/Artifacts/ArtifactsController.cs b/Engine/Source/Programs/Horde/Plugins/Build/HordeServer.Build/Artifacts/ArtifactsController.cs
--- a/Engine/Source/Programs/Horde/Plugins/Build/HordeServer.Build/Artifacts/ArtifactsController.cs
+++ b/Engine/Source/Programs/Horde/Plugins/Build/HordeServer.Build/Artifacts/ArtifactsController.cs
@@ -5,7 +5,6 @@ using System.Buffers.Binary;
 using System.Diagnostics;
 using System.Net;
 using System.Net.Mime;
-using System.Text;
 using System.Text.RegularExpressions;
 using EpicGames.Core;
 using EpicGames.Horde.Acls;
