Horde: Update release notes

#rnx

diff --git a/Engine/Source/Programs/Horde/Docs/ReleaseNotes.md b/Engine/Source/Programs/Horde/Docs/ReleaseNotes.md
--- a/Engine/Source/Programs/Horde/Docs/ReleaseNotes.md
+++ b/Engine/Source/Programs/Horde/Docs/ReleaseNotes.md
@@ -1,5 +1,44 @@
 # Release Notes
 
+## 2024-10-10
+
+* EpicGames.Perforce: Assume and mark workspace having untracked files prior to starting a sync (36972802)
+* Allow HTTP for OIDC discovery endpoints
+* Log errors when mongodb config/log files contain non-latin characters as mongodb.exe cannot handle these paths (36930465)
+* Improve description text for several tools bundled with Horde. Mention that P4V needs to be restarted after installing P4VUtils. (36912499)
+* Use P4 syntax on artifact modal to download current folder (36912473)
+* Fix issues with theming and login/setup view for installer build (36907272)
+* Don't add all files to filter when there is no selection in the artifact modal to avoid overflowing the max request length (36906929)
+* Add Horde's built-in user ID claim when issuing OIDC userinfo for Horde accounts (36899176)
+* Ensure auto-conforming is requested only when not already pending (36894611)
+* Distinguish between no matching compute resources vs all resources in use (36857071)
+* Allow resetting blob id for length scan. (36848920)
+* Add schedule auditing (36825106)
+* Prevent infinite loop when cancelling during agent registration (36819203)
+* Allow clients to specify a locator when uploading blobs. The server will not allow overwriting any existing blob, but allowing the client to determine the locator opens the door to it being able to write multiple blobs in parallel. (36813875)
+* Grant claims to leases identifying them as running certain projects, streams, and templates. (36807664)
+* Support running jobs in streams that do not have the engine directly under the stream root. The "EnginePath" property in the stream config can be used to configure the path to the engine folder. (36804827)
+* Include a Version.json file in Unreal Toolbox builds which can be used to detect upgrades/downgrades. When auto-updating from builds with this file, only upgrades will be allowed. (36789633)
+* Unreal Toolbox: Always show the settings dialog on launch unless the -Quiet argument is passed on the command line. This makes a more sensible default than requiring a -Settings argument. (36785739)
+* Remove agent enrollment server json upon uninstall (36784558)
+* Fix agent store updates upon cache invalidation/deletes (36739556)
+* Move UnrealBuildTool.Tests into Engine\Programs\Shared (36735944)
+* Fix issues with agent enrollment and deletion latency (36735037)
+* Fix bug where leases weren't immediately aborted when agent is marked as busy (36733795)
+* Fix issue with artfifact log rendering when there is no artifact type (36730128)
+* Fix AgentWorkspaceInfo objects constructed from RPC workspace messages having an empty string for 'method' instead of null, causing agents to get stuck in a conform loop. (36718257)
+* Fix issue with enrolling agent from installer goes to default server (36710606)
+* Fixed a bug where the same error/warning message using different slash directions would not match with an existing issue when handled by the hashed issue handler (36698600)
+* Add more logging and tracing for updating agents via REST API (36697817)
+* Allow specifying arbitary key/value metadata on tools, and add product ids for MSI installers. Unreal Toolbox will automatically install/uninstall MSIs using the given product id. (36694744)
+* Unreal Toolbox: Various improvements with installation process (36678430)
+* Fix bug where not all fields were included in Equals() check for AgentWorkspaceInfo (36667931)
+* Add scratch and min conform space to REST API agent response (36667252)
+* Fix issue with mixed ungrouped and grouped tool downloads (36647304)
+* Use the correct token for writing artifacts. Adds an IArtifactBuilder interface returned from IArtifactCollection.CreateAsync() which can be used to upload artifact data with the correct permissions. (36645064)
+* Resurrect option to explicitly update build health issues in advanced parameters (36644246)
+* Use same logic for shutting down on MacOS as Linux. (36578059)
+
 ## 2024-09-25
 
 * Add filter for agent properties in REST API (36563629)
