Horde: Check permissions for uploading blobs against the path from the specified locator, rather than prefix field in the upload requests. Fixes direct uploads using user-specified paths.

#rnx

diff --git a/Engine/Source/Programs/Horde/Plugins/Storage/HordeServer.Storage/Storage/StorageController.cs b/Engine/Source/Programs/Horde/Plugins/Storage/HordeServer.Storage/Storage/StorageController.cs
--- a/Engine/Source/Programs/Horde/Plugins/Storage/HordeServer.Storage/Storage/StorageController.cs
+++ b/Engine/Source/Programs/Horde/Plugins/Storage/HordeServer.Storage/Storage/StorageController.cs
@@ -124,7 +124,7 @@ namespace HordeServer.Storage
 			{
 				return NotFound(namespaceId);
 			}
-			if (!Authorize(namespaceId, StorageAclAction.WriteBlobs) && !HasPathClaim(User, HordeClaimTypes.WriteNamespace, namespaceId, request.Prefix ?? String.Empty))
+			if (!Authorize(namespaceId, StorageAclAction.WriteBlobs) && !HasPathClaim(User, HordeClaimTypes.WriteNamespace, namespaceId, GetPathFromLocator(locator)))
 			{
 				return Forbid(StorageAclAction.WriteBlobs, namespaceId);
 			}
@@ -132,6 +132,19 @@ namespace HordeServer.Storage
 			return await WriteBlobAsync(storageBackend, locator, request, cancellationToken);
 		}
 
+		static string GetPathFromLocator(BlobLocator locator)
+		{
+			int lastIdx = locator.Path.LastIndexOf('/');
+			if (lastIdx == -1)
+			{
+				return String.Empty;
+			}
+			else
+			{
+				return locator.Path.Substring(0, lastIdx).ToString();
+			}
+		}
+
 		/// <summary>
 		/// Writes a blob to storage. Exposed as a public utility method to allow other routes with their own authentication methods to wrap their own authentication/redirection.
 		/// </summary>
