Horde: Fix parameter teimming on new builds, and some layout tweaks

diff --git a/Engine/Source/Programs/Horde/HordeDashboard/src/components/SchedulePane.tsx b/Engine/Source/Programs/Horde/HordeDashboard/src/components/SchedulePane.tsx
--- a/Engine/Source/Programs/Horde/HordeDashboard/src/components/SchedulePane.tsx
+++ b/Engine/Source/Programs/Horde/HordeDashboard/src/components/SchedulePane.tsx
@@ -77,7 +77,7 @@ export const SchedulePane: React.FC<{ streamId: string, templates: GetTemplateRe
       { key: 'schedule_column7', name: 'time_4', minWidth: 100, maxWidth: 100, isResizable: false },
       { key: 'schedule_column8', name: 'time_5', minWidth: 100, maxWidth: 100, isResizable: false },
       { key: 'schedule_column9', name: 'time_6', minWidth: 100, maxWidth: 100, isResizable: false },
-      { key: 'schedule_column10', name: 'audit', minWidth: 64, maxWidth: 64, isResizable: false },
+      { key: 'schedule_column10', name: 'audit', minWidth: 56, maxWidth: 56, isResizable: false },
       { key: 'schedule_column11', name: 'space', minWidth: 2, isResizable: false }
 
    ];
diff --git a/Engine/Source/Programs/Horde/HordeDashboard/src/components/build/NewBuildV2.tsx b/Engine/Source/Programs/Horde/HordeDashboard/src/components/build/NewBuildV2.tsx
--- a/Engine/Source/Programs/Horde/HordeDashboard/src/components/build/NewBuildV2.tsx
+++ b/Engine/Source/Programs/Horde/HordeDashboard/src/components/build/NewBuildV2.tsx
@@ -1828,6 +1828,11 @@ const BuildModal: React.FC<{ setUseLegacyDialog: (value: boolean) => void }> = o
          options.advAdditionalArgs.trim().replaceAll("\n", " ").match(argRegex)?.forEach(arg => additionalArgs.push(arg.replace(/"/g, "")));
       }
 
+      // trim all parameters
+      const parameters: Record<string, string> = {};
+      for (const [id, value] of Object.entries(options.parameters ?? {})) {
+         parameters[id?.trim()] = value?.trim()
+      }
 
       const data: CreateJobRequest = {
          streamId: options.streamId,
@@ -1835,7 +1840,7 @@ const BuildModal: React.FC<{ setUseLegacyDialog: (value: boolean) => void }> = o
          name: options.advJobName,
          priority: options.advJobPriority,
          changeQueries: changeQueries,
-         parameters: options.parameters,
+         parameters: parameters,
          additionalArguments: additionalArgs?.length ? additionalArgs : undefined,
          targets: options.advTargets,
          updateIssues: options.updateIssues
