Horde: Allow closing multiple issues by separating issue numbers with commas (eg. "#horde X, Y")

#rnx

diff --git a/Engine/Source/Programs/Horde/Plugins/Build/HordeServer.Build.Tests/Issues/IssueTagTests.cs b/Engine/Source/Programs/Horde/Plugins/Build/HordeServer.Build.Tests/Issues/IssueTagTests.cs
--- a/Engine/Source/Programs/Horde/Plugins/Build/HordeServer.Build.Tests/Issues/IssueTagTests.cs
+++ b/Engine/Source/Programs/Horde/Plugins/Build/HordeServer.Build.Tests/Issues/IssueTagTests.cs
@@ -15,6 +15,8 @@ namespace HordeServer.Tests.Issues
 			Assert.IsTrue(IssueTagService.ParseTags("#horde", "hello\n #horde 123 ").SequenceEqual(new[] { 123 }));
 			Assert.IsTrue(IssueTagService.ParseTags("#horde", "hello\n #horde 123 \n#horde 456").SequenceEqual(new[] { 123, 456 }));
 			Assert.IsTrue(IssueTagService.ParseTags("#horde", "hello\n #horde 123 \n #ROBOMERGE-SOURCE foo").SequenceEqual(Array.Empty<int>()));
+			Assert.IsTrue(IssueTagService.ParseTags("#horde", "hello\n #horde 123, 456").SequenceEqual(new[] { 123, 456 }));
+			Assert.IsTrue(IssueTagService.ParseTags("#horde", "hello\n #horde 123 , garbage 456").SequenceEqual(new[] { 123, 456 }));
 		}
 	}
 }
diff --git a/Engine/Source/Programs/Horde/Plugins/Build/HordeServer.Build/Issues/IssueTagService.cs b/Engine/Source/Programs/Horde/Plugins/Build/HordeServer.Build/Issues/IssueTagService.cs
--- a/Engine/Source/Programs/Horde/Plugins/Build/HordeServer.Build/Issues/IssueTagService.cs
+++ b/Engine/Source/Programs/Horde/Plugins/Build/HordeServer.Build/Issues/IssueTagService.cs
@@ -159,7 +159,7 @@ namespace HordeServer.Issues
 			{
 				foreach (Match match in Regex.Matches(description, $"^\\s*{issueFixedTag}\\s+(.*)$", RegexOptions.Multiline))
 				{
-					string[] issues = match.Groups[1].Value.Split(new[] { ' ', '\t' }, StringSplitOptions.RemoveEmptyEntries);
+					string[] issues = match.Groups[1].Value.Split(new[] { ' ', '\t', ',' }, StringSplitOptions.TrimEntries | StringSplitOptions.RemoveEmptyEntries);
 					foreach (string issue in issues)
 					{
 						if (Int32.TryParse(issue, System.Globalization.NumberStyles.None, null, out int issueId))
