Horde: Restore logging for triggering re-conform.

#rnx

diff --git a/Engine/Source/Programs/Horde/Plugins/Build/HordeServer.Build/Jobs/JobRpcService.cs b/Engine/Source/Programs/Horde/Plugins/Build/HordeServer.Build/Jobs/JobRpcService.cs
--- a/Engine/Source/Programs/Horde/Plugins/Build/HordeServer.Build/Jobs/JobRpcService.cs
+++ b/Engine/Source/Programs/Horde/Plugins/Build/HordeServer.Build/Jobs/JobRpcService.cs
@@ -99,7 +99,8 @@ namespace HordeServer.Jobs
 
 				// Get the set of workspaces that are currently required
 				HashSet<AgentWorkspaceInfo> conformWorkspaces = await _poolService.GetWorkspacesAsync(agent, DateTime.UtcNow, _buildConfig.Value, context.CancellationToken);
-				bool pendingConform = !conformWorkspaces.SetEquals(newWorkspaces) || (agent.RequestFullConform && !request.RemoveUntrackedFiles);
+				bool workspaceSetEquals = conformWorkspaces.SetEquals(newWorkspaces);
+				bool pendingConform = !workspaceSetEquals || (agent.RequestFullConform && !request.RemoveUntrackedFiles);
 
 				// Update the workspaces
 				if (await _agentService.TryUpdateWorkspacesAsync(agent, newWorkspaces, pendingConform, context.CancellationToken))
@@ -109,6 +110,30 @@ namespace HordeServer.Jobs
 					{
 						response.Retry = await _conformTaskSource.GetWorkspacesAsync(agent, response.PendingWorkspaces, context.CancellationToken);
 						response.RemoveUntrackedFiles = request.RemoveUntrackedFiles || agent.RequestFullConform;
+
+						if (response.Retry)
+						{
+							HashSet<string> identifiers = new HashSet<string>();
+							identifiers.UnionWith(conformWorkspaces.Select(x => x.Identifier));
+							identifiers.UnionWith(newWorkspaces.Select(x => x.Identifier));
+
+							int numChanged = 0;
+							foreach (string identifier in identifiers)
+							{
+								AgentWorkspaceInfo? prevWorkspace = conformWorkspaces.FirstOrDefault(x => x.Identifier.Equals(identifier, StringComparison.OrdinalIgnoreCase));
+								AgentWorkspaceInfo? nextWorkspace = newWorkspaces.FirstOrDefault(x => x.Identifier.Equals(identifier, StringComparison.OrdinalIgnoreCase));
+								string prevWorkspaceText = (prevWorkspace == null)? "(does not exist)" : JsonSerializer.Serialize(prevWorkspace);
+								string nextWorkspaceText = (nextWorkspace == null)? "(does not exist)" : JsonSerializer.Serialize(nextWorkspace);
+
+								if (!String.Equals(prevWorkspaceText, nextWorkspaceText))
+								{
+									_logger.LogInformation("{AgentId} {Identifier} changed: {OldWorkspace} -> {NewWorkspace}", agent.Id, identifier, prevWorkspaceText, nextWorkspaceText);
+									numChanged++;
+								}
+							}
+
+							_logger.LogInformation("Retrying conform for {AgentId} ({NumChanged} changed, set equals: {SetEquals}, request conform: {RequestConform}, request full conform: {RequestFullConform}, remove untracked: {RemoveUntracked})", agent.Id, numChanged, workspaceSetEquals, agent.RequestConform, agent.RequestFullConform, request.RemoveUntrackedFiles);
+						}
 					}
 					return response;
 				}
