Horde: Fix incorrect path being used when uploading blobs through PUT method.

#rnx

diff --git a/Engine/Source/Programs/Horde/Plugins/Storage/HordeServer.Storage/Storage/StorageController.cs b/Engine/Source/Programs/Horde/Plugins/Storage/HordeServer.Storage/Storage/StorageController.cs
--- a/Engine/Source/Programs/Horde/Plugins/Storage/HordeServer.Storage/Storage/StorageController.cs
+++ b/Engine/Source/Programs/Horde/Plugins/Storage/HordeServer.Storage/Storage/StorageController.cs
@@ -158,13 +158,30 @@ namespace HordeServer.Storage
 			IReadOnlyCollection<BlobLocator> imports = request.Imports ?? (IReadOnlyCollection<BlobLocator>)Array.Empty<BlobLocator>();
 			if (request.File == null)
 			{
-				(BlobLocator Path, Uri UploadUrl)? result = await storageBackend.TryGetBlobWriteRedirectAsync(imports, request.Prefix ?? String.Empty, cancellationToken);
-				if (result == null)
+				if (locator == null)
 				{
-					return new WriteBlobResponse { SupportsRedirects = false };
+					(BlobLocator Locator, Uri UploadUrl)? result = await storageBackend.TryGetBlobWriteRedirectAsync(imports, request.Prefix ?? String.Empty, cancellationToken);
+					if (result == null)
+					{
+						return new WriteBlobResponse { SupportsRedirects = false };
+					}
+					else
+					{
+						return new WriteBlobResponse { Blob = result.Value.Locator.Path.ToString(), UploadUrl = result.Value.UploadUrl };
+					}
+				}
+				else
+				{
+					Uri? uploadUrl = await storageBackend.TryGetBlobWriteRedirectAsync(locator.Value, imports, cancellationToken);
+					if (uploadUrl == null)
+					{
+						return new WriteBlobResponse { SupportsRedirects = false };
+					}
+					else
+					{
+						return new WriteBlobResponse { Blob = locator.Value.Path.ToString(), UploadUrl = uploadUrl };
+					}
 				}
-
-				return new WriteBlobResponse { Blob = result.Value.Path.ToString(), UploadUrl = result.Value.UploadUrl };
 			}
 			else
 			{
