Horde: Fix static analysis warning from string comparison.

#rnx

diff --git a/Engine/Source/Programs/Horde/Plugins/Build/HordeServer.Build/Jobs/JobRpcService.cs b/Engine/Source/Programs/Horde/Plugins/Build/HordeServer.Build/Jobs/JobRpcService.cs
--- a/Engine/Source/Programs/Horde/Plugins/Build/HordeServer.Build/Jobs/JobRpcService.cs
+++ b/Engine/Source/Programs/Horde/Plugins/Build/HordeServer.Build/Jobs/JobRpcService.cs
@@ -125,7 +125,7 @@ namespace HordeServer.Jobs
 								string prevWorkspaceText = (prevWorkspace == null)? "(does not exist)" : JsonSerializer.Serialize(prevWorkspace);
 								string nextWorkspaceText = (nextWorkspace == null)? "(does not exist)" : JsonSerializer.Serialize(nextWorkspace);
 
-								if (!String.Equals(prevWorkspaceText, nextWorkspaceText))
+								if (!String.Equals(prevWorkspaceText, nextWorkspaceText, StringComparison.Ordinal))
 								{
 									_logger.LogInformation("{AgentId} {Identifier} changed: {OldWorkspace} -> {NewWorkspace}", agent.Id, identifier, prevWorkspaceText, nextWorkspaceText);
 									numChanged++;
