Horde: Fix conform loop due to agent messages not containing MinScratchSpace and ConformDiskFreeSpace properties.

#rnx

diff --git a/Engine/Source/Programs/Horde/Plugins/Build/HordeServer.Build/Jobs/JobRpcService.cs b/Engine/Source/Programs/Horde/Plugins/Build/HordeServer.Build/Jobs/JobRpcService.cs
--- a/Engine/Source/Programs/Horde/Plugins/Build/HordeServer.Build/Jobs/JobRpcService.cs
+++ b/Engine/Source/Programs/Horde/Plugins/Build/HordeServer.Build/Jobs/JobRpcService.cs
@@ -95,15 +95,15 @@ namespace HordeServer.Jobs
 				}
 
 				// Get the new workspaces
-				List<AgentWorkspaceInfo> newWorkspaces = request.Workspaces.Select(x => new AgentWorkspaceInfo(x)).ToList();
+				List<AgentWorkspaceInfo> currentWorkspaces = request.Workspaces.Select(x => new AgentWorkspaceInfo(x)).ToList();
 
 				// Get the set of workspaces that are currently required
 				HashSet<AgentWorkspaceInfo> conformWorkspaces = await _poolService.GetWorkspacesAsync(agent, DateTime.UtcNow, _buildConfig.Value, context.CancellationToken);
-				bool workspaceSetEquals = conformWorkspaces.SetEquals(newWorkspaces);
+				bool workspaceSetEquals = conformWorkspaces.SetEquals(currentWorkspaces);
 				bool pendingConform = !workspaceSetEquals || (agent.RequestFullConform && !request.RemoveUntrackedFiles);
 
 				// Update the workspaces
-				if (await _agentService.TryUpdateWorkspacesAsync(agent, newWorkspaces, pendingConform, context.CancellationToken))
+				if (await _agentService.TryUpdateWorkspacesAsync(agent, currentWorkspaces, pendingConform, context.CancellationToken))
 				{
 					RpcUpdateAgentWorkspacesResponse response = new RpcUpdateAgentWorkspacesResponse();
 					if (pendingConform)
@@ -114,20 +114,21 @@ namespace HordeServer.Jobs
 						if (response.Retry)
 						{
 							HashSet<string> identifiers = new HashSet<string>();
+							identifiers.UnionWith(currentWorkspaces.Select(x => x.Identifier));
 							identifiers.UnionWith(conformWorkspaces.Select(x => x.Identifier));
-							identifiers.UnionWith(newWorkspaces.Select(x => x.Identifier));
 
 							int numChanged = 0;
 							foreach (string identifier in identifiers)
 							{
-								AgentWorkspaceInfo? prevWorkspace = conformWorkspaces.FirstOrDefault(x => x.Identifier.Equals(identifier, StringComparison.OrdinalIgnoreCase));
-								AgentWorkspaceInfo? nextWorkspace = newWorkspaces.FirstOrDefault(x => x.Identifier.Equals(identifier, StringComparison.OrdinalIgnoreCase));
+								AgentWorkspaceInfo? prevWorkspace = currentWorkspaces.FirstOrDefault(x => x.Identifier.Equals(identifier, StringComparison.OrdinalIgnoreCase));
+								AgentWorkspaceInfo? nextWorkspace = conformWorkspaces.FirstOrDefault(x => x.Identifier.Equals(identifier, StringComparison.OrdinalIgnoreCase));
 								string prevWorkspaceText = (prevWorkspace == null)? "(does not exist)" : JsonSerializer.Serialize(prevWorkspace);
 								string nextWorkspaceText = (nextWorkspace == null)? "(does not exist)" : JsonSerializer.Serialize(nextWorkspace);
 
 								if (!String.Equals(prevWorkspaceText, nextWorkspaceText, StringComparison.Ordinal))
 								{
-									_logger.LogInformation("{AgentId} {Identifier} changed: {OldWorkspace} -> {NewWorkspace}", agent.Id, identifier, prevWorkspaceText, nextWorkspaceText);
+									_logger.LogInformation("{AgentId} {Identifier} was: {OldWorkspace}", agent.Id, identifier, prevWorkspaceText);
+									_logger.LogInformation("{AgentId} {Identifier} now: {NewWorkspace}", agent.Id, identifier, nextWorkspaceText);
 									numChanged++;
 								}
 							}
diff --git a/Engine/Source/Programs/Horde/Plugins/Build/HordeServer.Build/Perforce/PerforceWorkspace.cs b/Engine/Source/Programs/Horde/Plugins/Build/HordeServer.Build/Perforce/PerforceWorkspace.cs
--- a/Engine/Source/Programs/Horde/Plugins/Build/HordeServer.Build/Perforce/PerforceWorkspace.cs
+++ b/Engine/Source/Programs/Horde/Plugins/Build/HordeServer.Build/Perforce/PerforceWorkspace.cs
@@ -126,7 +126,9 @@ namespace HordeServer.Perforce
 				Stream = workspace.Stream,
 				Incremental = workspace.Incremental,
 				Partitioned = partitioned,
-				Method = workspace.Method ?? String.Empty
+				Method = workspace.Method ?? String.Empty,
+				MinScratchSpace = workspace.MinScratchSpace ?? 0,
+				ConformDiskFreeSpace = workspace.ConformDiskFreeSpace ?? 0,
 			};
 
 			if (workspace.View != null)
