Update more nuget packages (all system or ms that don't depend on System.Drawing.Common)

#rnx

diff --git a/Engine/Platforms/VisionOS/Source/Programs/AutomationTool/VisionOS.Automation.csproj b/Engine/Platforms/VisionOS/Source/Programs/AutomationTool/VisionOS.Automation.csproj
--- a/Engine/Platforms/VisionOS/Source/Programs/AutomationTool/VisionOS.Automation.csproj
+++ b/Engine/Platforms/VisionOS/Source/Programs/AutomationTool/VisionOS.Automation.csproj
@@ -50,8 +50,6 @@
   </ItemGroup>
   <ItemGroup>
     <PackageReference Include="Microsoft.CSharp" Version="4.7.0" />
-    <PackageReference Include="System.Data.DataSetExtensions" Version="4.5.0" />
-    <PackageReference Include="System.Net.Http" Version="4.3.4" />
     <PackageReference Include="System.Drawing.Common" Version="4.7.3" />
   </ItemGroup>
 </Project>
\ No newline at end of file
diff --git a/Engine/Source/Programs/AutomationTool.Tests/AutomationTool.Tests.csproj b/Engine/Source/Programs/AutomationTool.Tests/AutomationTool.Tests.csproj
--- a/Engine/Source/Programs/AutomationTool.Tests/AutomationTool.Tests.csproj
+++ b/Engine/Source/Programs/AutomationTool.Tests/AutomationTool.Tests.csproj
@@ -11,8 +11,8 @@
 
   <ItemGroup>
     <PackageReference Include="Microsoft.NET.Test.Sdk" Version="17.11.1" />
-    <PackageReference Include="MSTest.TestAdapter" Version="3.6.0" />
-    <PackageReference Include="MSTest.TestFramework" Version="3.6.0" />
+    <PackageReference Include="MSTest.TestAdapter" Version="3.6.1" />
+    <PackageReference Include="MSTest.TestFramework" Version="3.6.1" />
     <PackageReference Include="coverlet.collector" Version="6.0.2">
       <PrivateAssets>all</PrivateAssets>
       <IncludeAssets>runtime; build; native; contentfiles; analyzers; buildtransitive</IncludeAssets>
diff --git a/Engine/Source/Programs/AutomationTool/Android/Android.Automation.csproj b/Engine/Source/Programs/AutomationTool/Android/Android.Automation.csproj
--- a/Engine/Source/Programs/AutomationTool/Android/Android.Automation.csproj
+++ b/Engine/Source/Programs/AutomationTool/Android/Android.Automation.csproj
@@ -48,7 +48,5 @@
   </ItemGroup>
   <ItemGroup>
     <PackageReference Include="Microsoft.CSharp" Version="4.7.0" />
-    <PackageReference Include="System.Data.DataSetExtensions" Version="4.5.0" />
-    <PackageReference Include="System.Net.Http" Version="4.3.4" />
   </ItemGroup>
 </Project>
\ No newline at end of file
diff --git a/Engine/Source/Programs/AutomationTool/AutomationTool.csproj b/Engine/Source/Programs/AutomationTool/AutomationTool.csproj
--- a/Engine/Source/Programs/AutomationTool/AutomationTool.csproj
+++ b/Engine/Source/Programs/AutomationTool/AutomationTool.csproj
@@ -97,11 +97,13 @@
 
   <ItemGroup>
     <PackageReference Include="AWSSDK.SecurityToken" Version="3.7.300.105" />
-    <PackageReference Include="Microsoft.Build" Version="17.1.0" ExcludeAssets="runtime" PrivateAssets="all" />
-    <PackageReference Include="Microsoft.Build.Locator" Version="1.4.1" PrivateAssets="all" />
+    <PackageReference Include="Microsoft.Build" Version="17.11.4" ExcludeAssets="runtime" PrivateAssets="all">
+      <IncludeAssets>compile; build; native; contentfiles; analyzers; buildtransitive</IncludeAssets>
+    </PackageReference>
+    <PackageReference Include="Microsoft.Build.Locator" Version="1.7.8" PrivateAssets="all" />
     <!-- Added as a temporary workaround to running from visual studio where the system version of dotnet is run. 
          Adding a loader context resolves many of the issues, but it still fails when MSBuild tries to write the json -->
-    <PackageReference Include="Newtonsoft.Json" Version="13.0.1" />
+    <PackageReference Include="Newtonsoft.Json" Version="13.0.3" />
   </ItemGroup>
   
   <!-- This set of dependencies should be kept as small as possible
diff --git a/Engine/Source/Programs/AutomationTool/AutomationUtils/AutomationUtils.Automation.csproj b/Engine/Source/Programs/AutomationTool/AutomationUtils/AutomationUtils.Automation.csproj
--- a/Engine/Source/Programs/AutomationTool/AutomationUtils/AutomationUtils.Automation.csproj
+++ b/Engine/Source/Programs/AutomationTool/AutomationUtils/AutomationUtils.Automation.csproj
@@ -62,9 +62,8 @@
     <ProjectReference Include="..\..\UnrealBuildTool\UnrealBuildTool.csproj" PrivateAssets="All"><Private>true</Private></ProjectReference>
   </ItemGroup>
   <ItemGroup>
-    <PackageReference Include="System.Data.DataSetExtensions" Version="4.5.0" />
-    <PackageReference Include="System.Net.Http" Version="4.3.4" />
     <PackageReference Include="System.Drawing.Common" Version="4.7.3" />
+    <PackageReference Include="System.Net.Http" Version="4.3.4" />
 
     <!--Required by Ionic -->
     <PackageReference Include="System.Security.Permissions" Version="4.7.0" />
diff --git a/Engine/Source/Programs/AutomationTool/BuildGraph/BuildGraph.Automation.csproj b/Engine/Source/Programs/AutomationTool/BuildGraph/BuildGraph.Automation.csproj
--- a/Engine/Source/Programs/AutomationTool/BuildGraph/BuildGraph.Automation.csproj
+++ b/Engine/Source/Programs/AutomationTool/BuildGraph/BuildGraph.Automation.csproj
@@ -62,8 +62,6 @@
   </ItemGroup>
   <ItemGroup>
     <PackageReference Include="Microsoft.CSharp" Version="4.7.0" />
-    <PackageReference Include="System.Data.DataSetExtensions" Version="4.5.0" />
-    <PackageReference Include="System.Net.Http" Version="4.3.4" />
     <PackageReference Include="System.Drawing.Common" Version="4.7.3" />
   </ItemGroup>
 </Project>
\ No newline at end of file
diff --git a/Engine/Source/Programs/AutomationTool/Gauntlet/Gauntlet.Automation.csproj b/Engine/Source/Programs/AutomationTool/Gauntlet/Gauntlet.Automation.csproj
--- a/Engine/Source/Programs/AutomationTool/Gauntlet/Gauntlet.Automation.csproj
+++ b/Engine/Source/Programs/AutomationTool/Gauntlet/Gauntlet.Automation.csproj
@@ -92,7 +92,5 @@
   <ItemGroup>
     <PackageReference Include="Microsoft.CSharp" Version="4.7.0" />
     <PackageReference Include="MySql.Data" Version="6.10.9" PrivateAssets="all" />
-    <PackageReference Include="System.Data.DataSetExtensions" Version="4.5.0" />
-    <PackageReference Include="System.Net.Http" Version="4.3.4" />
   </ItemGroup>
 </Project>
\ No newline at end of file
diff --git a/Engine/Source/Programs/AutomationTool/IOS/IOS.Automation.csproj b/Engine/Source/Programs/AutomationTool/IOS/IOS.Automation.csproj
--- a/Engine/Source/Programs/AutomationTool/IOS/IOS.Automation.csproj
+++ b/Engine/Source/Programs/AutomationTool/IOS/IOS.Automation.csproj
@@ -54,8 +54,6 @@
   </ItemGroup>
   <ItemGroup>
     <PackageReference Include="Microsoft.CSharp" Version="4.7.0" PrivateAssets="all" />
-    <PackageReference Include="System.Data.DataSetExtensions" Version="4.5.0" PrivateAssets="all" />
-    <PackageReference Include="System.Net.Http" Version="4.3.4" PrivateAssets="all" />
     <PackageReference Include="System.Drawing.Common" Version="4.7.3" PrivateAssets="all" />
     <PackageReference Include="System.Resources.Extensions" Version="4.7.1" PrivateAssets="all" />
     <PackageReference Include="Portable.BouncyCastle" Version="1.8.6.7" PrivateAssets="all" />
diff --git a/Engine/Source/Programs/AutomationTool/Linux/Linux.Automation.csproj b/Engine/Source/Programs/AutomationTool/Linux/Linux.Automation.csproj
--- a/Engine/Source/Programs/AutomationTool/Linux/Linux.Automation.csproj
+++ b/Engine/Source/Programs/AutomationTool/Linux/Linux.Automation.csproj
@@ -43,8 +43,6 @@
   </ItemGroup>
   <ItemGroup>
     <PackageReference Include="Microsoft.CSharp" Version="4.7.0" />
-    <PackageReference Include="System.Data.DataSetExtensions" Version="4.5.0" />
-    <PackageReference Include="System.Net.Http" Version="4.3.4" />
     <PackageReference Include="System.Drawing.Common" Version="4.7.3" />
   </ItemGroup>
 </Project>
\ No newline at end of file
diff --git a/Engine/Source/Programs/AutomationTool/Localization/Localization.Automation.csproj b/Engine/Source/Programs/AutomationTool/Localization/Localization.Automation.csproj
--- a/Engine/Source/Programs/AutomationTool/Localization/Localization.Automation.csproj
+++ b/Engine/Source/Programs/AutomationTool/Localization/Localization.Automation.csproj
@@ -40,7 +40,4 @@
     <ProjectReference Include="..\..\UnrealBuildTool\UnrealBuildTool.csproj" PrivateAssets="All"><Private>false</Private></ProjectReference>
     <ProjectReference Include="..\AutomationUtils\AutomationUtils.Automation.csproj" PrivateAssets="All"><Private>false</Private></ProjectReference>
   </ItemGroup>
-  <ItemGroup>
-    <PackageReference Include="System.Data.DataSetExtensions" Version="4.5.0" />
-  </ItemGroup>
 </Project>
diff --git a/Engine/Source/Programs/AutomationTool/LowLevelTests/LowLevelTests.Automation.csproj b/Engine/Source/Programs/AutomationTool/LowLevelTests/LowLevelTests.Automation.csproj
--- a/Engine/Source/Programs/AutomationTool/LowLevelTests/LowLevelTests.Automation.csproj
+++ b/Engine/Source/Programs/AutomationTool/LowLevelTests/LowLevelTests.Automation.csproj
@@ -54,8 +54,6 @@
   <ItemGroup>
     <PackageReference Include="Microsoft.CSharp" Version="4.7.0" />
     <PackageReference Include="MySql.Data" Version="6.10.9" PrivateAssets="all" />
-    <PackageReference Include="System.Data.DataSetExtensions" Version="4.5.0" />
-    <PackageReference Include="System.Net.Http" Version="4.3.4" />
   </ItemGroup>
   <ItemGroup>
     <Folder Include="Utility\" />
diff --git a/Engine/Source/Programs/AutomationTool/Mac/Mac.Automation.csproj b/Engine/Source/Programs/AutomationTool/Mac/Mac.Automation.csproj
--- a/Engine/Source/Programs/AutomationTool/Mac/Mac.Automation.csproj
+++ b/Engine/Source/Programs/AutomationTool/Mac/Mac.Automation.csproj
@@ -42,8 +42,6 @@
   </ItemGroup>
   <ItemGroup>
     <PackageReference Include="Microsoft.CSharp" Version="4.7.0" />
-    <PackageReference Include="System.Data.DataSetExtensions" Version="4.5.0" />
-    <PackageReference Include="System.Net.Http" Version="4.3.4" />
     <PackageReference Include="System.Drawing.Common" Version="4.7.3" />
   </ItemGroup>
 </Project>
\ No newline at end of file
diff --git a/Engine/Source/Programs/AutomationTool/OneSkyLocalization/OneSkyLocalization.Automation.csproj b/Engine/Source/Programs/AutomationTool/OneSkyLocalization/OneSkyLocalization.Automation.csproj
--- a/Engine/Source/Programs/AutomationTool/OneSkyLocalization/OneSkyLocalization.Automation.csproj
+++ b/Engine/Source/Programs/AutomationTool/OneSkyLocalization/OneSkyLocalization.Automation.csproj
@@ -50,7 +50,5 @@
   </ItemGroup>
   <ItemGroup>
     <PackageReference Include="Microsoft.CSharp" Version="4.7.0" />
-    <PackageReference Include="System.Data.DataSetExtensions" Version="4.5.0" />
-    <PackageReference Include="System.Net.Http" Version="4.3.4" />
   </ItemGroup>
 </Project>
\ No newline at end of file
diff --git a/Engine/Source/Programs/AutomationTool/Scripts/AutomationScripts.Automation.csproj b/Engine/Source/Programs/AutomationTool/Scripts/AutomationScripts.Automation.csproj
--- a/Engine/Source/Programs/AutomationTool/Scripts/AutomationScripts.Automation.csproj
+++ b/Engine/Source/Programs/AutomationTool/Scripts/AutomationScripts.Automation.csproj
@@ -57,8 +57,6 @@
   <ItemGroup>
     <PackageReference Include="Microsoft.CSharp" Version="4.7.0" />
     <PackageReference Include="Microsoft.VisualStudio.OLE.Interop" Version="17.11.40262" />
-    <PackageReference Include="System.Data.DataSetExtensions" Version="4.5.0" />
-    <PackageReference Include="System.Net.Http" Version="4.3.4" />
     <PackageReference Include="System.Drawing.Common" Version="4.7.3" />
     <PackageReference Include="AWSSDK.S3" Version="3.3.113.2" PrivateAssets="all" />
   </ItemGroup>
diff --git a/Engine/Source/Programs/AutomationTool/SteamDeck/SteamDeck.Automation.csproj b/Engine/Source/Programs/AutomationTool/SteamDeck/SteamDeck.Automation.csproj
--- a/Engine/Source/Programs/AutomationTool/SteamDeck/SteamDeck.Automation.csproj
+++ b/Engine/Source/Programs/AutomationTool/SteamDeck/SteamDeck.Automation.csproj
@@ -42,8 +42,6 @@
   </ItemGroup>
   <ItemGroup>
     <PackageReference Include="Microsoft.CSharp" Version="4.7.0" />
-    <PackageReference Include="System.Data.DataSetExtensions" Version="4.5.0" />
-    <PackageReference Include="System.Net.Http" Version="4.3.4" />
     <PackageReference Include="System.Drawing.Common" Version="4.7.3" />
   </ItemGroup>
 </Project>
\ No newline at end of file
diff --git a/Engine/Source/Programs/AutomationTool/TVOS/TVOS.Automation.csproj b/Engine/Source/Programs/AutomationTool/TVOS/TVOS.Automation.csproj
--- a/Engine/Source/Programs/AutomationTool/TVOS/TVOS.Automation.csproj
+++ b/Engine/Source/Programs/AutomationTool/TVOS/TVOS.Automation.csproj
@@ -49,8 +49,6 @@
   </ItemGroup>
   <ItemGroup>
     <PackageReference Include="Microsoft.CSharp" Version="4.7.0" />
-    <PackageReference Include="System.Data.DataSetExtensions" Version="4.5.0" />
-    <PackageReference Include="System.Net.Http" Version="4.3.4" />
     <PackageReference Include="System.Drawing.Common" Version="4.7.3" />
   </ItemGroup>
 </Project>
\ No newline at end of file
diff --git a/Engine/Source/Programs/AutomationTool/Turnkey/Turnkey.Automation.csproj b/Engine/Source/Programs/AutomationTool/Turnkey/Turnkey.Automation.csproj
--- a/Engine/Source/Programs/AutomationTool/Turnkey/Turnkey.Automation.csproj
+++ b/Engine/Source/Programs/AutomationTool/Turnkey/Turnkey.Automation.csproj
@@ -56,8 +56,6 @@
   </ItemGroup>
   <ItemGroup>
     <PackageReference Include="Microsoft.CSharp" Version="4.7.0" />
-    <PackageReference Include="System.Data.DataSetExtensions" Version="4.5.0" />
-    <PackageReference Include="System.Net.Http" Version="4.3.4" />
     <PackageReference Include="System.Drawing.Common" Version="4.7.3" />
     <PackageReference Include="System.IO.FileSystem.AccessControl" Version="4.7.0" />
     <PackageReference Include="Google.Apis.Drive.v3" Version="1.33.1.1225" PrivateAssets="all" />
diff --git a/Engine/Source/Programs/AutomationTool/Win/Win.Automation.csproj b/Engine/Source/Programs/AutomationTool/Win/Win.Automation.csproj
--- a/Engine/Source/Programs/AutomationTool/Win/Win.Automation.csproj
+++ b/Engine/Source/Programs/AutomationTool/Win/Win.Automation.csproj
@@ -50,8 +50,6 @@
   </ItemGroup>
   <ItemGroup>
     <PackageReference Include="Microsoft.CSharp" Version="4.7.0" />
-    <PackageReference Include="System.Data.DataSetExtensions" Version="4.5.0" />
-    <PackageReference Include="System.Net.Http" Version="4.3.4" />
     <PackageReference Include="System.Drawing.Common" Version="4.7.3" />
   </ItemGroup>
 </Project>
\ No newline at end of file
diff --git a/Engine/Source/Programs/AutomationTool/XLocLocalization/XLocLocalization.Automation.csproj b/Engine/Source/Programs/AutomationTool/XLocLocalization/XLocLocalization.Automation.csproj
--- a/Engine/Source/Programs/AutomationTool/XLocLocalization/XLocLocalization.Automation.csproj
+++ b/Engine/Source/Programs/AutomationTool/XLocLocalization/XLocLocalization.Automation.csproj
@@ -40,7 +40,6 @@
   </ItemGroup>
   <ItemGroup>
     <PackageReference Include="Microsoft.CSharp" Version="4.7.0" PrivateAssets="all" />
-    <PackageReference Include="System.Data.DataSetExtensions" Version="4.5.0" PrivateAssets="all" />
     <PackageReference Include="System.ServiceModel.Http" Version="4.5.3" PrivateAssets="all" />
   </ItemGroup>
 </Project>
\ No newline at end of file
diff --git a/Engine/Source/Programs/Horde/Drivers/JobDriver.Tests/JobDriver.Tests.csproj b/Engine/Source/Programs/Horde/Drivers/JobDriver.Tests/JobDriver.Tests.csproj
--- a/Engine/Source/Programs/Horde/Drivers/JobDriver.Tests/JobDriver.Tests.csproj
+++ b/Engine/Source/Programs/Horde/Drivers/JobDriver.Tests/JobDriver.Tests.csproj
@@ -22,8 +22,8 @@
     <PackageReference Include="Microsoft.Extensions.Logging.Abstractions" Version="8.0.2" />
     <PackageReference Include="Microsoft.NET.Test.Sdk" Version="17.11.1" />
     <PackageReference Include="Moq" Version="4.20.70" />
-    <PackageReference Include="MSTest.TestAdapter" Version="3.6.0" />
-    <PackageReference Include="MSTest.TestFramework" Version="3.6.0" />
+    <PackageReference Include="MSTest.TestAdapter" Version="3.6.1" />
+    <PackageReference Include="MSTest.TestFramework" Version="3.6.1" />
     <PackageReference Include="coverlet.collector" Version="6.0.2">
       <PrivateAssets>all</PrivateAssets>
       <IncludeAssets>runtime; build; native; contentfiles; analyzers; buildtransitive</IncludeAssets>
diff --git a/Engine/Source/Programs/Horde/Drivers/JobDriver/JobDriver.csproj b/Engine/Source/Programs/Horde/Drivers/JobDriver/JobDriver.csproj
--- a/Engine/Source/Programs/Horde/Drivers/JobDriver/JobDriver.csproj
+++ b/Engine/Source/Programs/Horde/Drivers/JobDriver/JobDriver.csproj
@@ -65,16 +65,16 @@
       <PrivateAssets>all</PrivateAssets>
       <IncludeAssets>runtime; build; native; contentfiles; analyzers; buildtransitive</IncludeAssets>
     </PackageReference>
-    <PackageReference Include="Microsoft.Extensions.Caching.Memory" Version="8.0.0" />
-    <PackageReference Include="Microsoft.Extensions.Hosting" Version="8.0.0" />
-    <PackageReference Include="Microsoft.Extensions.Hosting.Abstractions" Version="8.0.0" />
-    <PackageReference Include="Microsoft.Extensions.Hosting.WindowsServices" Version="8.0.0" />
+    <PackageReference Include="Microsoft.Extensions.Caching.Memory" Version="8.0.1" />
+    <PackageReference Include="Microsoft.Extensions.Hosting" Version="8.0.1" />
+    <PackageReference Include="Microsoft.Extensions.Hosting.Abstractions" Version="8.0.1" />
+    <PackageReference Include="Microsoft.Extensions.Hosting.WindowsServices" Version="8.0.1" />
     <PackageReference Include="Microsoft.Extensions.Http" Version="8.0.1" />
     <PackageReference Include="Microsoft.Extensions.Http.Polly" Version="8.0.10" />
     <PackageReference Include="Microsoft.Extensions.Logging.Abstractions" Version="8.0.2" />
     <PackageReference Include="Microsoft.Extensions.Options.DataAnnotations" Version="8.0.0" />
     <PackageReference Include="Microsoft.Management.Infrastructure" Version="3.0.0" />
-    <PackageReference Include="Microsoft.VisualStudio.Threading.Analyzers" Version="17.8.14">
+    <PackageReference Include="Microsoft.VisualStudio.Threading.Analyzers" Version="17.11.20">
       <PrivateAssets>all</PrivateAssets>
       <IncludeAssets>runtime; build; native; contentfiles; analyzers; buildtransitive</IncludeAssets>
     </PackageReference>
diff --git a/Engine/Source/Programs/Horde/HordeAgent.Tests/HordeAgent.Tests.csproj b/Engine/Source/Programs/Horde/HordeAgent.Tests/HordeAgent.Tests.csproj
--- a/Engine/Source/Programs/Horde/HordeAgent.Tests/HordeAgent.Tests.csproj
+++ b/Engine/Source/Programs/Horde/HordeAgent.Tests/HordeAgent.Tests.csproj
@@ -23,8 +23,8 @@
     <PackageReference Include="Microsoft.Extensions.Logging.Abstractions" Version="8.0.2" />
     <PackageReference Include="Microsoft.NET.Test.Sdk" Version="17.11.1" />
     <PackageReference Include="Moq" Version="4.20.70" />
-    <PackageReference Include="MSTest.TestAdapter" Version="3.6.0" />
-    <PackageReference Include="MSTest.TestFramework" Version="3.6.0" />
+    <PackageReference Include="MSTest.TestAdapter" Version="3.6.1" />
+    <PackageReference Include="MSTest.TestFramework" Version="3.6.1" />
     <PackageReference Include="coverlet.collector" Version="6.0.2">
       <PrivateAssets>all</PrivateAssets>
       <IncludeAssets>runtime; build; native; contentfiles; analyzers; buildtransitive</IncludeAssets>
diff --git a/Engine/Source/Programs/Horde/HordeAgent/HordeAgent.csproj b/Engine/Source/Programs/Horde/HordeAgent/HordeAgent.csproj
--- a/Engine/Source/Programs/Horde/HordeAgent/HordeAgent.csproj
+++ b/Engine/Source/Programs/Horde/HordeAgent/HordeAgent.csproj
@@ -47,16 +47,16 @@
       <PrivateAssets>all</PrivateAssets>
       <IncludeAssets>runtime; build; native; contentfiles; analyzers; buildtransitive</IncludeAssets>
     </PackageReference>
-    <PackageReference Include="Microsoft.Extensions.Caching.Memory" Version="8.0.0" />
-    <PackageReference Include="Microsoft.Extensions.Hosting" Version="8.0.0" />
-    <PackageReference Include="Microsoft.Extensions.Hosting.Abstractions" Version="8.0.0" />
-    <PackageReference Include="Microsoft.Extensions.Hosting.WindowsServices" Version="8.0.0" />
+    <PackageReference Include="Microsoft.Extensions.Caching.Memory" Version="8.0.1" />
+    <PackageReference Include="Microsoft.Extensions.Hosting" Version="8.0.1" />
+    <PackageReference Include="Microsoft.Extensions.Hosting.Abstractions" Version="8.0.1" />
+    <PackageReference Include="Microsoft.Extensions.Hosting.WindowsServices" Version="8.0.1" />
     <PackageReference Include="Microsoft.Extensions.Http" Version="8.0.1" />
     <PackageReference Include="Microsoft.Extensions.Http.Polly" Version="8.0.10" />
     <PackageReference Include="Microsoft.Extensions.Logging.Abstractions" Version="8.0.2" />
     <PackageReference Include="Microsoft.Extensions.Options.DataAnnotations" Version="8.0.0" />
     <PackageReference Include="Microsoft.Management.Infrastructure" Version="3.0.0" />
-    <PackageReference Include="Microsoft.VisualStudio.Threading.Analyzers" Version="17.8.14">
+    <PackageReference Include="Microsoft.VisualStudio.Threading.Analyzers" Version="17.11.20">
       <PrivateAssets>all</PrivateAssets>
       <IncludeAssets>runtime; build; native; contentfiles; analyzers; buildtransitive</IncludeAssets>
     </PackageReference>
diff --git a/Engine/Source/Programs/Horde/HordeServer.Tests.Shared/HordeServer.Tests.Shared.csproj b/Engine/Source/Programs/Horde/HordeServer.Tests.Shared/HordeServer.Tests.Shared.csproj
--- a/Engine/Source/Programs/Horde/HordeServer.Tests.Shared/HordeServer.Tests.Shared.csproj
+++ b/Engine/Source/Programs/Horde/HordeServer.Tests.Shared/HordeServer.Tests.Shared.csproj
@@ -22,12 +22,12 @@
   </ItemGroup>
 
   <ItemGroup>
-    <PackageReference Include="Microsoft.AspNetCore.Mvc.Testing" Version="6.0.25" />
+    <PackageReference Include="Microsoft.AspNetCore.Mvc.Testing" Version="8.0.10" />
     <PackageReference Include="Microsoft.NET.Test.Sdk" Version="17.11.1" />
     <PackageReference Include="Moq" Version="4.20.70" />
-    <PackageReference Include="MSTest.TestAdapter" Version="3.6.0" />
-    <PackageReference Include="MSTest.TestFramework" Version="3.6.0" />
-    <PackageReference Include="Microsoft.Extensions.Caching.Memory" Version="8.0.0" />
+    <PackageReference Include="MSTest.TestAdapter" Version="3.6.1" />
+    <PackageReference Include="MSTest.TestFramework" Version="3.6.1" />
+    <PackageReference Include="Microsoft.Extensions.Caching.Memory" Version="8.0.1" />
     <PackageReference Include="coverlet.collector" Version="6.0.2">
       <PrivateAssets>all</PrivateAssets>
       <IncludeAssets>runtime; build; native; contentfiles; analyzers; buildtransitive</IncludeAssets>
diff --git a/Engine/Source/Programs/Horde/HordeServer.Tests/HordeServer.Tests.csproj b/Engine/Source/Programs/Horde/HordeServer.Tests/HordeServer.Tests.csproj
--- a/Engine/Source/Programs/Horde/HordeServer.Tests/HordeServer.Tests.csproj
+++ b/Engine/Source/Programs/Horde/HordeServer.Tests/HordeServer.Tests.csproj
@@ -19,12 +19,12 @@
   </PropertyGroup>
 
   <ItemGroup>
-    <PackageReference Include="Microsoft.AspNetCore.Mvc.Testing" Version="6.0.25" />
+    <PackageReference Include="Microsoft.AspNetCore.Mvc.Testing" Version="8.0.10" />
     <PackageReference Include="Microsoft.NET.Test.Sdk" Version="17.11.1" />
     <PackageReference Include="Moq" Version="4.20.70" />
-    <PackageReference Include="MSTest.TestAdapter" Version="3.6.0" />
-    <PackageReference Include="MSTest.TestFramework" Version="3.6.0" />
-    <PackageReference Include="Microsoft.Extensions.Caching.Memory" Version="8.0.0" />
+    <PackageReference Include="MSTest.TestAdapter" Version="3.6.1" />
+    <PackageReference Include="MSTest.TestFramework" Version="3.6.1" />
+    <PackageReference Include="Microsoft.Extensions.Caching.Memory" Version="8.0.1" />
     <PackageReference Include="coverlet.collector" Version="6.0.2">
       <PrivateAssets>all</PrivateAssets>
       <IncludeAssets>runtime; build; native; contentfiles; analyzers; buildtransitive</IncludeAssets>
diff --git a/Engine/Source/Programs/Horde/HordeServer/HordeServer.csproj b/Engine/Source/Programs/Horde/HordeServer/HordeServer.csproj
--- a/Engine/Source/Programs/Horde/HordeServer/HordeServer.csproj
+++ b/Engine/Source/Programs/Horde/HordeServer/HordeServer.csproj
@@ -44,17 +44,17 @@
       <IncludeAssets>runtime; build; native; contentfiles; analyzers; buildtransitive</IncludeAssets>
     </PackageReference>
     <PackageReference Include="Microsoft.AspNetCore.Authentication.Cookies" Version="2.2.0" />
-    <PackageReference Include="Microsoft.AspNetCore.Authentication.JwtBearer" Version="8.0.0" />
-    <PackageReference Include="Microsoft.AspNetCore.Authentication.OpenIdConnect" Version="8.0.0" />
-    <PackageReference Include="Microsoft.AspNetCore.DataProtection.StackExchangeRedis" Version="8.0.0" />
-    <PackageReference Include="Microsoft.AspNetCore.Mvc.Razor.RuntimeCompilation" Version="6.0.25" />
-    <PackageReference Include="Microsoft.AspNetCore.SpaServices.Extensions" Version="8.0.0" />
-    <PackageReference Include="Microsoft.Extensions.Caching.Memory" Version="8.0.0" />
-    <PackageReference Include="Microsoft.Extensions.Hosting.WindowsServices" Version="8.0.0" />
+    <PackageReference Include="Microsoft.AspNetCore.Authentication.JwtBearer" Version="8.0.10" />
+    <PackageReference Include="Microsoft.AspNetCore.Authentication.OpenIdConnect" Version="8.0.10" />
+    <PackageReference Include="Microsoft.AspNetCore.DataProtection.StackExchangeRedis" Version="8.0.10" />
+    <PackageReference Include="Microsoft.AspNetCore.Mvc.Razor.RuntimeCompilation" Version="8.0.10" />
+    <PackageReference Include="Microsoft.AspNetCore.SpaServices.Extensions" Version="8.0.10" />
+    <PackageReference Include="Microsoft.Extensions.Caching.Memory" Version="8.0.1" />
+    <PackageReference Include="Microsoft.Extensions.Hosting.WindowsServices" Version="8.0.1" />
     <PackageReference Include="Microsoft.Extensions.Http" Version="8.0.1" />
     <PackageReference Include="Microsoft.Extensions.Http.Polly" Version="8.0.10" />
     <PackageReference Include="Microsoft.Extensions.Logging" Version="8.0.1" />
-    <PackageReference Include="Microsoft.VisualStudio.Threading.Analyzers" Version="17.8.14">
+    <PackageReference Include="Microsoft.VisualStudio.Threading.Analyzers" Version="17.11.20">
       <PrivateAssets>all</PrivateAssets>
       <IncludeAssets>runtime; build; native; contentfiles; analyzers; buildtransitive</IncludeAssets>
     </PackageReference>
@@ -76,7 +76,7 @@
     <PackageReference Include="Serilog.Sinks.File" Version="5.0.0" />
     <PackageReference Include="SerilogAnalyzer" Version="0.15.0" />
     <PackageReference Include="SharpZipLib" Version="1.4.2" />
-    <PackageReference Include="StackExchange.Redis" Version="2.7.4" />
+    <PackageReference Include="StackExchange.Redis" Version="2.8.16" />
     <PackageReference Include="Swashbuckle.AspNetCore" Version="6.5.0" />
     <PackageReference Include="System.Linq.Async" Version="6.0.1" />
     <PackageReference Include="System.Text.Json" Version="8.0.5" />
diff --git a/Engine/Source/Programs/Horde/HordeServer/Utilities/Clock.cs b/Engine/Source/Programs/Horde/HordeServer/Utilities/Clock.cs
--- a/Engine/Source/Programs/Horde/HordeServer/Utilities/Clock.cs
+++ b/Engine/Source/Programs/Horde/HordeServer/Utilities/Clock.cs
@@ -276,12 +276,12 @@ namespace HordeCommon
 		/// <inheritdoc/>
 		public ITicker AddSharedTicker(string name, TimeSpan interval, Func<CancellationToken, ValueTask> tickAsync, ILogger logger)
 		{
-			async ValueTask<TimeSpan?> Tick(CancellationToken token)
+			async ValueTask<TimeSpan?> TickAsync(CancellationToken token)
 			{
 				await tickAsync(token);
 				return interval;
 			}
-			return new TickerImpl(this, name, interval, Tick);
+			return new TickerImpl(this, name, interval, TickAsync);
 		}
 	}
 }
\ No newline at end of file
diff --git a/Engine/Source/Programs/Horde/Plugins/Analytics/HordeServer.Analytics.Tests/HordeServer.Analytics.Tests.csproj b/Engine/Source/Programs/Horde/Plugins/Analytics/HordeServer.Analytics.Tests/HordeServer.Analytics.Tests.csproj
--- a/Engine/Source/Programs/Horde/Plugins/Analytics/HordeServer.Analytics.Tests/HordeServer.Analytics.Tests.csproj
+++ b/Engine/Source/Programs/Horde/Plugins/Analytics/HordeServer.Analytics.Tests/HordeServer.Analytics.Tests.csproj
@@ -24,8 +24,8 @@
       <IncludeAssets>runtime; build; native; contentfiles; analyzers; buildtransitive</IncludeAssets>
     </PackageReference>
     <PackageReference Include="Microsoft.NET.Test.Sdk" Version="17.11.1" />
-    <PackageReference Include="MSTest.TestAdapter" Version="3.6.0" />
-    <PackageReference Include="MSTest.TestFramework" Version="3.6.0" />
+    <PackageReference Include="MSTest.TestAdapter" Version="3.6.1" />
+    <PackageReference Include="MSTest.TestFramework" Version="3.6.1" />
   </ItemGroup>
 
   <ItemGroup>
diff --git a/Engine/Source/Programs/Horde/Plugins/Build/HordeServer.Build.Tests/HordeServer.Build.Tests.csproj b/Engine/Source/Programs/Horde/Plugins/Build/HordeServer.Build.Tests/HordeServer.Build.Tests.csproj
--- a/Engine/Source/Programs/Horde/Plugins/Build/HordeServer.Build.Tests/HordeServer.Build.Tests.csproj
+++ b/Engine/Source/Programs/Horde/Plugins/Build/HordeServer.Build.Tests/HordeServer.Build.Tests.csproj
@@ -25,11 +25,11 @@
       <IncludeAssets>runtime; build; native; contentfiles; analyzers; buildtransitive</IncludeAssets>
     </PackageReference>
     <PackageReference Include="Microsoft.NET.Test.Sdk" Version="17.11.1" />
-    <PackageReference Include="MSTest.TestAdapter" Version="3.6.0" />
-    <PackageReference Include="MSTest.TestFramework" Version="3.6.0" />
-    <PackageReference Include="Microsoft.AspNetCore.Mvc.Testing" Version="6.0.25" />
+    <PackageReference Include="MSTest.TestAdapter" Version="3.6.1" />
+    <PackageReference Include="MSTest.TestFramework" Version="3.6.1" />
+    <PackageReference Include="Microsoft.AspNetCore.Mvc.Testing" Version="8.0.10" />
     <PackageReference Include="Moq" Version="4.20.70" />
-    <PackageReference Include="Microsoft.Extensions.Caching.Memory" Version="8.0.0" />
+    <PackageReference Include="Microsoft.Extensions.Caching.Memory" Version="8.0.1" />
   </ItemGroup>
 
   <ItemGroup>
diff --git a/Engine/Source/Programs/Horde/Plugins/Compute/HordeServer.Compute.Tests/HordeServer.Compute.Tests.csproj b/Engine/Source/Programs/Horde/Plugins/Compute/HordeServer.Compute.Tests/HordeServer.Compute.Tests.csproj
--- a/Engine/Source/Programs/Horde/Plugins/Compute/HordeServer.Compute.Tests/HordeServer.Compute.Tests.csproj
+++ b/Engine/Source/Programs/Horde/Plugins/Compute/HordeServer.Compute.Tests/HordeServer.Compute.Tests.csproj
@@ -25,8 +25,8 @@
       <IncludeAssets>runtime; build; native; contentfiles; analyzers; buildtransitive</IncludeAssets>
     </PackageReference>
     <PackageReference Include="Microsoft.NET.Test.Sdk" Version="17.11.1" />
-    <PackageReference Include="MSTest.TestAdapter" Version="3.6.0" />
-    <PackageReference Include="MSTest.TestFramework" Version="3.6.0" />
+    <PackageReference Include="MSTest.TestAdapter" Version="3.6.1" />
+    <PackageReference Include="MSTest.TestFramework" Version="3.6.1" />
   </ItemGroup>
 
   <ItemGroup>
diff --git a/Engine/Source/Programs/Horde/Plugins/Ddc/HordeServer.Ddc.Tests/HordeServer.Ddc.Tests.csproj b/Engine/Source/Programs/Horde/Plugins/Ddc/HordeServer.Ddc.Tests/HordeServer.Ddc.Tests.csproj
--- a/Engine/Source/Programs/Horde/Plugins/Ddc/HordeServer.Ddc.Tests/HordeServer.Ddc.Tests.csproj
+++ b/Engine/Source/Programs/Horde/Plugins/Ddc/HordeServer.Ddc.Tests/HordeServer.Ddc.Tests.csproj
@@ -25,8 +25,8 @@
       <IncludeAssets>runtime; build; native; contentfiles; analyzers; buildtransitive</IncludeAssets>
     </PackageReference>
     <PackageReference Include="Microsoft.NET.Test.Sdk" Version="17.11.1" />
-    <PackageReference Include="MSTest.TestAdapter" Version="3.6.0" />
-    <PackageReference Include="MSTest.TestFramework" Version="3.6.0" />
+    <PackageReference Include="MSTest.TestAdapter" Version="3.6.1" />
+    <PackageReference Include="MSTest.TestFramework" Version="3.6.1" />
   </ItemGroup>
 
   <ItemGroup>
diff --git a/Engine/Source/Programs/Horde/Plugins/Storage/HordeServer.Storage.Tests/HordeServer.Storage.Tests.csproj b/Engine/Source/Programs/Horde/Plugins/Storage/HordeServer.Storage.Tests/HordeServer.Storage.Tests.csproj
--- a/Engine/Source/Programs/Horde/Plugins/Storage/HordeServer.Storage.Tests/HordeServer.Storage.Tests.csproj
+++ b/Engine/Source/Programs/Horde/Plugins/Storage/HordeServer.Storage.Tests/HordeServer.Storage.Tests.csproj
@@ -25,8 +25,8 @@
       <IncludeAssets>runtime; build; native; contentfiles; analyzers; buildtransitive</IncludeAssets>
     </PackageReference>
     <PackageReference Include="Microsoft.NET.Test.Sdk" Version="17.11.1" />
-    <PackageReference Include="MSTest.TestAdapter" Version="3.6.0" />
-    <PackageReference Include="MSTest.TestFramework" Version="3.6.0" />
+    <PackageReference Include="MSTest.TestAdapter" Version="3.6.1" />
+    <PackageReference Include="MSTest.TestFramework" Version="3.6.1" />
   </ItemGroup>
 
   <ItemGroup>
diff --git a/Engine/Source/Programs/Horde/Plugins/Tools/HordeServer.Tools.Tests/HordeServer.Tools.Tests.csproj b/Engine/Source/Programs/Horde/Plugins/Tools/HordeServer.Tools.Tests/HordeServer.Tools.Tests.csproj
--- a/Engine/Source/Programs/Horde/Plugins/Tools/HordeServer.Tools.Tests/HordeServer.Tools.Tests.csproj
+++ b/Engine/Source/Programs/Horde/Plugins/Tools/HordeServer.Tools.Tests/HordeServer.Tools.Tests.csproj
@@ -25,8 +25,8 @@
       <IncludeAssets>runtime; build; native; contentfiles; analyzers; buildtransitive</IncludeAssets>
     </PackageReference>
     <PackageReference Include="Microsoft.NET.Test.Sdk" Version="17.11.1" />
-    <PackageReference Include="MSTest.TestAdapter" Version="3.6.0" />
-    <PackageReference Include="MSTest.TestFramework" Version="3.6.0" />
+    <PackageReference Include="MSTest.TestAdapter" Version="3.6.1" />
+    <PackageReference Include="MSTest.TestFramework" Version="3.6.1" />
   </ItemGroup>
 
   <ItemGroup>
diff --git a/Engine/Source/Programs/Shared/EpicGames.Build/Exceptions/BuildException.cs b/Engine/Source/Programs/Shared/EpicGames.Build/Exceptions/BuildException.cs
--- a/Engine/Source/Programs/Shared/EpicGames.Build/Exceptions/BuildException.cs
+++ b/Engine/Source/Programs/Shared/EpicGames.Build/Exceptions/BuildException.cs
@@ -1,6 +1,7 @@
 // Copyright Epic Games, Inc. All Rights Reserved.
 
 using System;
+using System.Diagnostics.CodeAnalysis;
 using EpicGames.Core;
 using Microsoft.Extensions.Logging;
 using UnrealBuildBase;
@@ -12,6 +13,7 @@ namespace UnrealBuildTool
 	/// <summary>
 	/// Base class for exceptions thrown by UnrealBuildTool and AutomationTool
 	/// </summary>
+	[SuppressMessage("Naming", "CA1724:Type names should not match namespaces", Justification = "Renaming would break public api")]
 	public class BuildException : Exception
 	{
 		/// <summary>
diff --git a/Engine/Source/Programs/Shared/EpicGames.BuildGraph.Tests/EpicGames.BuildGraph.Tests.csproj b/Engine/Source/Programs/Shared/EpicGames.BuildGraph.Tests/EpicGames.BuildGraph.Tests.csproj
--- a/Engine/Source/Programs/Shared/EpicGames.BuildGraph.Tests/EpicGames.BuildGraph.Tests.csproj
+++ b/Engine/Source/Programs/Shared/EpicGames.BuildGraph.Tests/EpicGames.BuildGraph.Tests.csproj
@@ -46,10 +46,10 @@
 
   <ItemGroup>
     <PackageReference Include="Microsoft.NET.Test.Sdk" Version="17.11.1" />
-    <PackageReference Include="MSTest.TestAdapter" Version="3.6.0">
+    <PackageReference Include="MSTest.TestAdapter" Version="3.6.1">
       <TreatAsUsed>true</TreatAsUsed>
     </PackageReference>
-    <PackageReference Include="MSTest.TestFramework" Version="3.6.0" />
+    <PackageReference Include="MSTest.TestFramework" Version="3.6.1" />
     <PackageReference Include="coverlet.collector" Version="6.0.2">
       <PrivateAssets>all</PrivateAssets>
       <IncludeAssets>runtime; build; native; contentfiles; analyzers; buildtransitive</IncludeAssets>
diff --git a/Engine/Source/Programs/Shared/EpicGames.Core.Tests/EpicGames.Core.Tests.csproj b/Engine/Source/Programs/Shared/EpicGames.Core.Tests/EpicGames.Core.Tests.csproj
--- a/Engine/Source/Programs/Shared/EpicGames.Core.Tests/EpicGames.Core.Tests.csproj
+++ b/Engine/Source/Programs/Shared/EpicGames.Core.Tests/EpicGames.Core.Tests.csproj
@@ -50,10 +50,10 @@
 
   <ItemGroup>
     <PackageReference Include="Microsoft.NET.Test.Sdk" Version="17.11.1" />
-    <PackageReference Include="MSTest.TestAdapter" Version="3.6.0">
+    <PackageReference Include="MSTest.TestAdapter" Version="3.6.1">
       <TreatAsUsed>true</TreatAsUsed>
     </PackageReference>
-    <PackageReference Include="MSTest.TestFramework" Version="3.6.0" />
+    <PackageReference Include="MSTest.TestFramework" Version="3.6.1" />
     <PackageReference Include="coverlet.collector" Version="6.0.2">
       <PrivateAssets>all</PrivateAssets>
       <IncludeAssets>runtime; build; native; contentfiles; analyzers; buildtransitive</IncludeAssets>
diff --git a/Engine/Source/Programs/Shared/EpicGames.Horde.Tests/EpicGames.Horde.Tests.csproj b/Engine/Source/Programs/Shared/EpicGames.Horde.Tests/EpicGames.Horde.Tests.csproj
--- a/Engine/Source/Programs/Shared/EpicGames.Horde.Tests/EpicGames.Horde.Tests.csproj
+++ b/Engine/Source/Programs/Shared/EpicGames.Horde.Tests/EpicGames.Horde.Tests.csproj
@@ -32,10 +32,10 @@
 
   <ItemGroup>
     <PackageReference Include="Microsoft.NET.Test.Sdk" Version="17.11.1" />
-    <PackageReference Include="MSTest.TestAdapter" Version="3.6.0">
+    <PackageReference Include="MSTest.TestAdapter" Version="3.6.1">
       <TreatAsUsed>true</TreatAsUsed>
     </PackageReference>
-    <PackageReference Include="MSTest.TestFramework" Version="3.6.0" />
+    <PackageReference Include="MSTest.TestFramework" Version="3.6.1" />
     <PackageReference Include="coverlet.collector" Version="6.0.2">
       <PrivateAssets>all</PrivateAssets>
       <IncludeAssets>runtime; build; native; contentfiles; analyzers; buildtransitive</IncludeAssets>
diff --git a/Engine/Source/Programs/Shared/EpicGames.Horde/Compute/AgentMessageHandler.cs b/Engine/Source/Programs/Shared/EpicGames.Horde/Compute/AgentMessageHandler.cs
--- a/Engine/Source/Programs/Shared/EpicGames.Horde/Compute/AgentMessageHandler.cs
+++ b/Engine/Source/Programs/Shared/EpicGames.Horde/Compute/AgentMessageHandler.cs
@@ -63,14 +63,14 @@ namespace EpicGames.Horde.Compute
 			using CancellationTokenSource cancellationSource = CancellationTokenSource.CreateLinkedTokenSource(cancellationToken);
 			ExceptionDispatchInfo? exceptionInfo = null;
 
-			async Task PostException(Exception ex)
+			async Task PostExceptionAsync(Exception ex)
 			{
 				// Capture stack from call site
 				Interlocked.CompareExchange(ref exceptionInfo, ExceptionDispatchInfo.Capture(ex), null);
 				await cancellationSource.CancelAsync();
 			}
 
-			await RunAsync(socket, 0, 4 * 1024 * 1024, PostException, cancellationSource.Token);
+			await RunAsync(socket, 0, 4 * 1024 * 1024, PostExceptionAsync, cancellationSource.Token);
 
 			// Throw the regular cancellation exception if requested
 			cancellationToken.ThrowIfCancellationRequested();
diff --git a/Engine/Source/Programs/Shared/EpicGames.Horde/EpicGames.Horde.csproj b/Engine/Source/Programs/Shared/EpicGames.Horde/EpicGames.Horde.csproj
--- a/Engine/Source/Programs/Shared/EpicGames.Horde/EpicGames.Horde.csproj
+++ b/Engine/Source/Programs/Shared/EpicGames.Horde/EpicGames.Horde.csproj
@@ -50,10 +50,10 @@
     </PackageReference>
     <PackageReference Include="K4os.Compression.LZ4" Version="1.2.16" />
     <PackageReference Include="System.IO.Pipelines" Version="8.0.0" />
-    <PackageReference Include="Microsoft.Data.Sqlite" Version="8.0.0" />
-    <PackageReference Include="Microsoft.Extensions.Caching.Abstractions" Version="6.0.0" />
-    <PackageReference Include="Microsoft.Extensions.Caching.Memory" Version="6.0.1" />
-    <PackageReference Include="Microsoft.VisualStudio.Threading.Analyzers" Version="17.8.14">
+    <PackageReference Include="Microsoft.Data.Sqlite" Version="8.0.10" />
+    <PackageReference Include="Microsoft.Extensions.Caching.Abstractions" Version="8.0.0" />
+    <PackageReference Include="Microsoft.Extensions.Caching.Memory" Version="8.0.1" />
+    <PackageReference Include="Microsoft.VisualStudio.Threading.Analyzers" Version="17.11.20">
       <PrivateAssets>all</PrivateAssets>
       <IncludeAssets>runtime; build; native; contentfiles; analyzers; buildtransitive</IncludeAssets>
     </PackageReference>
diff --git a/Engine/Source/Programs/Shared/EpicGames.IoHash.Tests/EpicGames.IoHash.Tests.csproj b/Engine/Source/Programs/Shared/EpicGames.IoHash.Tests/EpicGames.IoHash.Tests.csproj
--- a/Engine/Source/Programs/Shared/EpicGames.IoHash.Tests/EpicGames.IoHash.Tests.csproj
+++ b/Engine/Source/Programs/Shared/EpicGames.IoHash.Tests/EpicGames.IoHash.Tests.csproj
@@ -19,10 +19,10 @@
 
   <ItemGroup>
     <PackageReference Include="Microsoft.NET.Test.Sdk" Version="17.11.1" />
-    <PackageReference Include="MSTest.TestAdapter" Version="3.6.0">
+    <PackageReference Include="MSTest.TestAdapter" Version="3.6.1">
       <TreatAsUsed>true</TreatAsUsed>
     </PackageReference>
-    <PackageReference Include="MSTest.TestFramework" Version="3.6.0" />
+    <PackageReference Include="MSTest.TestFramework" Version="3.6.1" />
     <PackageReference Include="coverlet.collector" Version="6.0.2">
       <PrivateAssets>all</PrivateAssets>
       <IncludeAssets>runtime; build; native; contentfiles; analyzers; buildtransitive</IncludeAssets>
diff --git a/Engine/Source/Programs/Shared/EpicGames.MsBuild/EpicGames.MsBuild.csproj b/Engine/Source/Programs/Shared/EpicGames.MsBuild/EpicGames.MsBuild.csproj
--- a/Engine/Source/Programs/Shared/EpicGames.MsBuild/EpicGames.MsBuild.csproj
+++ b/Engine/Source/Programs/Shared/EpicGames.MsBuild/EpicGames.MsBuild.csproj
@@ -38,8 +38,8 @@
   </ItemGroup>
 
   <ItemGroup>
-    <PackageReference Include="Microsoft.Build" Version="17.1.0" ExcludeAssets="runtime" />
-    <PackageReference Include="Microsoft.Build.Locator" Version="1.4.1" />
+    <PackageReference Include="Microsoft.Build" Version="17.11.4" ExcludeAssets="runtime" />
+    <PackageReference Include="Microsoft.Build.Locator" Version="1.7.8" />
     <PackageReference Include="System.Drawing.Common" Version="4.7.3" />
   </ItemGroup>
 
diff --git a/Engine/Source/Programs/Shared/EpicGames.OIDC/EpicGames.OIDC.csproj b/Engine/Source/Programs/Shared/EpicGames.OIDC/EpicGames.OIDC.csproj
--- a/Engine/Source/Programs/Shared/EpicGames.OIDC/EpicGames.OIDC.csproj
+++ b/Engine/Source/Programs/Shared/EpicGames.OIDC/EpicGames.OIDC.csproj
@@ -38,8 +38,8 @@
   
   <ItemGroup>
     <PackageReference Include="IdentityModel.OidcClient" Version="4.0.0" />
-    <PackageReference Include="Microsoft.Extensions.Configuration.Binder" Version="6.0.0" />
-    <PackageReference Include="Microsoft.Extensions.Configuration.Json" Version="6.0.0" />
+    <PackageReference Include="Microsoft.Extensions.Configuration.Binder" Version="8.0.2" />
+    <PackageReference Include="Microsoft.Extensions.Configuration.Json" Version="8.0.1" />
     <PackageReference Include="System.Text.Json" Version="8.0.5" />
   </ItemGroup>
 
diff --git a/Engine/Source/Programs/Shared/EpicGames.Perforce.Fixture/EpicGames.Perforce.Fixture.csproj b/Engine/Source/Programs/Shared/EpicGames.Perforce.Fixture/EpicGames.Perforce.Fixture.csproj
--- a/Engine/Source/Programs/Shared/EpicGames.Perforce.Fixture/EpicGames.Perforce.Fixture.csproj
+++ b/Engine/Source/Programs/Shared/EpicGames.Perforce.Fixture/EpicGames.Perforce.Fixture.csproj
@@ -34,10 +34,10 @@
   <ItemGroup>
     <PackageReference Include="Microsoft.Extensions.Logging.Console" Version="8.0.1" />
     <PackageReference Include="Microsoft.NET.Test.Sdk" Version="17.11.1" />
-    <PackageReference Include="MSTest.TestAdapter" Version="3.6.0">
+    <PackageReference Include="MSTest.TestAdapter" Version="3.6.1">
       <TreatAsUsed>true</TreatAsUsed>
     </PackageReference>
-    <PackageReference Include="MSTest.TestFramework" Version="3.6.0" />
+    <PackageReference Include="MSTest.TestFramework" Version="3.6.1" />
     <PackageReference Include="coverlet.collector" Version="6.0.2">
       <PrivateAssets>all</PrivateAssets>
       <IncludeAssets>runtime; build; native; contentfiles; analyzers; buildtransitive</IncludeAssets>
diff --git a/Engine/Source/Programs/Shared/EpicGames.Perforce.Managed.Tests/EpicGames.Perforce.Managed.Tests.csproj b/Engine/Source/Programs/Shared/EpicGames.Perforce.Managed.Tests/EpicGames.Perforce.Managed.Tests.csproj
--- a/Engine/Source/Programs/Shared/EpicGames.Perforce.Managed.Tests/EpicGames.Perforce.Managed.Tests.csproj
+++ b/Engine/Source/Programs/Shared/EpicGames.Perforce.Managed.Tests/EpicGames.Perforce.Managed.Tests.csproj
@@ -34,10 +34,10 @@
     <PackageReference Include="Microsoft.NET.Test.Sdk" Version="17.11.1">
       <TreatAsUsed>true</TreatAsUsed>
     </PackageReference>
-    <PackageReference Include="MSTest.TestAdapter" Version="3.6.0">
+    <PackageReference Include="MSTest.TestAdapter" Version="3.6.1">
       <TreatAsUsed>true</TreatAsUsed>
     </PackageReference>
-    <PackageReference Include="MSTest.TestFramework" Version="3.6.0">
+    <PackageReference Include="MSTest.TestFramework" Version="3.6.1">
       <TreatAsUsed>true</TreatAsUsed>
     </PackageReference>
     <PackageReference Include="coverlet.collector" Version="6.0.2">
diff --git a/Engine/Source/Programs/Shared/EpicGames.Perforce.Tests/EpicGames.Perforce.Tests.csproj b/Engine/Source/Programs/Shared/EpicGames.Perforce.Tests/EpicGames.Perforce.Tests.csproj
--- a/Engine/Source/Programs/Shared/EpicGames.Perforce.Tests/EpicGames.Perforce.Tests.csproj
+++ b/Engine/Source/Programs/Shared/EpicGames.Perforce.Tests/EpicGames.Perforce.Tests.csproj
@@ -31,10 +31,10 @@
 
   <ItemGroup>
     <PackageReference Include="Microsoft.NET.Test.Sdk" Version="17.11.1" />
-    <PackageReference Include="MSTest.TestAdapter" Version="3.6.0">
+    <PackageReference Include="MSTest.TestAdapter" Version="3.6.1">
       <TreatAsUsed>true</TreatAsUsed>
     </PackageReference>
-    <PackageReference Include="MSTest.TestFramework" Version="3.6.0" />
+    <PackageReference Include="MSTest.TestFramework" Version="3.6.1" />
     <PackageReference Include="coverlet.collector" Version="6.0.2">
       <PrivateAssets>all</PrivateAssets>
       <IncludeAssets>runtime; build; native; contentfiles; analyzers; buildtransitive</IncludeAssets>
diff --git a/Engine/Source/Programs/Shared/EpicGames.Redis.Tests/EpicGames.Redis.Tests.csproj b/Engine/Source/Programs/Shared/EpicGames.Redis.Tests/EpicGames.Redis.Tests.csproj
--- a/Engine/Source/Programs/Shared/EpicGames.Redis.Tests/EpicGames.Redis.Tests.csproj
+++ b/Engine/Source/Programs/Shared/EpicGames.Redis.Tests/EpicGames.Redis.Tests.csproj
@@ -32,10 +32,10 @@
 
   <ItemGroup>
     <PackageReference Include="Microsoft.NET.Test.Sdk" Version="17.11.1" />
-    <PackageReference Include="MSTest.TestAdapter" Version="3.6.0">
+    <PackageReference Include="MSTest.TestAdapter" Version="3.6.1">
       <TreatAsUsed>true</TreatAsUsed>
     </PackageReference>
-    <PackageReference Include="MSTest.TestFramework" Version="3.6.0" />
+    <PackageReference Include="MSTest.TestFramework" Version="3.6.1" />
     <PackageReference Include="coverlet.collector" Version="6.0.2">
       <PrivateAssets>all</PrivateAssets>
       <IncludeAssets>runtime; build; native; contentfiles; analyzers; buildtransitive</IncludeAssets>
diff --git a/Engine/Source/Programs/Shared/EpicGames.Redis/EpicGames.Redis.csproj b/Engine/Source/Programs/Shared/EpicGames.Redis/EpicGames.Redis.csproj
--- a/Engine/Source/Programs/Shared/EpicGames.Redis/EpicGames.Redis.csproj
+++ b/Engine/Source/Programs/Shared/EpicGames.Redis/EpicGames.Redis.csproj
@@ -37,7 +37,7 @@
   <ItemGroup>
     <PackageReference Include="Google.Protobuf" Version="3.25.1" />
     <PackageReference Include="protobuf-net" Version="3.2.30" />
-    <PackageReference Include="StackExchange.Redis" Version="2.7.4" />
+    <PackageReference Include="StackExchange.Redis" Version="2.8.16" />
   </ItemGroup>
   <ItemGroup>
     <ProjectReference Include="..\EpicGames.Core\EpicGames.Core.csproj" />
diff --git a/Engine/Source/Programs/Shared/EpicGames.Serialization.Tests/EpicGames.Serialization.Tests.csproj b/Engine/Source/Programs/Shared/EpicGames.Serialization.Tests/EpicGames.Serialization.Tests.csproj
--- a/Engine/Source/Programs/Shared/EpicGames.Serialization.Tests/EpicGames.Serialization.Tests.csproj
+++ b/Engine/Source/Programs/Shared/EpicGames.Serialization.Tests/EpicGames.Serialization.Tests.csproj
@@ -35,10 +35,10 @@
 
   <ItemGroup>
     <PackageReference Include="Microsoft.NET.Test.Sdk" Version="17.11.1" />
-    <PackageReference Include="MSTest.TestAdapter" Version="3.6.0">
+    <PackageReference Include="MSTest.TestAdapter" Version="3.6.1">
       <TreatAsUsed>true</TreatAsUsed>
     </PackageReference>
-    <PackageReference Include="MSTest.TestFramework" Version="3.6.0" />
+    <PackageReference Include="MSTest.TestFramework" Version="3.6.1" />
     <PackageReference Include="coverlet.collector" Version="6.0.2">
       <PrivateAssets>all</PrivateAssets>
       <IncludeAssets>runtime; build; native; contentfiles; analyzers; buildtransitive</IncludeAssets>
diff --git a/Engine/Source/Programs/Shared/EpicGames.Tracing.Tests/EpicGames.Tracing.Tests.csproj b/Engine/Source/Programs/Shared/EpicGames.Tracing.Tests/EpicGames.Tracing.Tests.csproj
--- a/Engine/Source/Programs/Shared/EpicGames.Tracing.Tests/EpicGames.Tracing.Tests.csproj
+++ b/Engine/Source/Programs/Shared/EpicGames.Tracing.Tests/EpicGames.Tracing.Tests.csproj
@@ -30,10 +30,10 @@
 
   <ItemGroup>
     <PackageReference Include="Microsoft.NET.Test.Sdk" Version="17.11.1" />
-    <PackageReference Include="MSTest.TestAdapter" Version="3.6.0">
+    <PackageReference Include="MSTest.TestAdapter" Version="3.6.1">
       <TreatAsUsed>true</TreatAsUsed>
     </PackageReference>
-    <PackageReference Include="MSTest.TestFramework" Version="3.6.0" />
+    <PackageReference Include="MSTest.TestFramework" Version="3.6.1" />
     <PackageReference Include="coverlet.collector" Version="6.0.2">
       <PrivateAssets>all</PrivateAssets>
       <IncludeAssets>runtime; build; native; contentfiles; analyzers; buildtransitive</IncludeAssets>
diff --git a/Engine/Source/Programs/Shared/UnrealBuildTool.Tests/UnrealBuildTool.Tests.csproj b/Engine/Source/Programs/Shared/UnrealBuildTool.Tests/UnrealBuildTool.Tests.csproj
--- a/Engine/Source/Programs/Shared/UnrealBuildTool.Tests/UnrealBuildTool.Tests.csproj
+++ b/Engine/Source/Programs/Shared/UnrealBuildTool.Tests/UnrealBuildTool.Tests.csproj
@@ -18,10 +18,10 @@
 
   <ItemGroup>
     <PackageReference Include="Microsoft.NET.Test.Sdk" Version="17.11.1" />
-    <PackageReference Include="MSTest.TestAdapter" Version="3.6.0">
+    <PackageReference Include="MSTest.TestAdapter" Version="3.6.1">
       <TreatAsUsed>true</TreatAsUsed>
     </PackageReference>
-    <PackageReference Include="MSTest.TestFramework" Version="3.6.0" />
+    <PackageReference Include="MSTest.TestFramework" Version="3.6.1" />
     <PackageReference Include="coverlet.collector" Version="6.0.2">
       <PrivateAssets>all</PrivateAssets>
       <IncludeAssets>runtime; build; native; contentfiles; analyzers; buildtransitive</IncludeAssets>
diff --git a/Engine/Source/Programs/UnrealBuildTool/UnrealBuildTool.csproj b/Engine/Source/Programs/UnrealBuildTool/UnrealBuildTool.csproj
--- a/Engine/Source/Programs/UnrealBuildTool/UnrealBuildTool.csproj
+++ b/Engine/Source/Programs/UnrealBuildTool/UnrealBuildTool.csproj
@@ -127,12 +127,14 @@
     <PackageReference Include="System.Reflection.MetadataLoadContext" Version="8.0.1" /> 
     
     <!-- Required for EpicGames.MsBuild -->
-    <PackageReference Include="Microsoft.Build" Version="17.1.0" ExcludeAssets="runtime" PrivateAssets="all" />
-    <PackageReference Include="Microsoft.Build.Locator" Version="1.4.1" PrivateAssets="all" />
+    <PackageReference Include="Microsoft.Build" Version="17.11.4" ExcludeAssets="runtime" PrivateAssets="all">
+      <IncludeAssets>compile; build; native; contentfiles; analyzers; buildtransitive</IncludeAssets>
+    </PackageReference>
+    <PackageReference Include="Microsoft.Build.Locator" Version="1.7.8" PrivateAssets="all" />
 
     <!-- Added as a temporary workaround to running from visual studio where the system version of dotnet is run. 
          Adding a loader context resolves many of the issues, but it still fails when MSBuild tries to write the json -->
-    <PackageReference Include="Newtonsoft.Json" Version="13.0.1" />
+    <PackageReference Include="Newtonsoft.Json" Version="13.0.3" />
   </ItemGroup>
 
   <!-- Various batch files that build or run UnrealBuildTool -->
diff --git a/Engine/Source/Programs/UnrealGameSync/UnrealGameSync/UnrealGameSync.csproj b/Engine/Source/Programs/UnrealGameSync/UnrealGameSync/UnrealGameSync.csproj
--- a/Engine/Source/Programs/UnrealGameSync/UnrealGameSync/UnrealGameSync.csproj
+++ b/Engine/Source/Programs/UnrealGameSync/UnrealGameSync/UnrealGameSync.csproj
@@ -36,7 +36,7 @@
     <DefineConstants>$(DefineConstants);WITH_AUTOUPDATE</DefineConstants>
   </PropertyGroup>
   <ItemGroup>
-    <PackageReference Include="EnvDTE" Version="16.7.30328.74" />
+    <PackageReference Include="EnvDTE" Version="17.11.40262" />
     <PackageReference Include="Sentry" Version="3.15.0" />
     <PackageReference Include="Serilog.Expressions" Version="3.2.1" />
     <PackageReference Include="Serilog.Extensions.Hosting" Version="4.2.0" />
diff --git a/Engine/Source/Programs/UnrealGameSync/UnrealGameSyncShared/UnrealGameSyncShared.csproj b/Engine/Source/Programs/UnrealGameSync/UnrealGameSyncShared/UnrealGameSyncShared.csproj
--- a/Engine/Source/Programs/UnrealGameSync/UnrealGameSyncShared/UnrealGameSyncShared.csproj
+++ b/Engine/Source/Programs/UnrealGameSync/UnrealGameSyncShared/UnrealGameSyncShared.csproj
@@ -38,7 +38,7 @@
   </Choose>
   
   <ItemGroup>
-    <PackageReference Include="Microsoft.Extensions.Caching.Memory" Version="6.0.1" />
+    <PackageReference Include="Microsoft.Extensions.Caching.Memory" Version="8.0.1" />
     <PackageReference Include="System.IO.FileSystem.DriveInfo" Version="4.3.1" />
   </ItemGroup>
 
diff --git a/Engine/Source/Programs/UnrealGameSync/UnrealGameSyncShared/Utility.cs b/Engine/Source/Programs/UnrealGameSync/UnrealGameSyncShared/Utility.cs
--- a/Engine/Source/Programs/UnrealGameSync/UnrealGameSyncShared/Utility.cs
+++ b/Engine/Source/Programs/UnrealGameSync/UnrealGameSyncShared/Utility.cs
@@ -133,7 +133,7 @@ namespace UnrealGameSync
 					}
 
 					string cacheKey = GetChangeCacheKey(maxChangeNumber.Value, configHash);
-					if (s_changeCache.TryGetValue(cacheKey, out CachedChangeRecord change))
+					if (s_changeCache.TryGetValue(cacheKey, out CachedChangeRecord? change) && change != null)
 					{
 						yield return change.ChangesRecord;
 
diff --git a/Engine/Source/Programs/UnrealToolbox/UnrealToolbox.csproj b/Engine/Source/Programs/UnrealToolbox/UnrealToolbox.csproj
--- a/Engine/Source/Programs/UnrealToolbox/UnrealToolbox.csproj
+++ b/Engine/Source/Programs/UnrealToolbox/UnrealToolbox.csproj
@@ -86,6 +86,7 @@
     <PackageReference Include="FluentAvaloniaUI" Version="2.1.0" />
     <PackageReference Include="CommunityToolkit.Mvvm" Version="8.3.1" />
     <PackageReference Include="DesktopNotifications.Avalonia" Version="1.3.1" />
+    <PackageReference Include="System.Drawing.Common" Version="4.7.3" />
   </ItemGroup>
 
 </Project>
\ No newline at end of file
