Horde: Pass OpenTelemetry settings to agent driver subprocess

diff --git a/Engine/Source/Programs/Horde/HordeAgent/Leases/Handlers/JobHandler.cs b/Engine/Source/Programs/Horde/HordeAgent/Leases/Handlers/JobHandler.cs
--- a/Engine/Source/Programs/Horde/HordeAgent/Leases/Handlers/JobHandler.cs
+++ b/Engine/Source/Programs/Horde/HordeAgent/Leases/Handlers/JobHandler.cs
@@ -4,20 +4,26 @@ using EpicGames.Core;
 using EpicGames.Horde;
 using EpicGames.Horde.Agents.Leases;
 using EpicGames.Horde.Logs;
+using EpicGames.Horde.Utilities;
 using Google.Protobuf;
 using HordeAgent.Services;
 using HordeCommon.Rpc.Messages;
 using HordeCommon.Rpc.Tasks;
 using Microsoft.Extensions.Logging;
+using Microsoft.Extensions.Options;
 using OpenTelemetry.Trace;
 
 namespace HordeAgent.Leases.Handlers
 {
 	class JobHandler : LeaseHandler<ExecuteJobTask>
 	{
-		public JobHandler(RpcLease lease)
+		private readonly IOptions<AgentSettings> _settings;
+		
+		public JobHandler(RpcLease lease, IOptions<AgentSettings> settings)
 			: base(lease)
-		{ }
+		{
+			_settings = settings;
+		}
 
 		/// <inheritdoc/>
 		protected override async Task<LeaseResult> ExecuteAsync(ISession session, LeaseId leaseId, ExecuteJobTask executeTask, Tracer tracer, ILogger localLogger, CancellationToken cancellationToken)
@@ -49,6 +55,7 @@ namespace HordeAgent.Leases.Handlers
 				environment[HordeHttpClient.HordeUrlEnvVarName] = session.HordeClient.ServerUrl.ToString();
 				environment[HordeHttpClient.HordeTokenEnvVarName] = executeTask.Token;
 				environment["UE_LOG_JSON_TO_STDOUT"] = "1";
+				environment["UE_HORDE_OTEL_SETTINGS"] = OpenTelemetrySettingsExtensions.Serialize(_settings.Value.OpenTelemetry);
 
 				int exitCode = await RunDotNetProcessAsync(driverAssembly, arguments, environment, AgentApp.IsSelfContained, logger, cancellationToken);
 				logger.LogInformation("Driver finished with exit code {ExitCode}", exitCode);
@@ -70,8 +77,15 @@ namespace HordeAgent.Leases.Handlers
 
 	class JobHandlerFactory : LeaseHandlerFactory<ExecuteJobTask>
 	{
+		private readonly IOptions<AgentSettings> _settings;
+		
+		public JobHandlerFactory(IOptions<AgentSettings> settings)
+		{
+			_settings = settings;
+		}
+		
 		public override LeaseHandler<ExecuteJobTask> CreateHandler(RpcLease lease)
-			=> new JobHandler(lease);
+			=> new JobHandler(lease, _settings);
 	}
 }
 
