Horde: Fix bug with OpenTelemetry settings not being base64 encoded

diff --git a/Engine/Source/Programs/Horde/HordeAgent/Leases/Handlers/JobHandler.cs b/Engine/Source/Programs/Horde/HordeAgent/Leases/Handlers/JobHandler.cs
--- a/Engine/Source/Programs/Horde/HordeAgent/Leases/Handlers/JobHandler.cs
+++ b/Engine/Source/Programs/Horde/HordeAgent/Leases/Handlers/JobHandler.cs
@@ -55,7 +55,7 @@ namespace HordeAgent.Leases.Handlers
 				environment[HordeHttpClient.HordeUrlEnvVarName] = session.HordeClient.ServerUrl.ToString();
 				environment[HordeHttpClient.HordeTokenEnvVarName] = executeTask.Token;
 				environment["UE_LOG_JSON_TO_STDOUT"] = "1";
-				environment["UE_HORDE_OTEL_SETTINGS"] = OpenTelemetrySettingsExtensions.Serialize(_settings.Value.OpenTelemetry);
+				environment["UE_HORDE_OTEL_SETTINGS"] = OpenTelemetrySettingsExtensions.Serialize(_settings.Value.OpenTelemetry, true);
 
 				int exitCode = await RunDotNetProcessAsync(driverAssembly, arguments, environment, AgentApp.IsSelfContained, logger, cancellationToken);
 				logger.LogInformation("Driver finished with exit code {ExitCode}", exitCode);
