Horde: Remove use of resource.name in spans

resource.name is Datadog concept and doesn't play well presenting traces from OpenTelemetry in web UI.

diff --git a/Engine/Source/Programs/Horde/Drivers/JobDriver/Execution/JobExecutor.cs b/Engine/Source/Programs/Horde/Drivers/JobDriver/Execution/JobExecutor.cs
--- a/Engine/Source/Programs/Horde/Drivers/JobDriver/Execution/JobExecutor.cs
+++ b/Engine/Source/Programs/Horde/Drivers/JobDriver/Execution/JobExecutor.cs
@@ -400,7 +400,6 @@ namespace JobDriver.Execution
 
 					// Create a trace span
 					using TelemetrySpan span = Tracer.StartActiveSpan("Execute");
-					span.SetAttribute("resource.name", step.Name);
 					span.SetAttribute("horde.job.step.id", step.StepId.ToString());
 					span.SetAttribute("horde.job.step.name", step.Name);
 					span.SetAttribute("horde.job.step.log_id", step.LogId.ToString());
@@ -2193,7 +2192,7 @@ namespace JobDriver.Execution
 		{
 			TelemetrySpan newSpan = Tracer.StartSpan(span.Name ?? "TraceJsonUnknown", SpanKind.Internal, parent, startTime: new DateTime(span.Start, DateTimeKind.Utc));
 			newSpan.SetAttribute("service.name", span.Service);
-			newSpan.SetAttribute("resource.name", span.Resource);
+			newSpan.SetAttribute("resource.name_ot", span.Resource);
 
 			if (span.Properties != null)
 			{
diff --git a/Engine/Source/Programs/Horde/Drivers/JobDriver/Execution/ManagedWorkspaceMaterializer.cs b/Engine/Source/Programs/Horde/Drivers/JobDriver/Execution/ManagedWorkspaceMaterializer.cs
--- a/Engine/Source/Programs/Horde/Drivers/JobDriver/Execution/ManagedWorkspaceMaterializer.cs
+++ b/Engine/Source/Programs/Horde/Drivers/JobDriver/Execution/ManagedWorkspaceMaterializer.cs
@@ -148,9 +148,9 @@ public sealed class ManagedWorkspaceMaterializer : IWorkspaceMaterializer
 	private TelemetrySpan CreateTraceSpan(string operationName)
 	{
 		TelemetrySpan span = _tracer.StartActiveSpan(operationName);
-		span.SetAttribute("resource.name", _agentWorkspace.Identifier);
 		span.SetAttribute("horde.use_have_table", WorkspaceInfo.ShouldUseHaveTable(_agentWorkspace.Method));
 		span.SetAttribute("horde.cluster", _agentWorkspace.Cluster);
+		span.SetAttribute("horde.identifier", _agentWorkspace.Identifier);
 		span.SetAttribute("horde.incremental", _agentWorkspace.Incremental);
 		span.SetAttribute("horde.method", _agentWorkspace.Method);
 		span.SetAttribute("horde.stream", _agentWorkspace.Stream);
diff --git a/Engine/Source/Programs/Horde/Drivers/JobDriver/Execution/PerforceExecutor.cs b/Engine/Source/Programs/Horde/Drivers/JobDriver/Execution/PerforceExecutor.cs
--- a/Engine/Source/Programs/Horde/Drivers/JobDriver/Execution/PerforceExecutor.cs
+++ b/Engine/Source/Programs/Horde/Drivers/JobDriver/Execution/PerforceExecutor.cs
@@ -56,7 +56,7 @@ namespace JobDriver.Execution
 			if (_autoSdkWorkspaceInfo != null)
 			{
 				using TelemetrySpan autoSdkSpan = Tracer.StartActiveSpan("Workspace");
-				autoSdkSpan.SetAttribute("resource.name", "AutoSDK");
+				autoSdkSpan.SetAttribute("horde.workspace", "AutoSDK");
 
 				ManagedWorkspaceOptions options = WorkspaceInfo.GetMwOptions(_autoSdkWorkspaceInfo);
 				_autoSdkWorkspace = await WorkspaceInfo.CreateWorkspaceInfoAsync(_autoSdkWorkspaceInfo, _rootDir, options, logger, cancellationToken);
@@ -105,8 +105,8 @@ namespace JobDriver.Execution
 			}
 
 			{
-				using TelemetrySpan autoSdkSpan = Tracer.StartActiveSpan("Workspace");
-				autoSdkSpan.SetAttribute("resource.name", _workspaceInfo.Identifier);
+				using TelemetrySpan workspaceSpan = Tracer.StartActiveSpan("Workspace");
+				workspaceSpan.SetAttribute("horde.workspace", _workspaceInfo.Identifier);
 				
 				// Sync the regular workspace
 				ManagedWorkspaceOptions options = WorkspaceInfo.GetMwOptions(_workspaceInfo);
diff --git a/Engine/Source/Programs/Horde/Drivers/JobDriver/Execution/WorkspaceExecutor.cs b/Engine/Source/Programs/Horde/Drivers/JobDriver/Execution/WorkspaceExecutor.cs
--- a/Engine/Source/Programs/Horde/Drivers/JobDriver/Execution/WorkspaceExecutor.cs
+++ b/Engine/Source/Programs/Horde/Drivers/JobDriver/Execution/WorkspaceExecutor.cs
@@ -50,7 +50,7 @@ namespace JobDriver.Execution
 			if (_autoSdkWorkspace != null)
 			{
 				using TelemetrySpan autoSdkSpan = Tracer.StartActiveSpan("Workspace");
-				autoSdkSpan.SetAttribute("resource.name", "AutoSDK");
+				autoSdkSpan.SetAttribute("horde.workspace", "AutoSDK");
 				// TODO: Set type of workspace materializer as scope tag.
 
 				SyncOptions syncOptions = new();
@@ -59,8 +59,8 @@ namespace JobDriver.Execution
 
 			// Sync the regular workspace
 			{
-				using TelemetrySpan autoSdkSpan = Tracer.StartActiveSpan("Workspace");
-				autoSdkSpan.SetAttribute("resource.name", _workspace.Identifier);
+				using TelemetrySpan workspaceSpan = Tracer.StartActiveSpan("Workspace");
+				workspaceSpan.SetAttribute("horde.workspace", _workspace.Identifier);
 
 				await _workspace.SyncAsync(Batch.Change, Batch.PreflightChange, new SyncOptions(), cancellationToken);
 
diff --git a/Engine/Source/Programs/Horde/Plugins/Ddc/HordeServer.Ddc/Ddc/BlobService.cs b/Engine/Source/Programs/Horde/Plugins/Ddc/HordeServer.Ddc/Ddc/BlobService.cs
--- a/Engine/Source/Programs/Horde/Plugins/Ddc/HordeServer.Ddc/Ddc/BlobService.cs
+++ b/Engine/Source/Programs/Horde/Plugins/Ddc/HordeServer.Ddc/Ddc/BlobService.cs
@@ -42,8 +42,8 @@ namespace HordeServer.Ddc
 		{
 			using TelemetrySpan scope = _tracer.StartActiveSpan("put_blob")
 				.SetAttribute("operation.name", "put_blob")
-				.SetAttribute("resource.name", identifier.ToString())
-				.SetAttribute("Content-Length", payload.Length.ToString());
+				.SetAttribute("blob.id", identifier.ToString())
+				.SetAttribute("blob.length", payload.Length.ToString());
 
 			await using Stream hashStream = payload.GetStream();
 			await VerifyContentMatchesHashAsync(hashStream, identifier, cancellationToken);
