Horde: Move default settings for pools to the correct section (was in the "build" plugin section, should be in "compute").

#rnx

diff --git a/Engine/Source/Programs/Horde/HordeServer/Defaults/default.global.json b/Engine/Source/Programs/Horde/HordeServer/Defaults/default.global.json
--- a/Engine/Source/Programs/Horde/HordeServer/Defaults/default.global.json
+++ b/Engine/Source/Programs/Horde/HordeServer/Defaults/default.global.json
@@ -29,6 +29,28 @@
 					"description": "Default network (matching all IPs)",
 					"computeId": "default"
 				}
+			],
+			"pools": [
+				{
+					"name": "Win-UE5",
+					"condition": "Platform == 'Win64'",
+					"color": "Blue"
+				},
+				{
+					"name": "Win-UE5-GPU",
+					"condition": "Platform == 'Win64' && Interactive == true",
+					"color": "Blue"
+				},
+				{
+					"name": "Mac-UE5",
+					"condition": "Platform == 'Mac'",
+					"color": "Green"
+				},
+				{
+					"name": "Linux-UE5",
+					"condition": "Platform == 'Linux'",
+					"color": "Orange"
+				}
 			]
 		},
 
@@ -106,29 +128,7 @@
 						"poolType": "Shared"
 					}
 				]
-			},
-			"pools": [
-				{
-					"name": "Win-UE5",
-					"condition": "Platform == 'Win64'",
-					"color": "Blue"
-				},
-				{
-					"name": "Win-UE5-GPU",
-					"condition": "Platform == 'Win64' && Interactive == true",
-					"color": "Blue"
-				},
-				{
-					"name": "Mac-UE5",
-					"condition": "Platform == 'Mac'",
-					"color": "Green"
-				},
-				{
-					"name": "Linux-UE5",
-					"condition": "Platform == 'Linux'",
-					"color": "Orange"
-				}
-			]
+			}
 		}
 	},
 
