Horde: Do not update state of "healthy" flag when restarting an agent session until the agent is actually connected. If the server goes offline, the status switches to unhealthy, then immediately back to healthy as it tries to connect.

#rnx

diff --git a/Engine/Source/Programs/Horde/HordeAgent/IpcMessages.cs b/Engine/Source/Programs/Horde/HordeAgent/IpcMessages.cs
--- a/Engine/Source/Programs/Horde/HordeAgent/IpcMessages.cs
+++ b/Engine/Source/Programs/Horde/HordeAgent/IpcMessages.cs
@@ -174,16 +174,16 @@ namespace HordeAgent
 		/// <summary>
 		/// Static status object for starting an agent
 		/// </summary>
-		public static AgentStatusMessage Starting { get; } = new AgentStatusMessage(true, 0, "Starting up...");
+		public const string Starting = "Starting up...";
 
 		/// <summary>
 		/// Agent is waiting to be enrolled with the server
 		/// </summary>
-		public static AgentStatusMessage WaitingForEnrollment { get; } = new AgentStatusMessage(true, 0, "Waiting for enrollment...");
+		public const string WaitingForEnrollment = "Waiting for enrollment...";
 
 		/// <summary>
 		/// Agent is connecting to the server
 		/// </summary>
-		public static AgentStatusMessage ConnectingToServer { get; } = new AgentStatusMessage(true, 0, "Connecting to server...");
+		public const string ConnectingToServer = "Connecting to server...";
 	}
 }
diff --git a/Engine/Source/Programs/Horde/HordeAgent/Services/SessionFactory.cs b/Engine/Source/Programs/Horde/HordeAgent/Services/SessionFactory.cs
--- a/Engine/Source/Programs/Horde/HordeAgent/Services/SessionFactory.cs
+++ b/Engine/Source/Programs/Horde/HordeAgent/Services/SessionFactory.cs
@@ -155,7 +155,7 @@ namespace HordeAgent.Services
 			AgentRegistration? registrationInfo = registrationList.Entries.FirstOrDefault(x => x.Server == serverProfile.Url);
 			if (registrationInfo == null)
 			{
-				statusService.Set(AgentStatusMessage.WaitingForEnrollment);
+				statusService.SetDescription(AgentStatusMessage.WaitingForEnrollment);
 
 				registrationInfo = await RegisterAgentAsync(grpcService, currentSettings, capabilities, logger, cancellationToken);
 				registrationList.Entries.Add(registrationInfo);
@@ -169,7 +169,7 @@ namespace HordeAgent.Services
 			}
 
 			// Create the session
-			statusService.Set(AgentStatusMessage.ConnectingToServer);
+			statusService.SetDescription(AgentStatusMessage.ConnectingToServer);
 
 			RpcCreateSessionResponse createSessionResponse;
 			await using (HordeClient sessionClient = hordeClientFactory.Create(accessToken: registrationInfo.Token))
diff --git a/Engine/Source/Programs/Horde/HordeAgent/Services/StatusService.cs b/Engine/Source/Programs/Horde/HordeAgent/Services/StatusService.cs
--- a/Engine/Source/Programs/Horde/HordeAgent/Services/StatusService.cs
+++ b/Engine/Source/Programs/Horde/HordeAgent/Services/StatusService.cs
@@ -74,7 +74,7 @@ namespace HordeAgent.Services
 		/// </summary>
 		public StatusService(IOptionsMonitor<AgentSettings> settings, ILogger<StatusService> logger)
 		{
-			_current = AgentStatusMessage.Starting;
+			_current = new AgentStatusMessage(true, 0, AgentStatusMessage.Starting);
 			_settings = settings;
 			_task = new BackgroundTask(RunPipeServerAsync);
 			_logger = logger;
diff --git a/Engine/Source/Programs/Horde/HordeAgent/Services/WorkerService.cs b/Engine/Source/Programs/Horde/HordeAgent/Services/WorkerService.cs
--- a/Engine/Source/Programs/Horde/HordeAgent/Services/WorkerService.cs
+++ b/Engine/Source/Programs/Horde/HordeAgent/Services/WorkerService.cs
@@ -109,7 +109,7 @@ namespace HordeAgent.Services
 			while (!stoppingToken.IsCancellationRequested)
 			{
 				SessionResult? result = null;
-				_statusService.Set(AgentStatusMessage.Starting);
+				_statusService.SetDescription(AgentStatusMessage.Starting);
 
 				Stopwatch sessionTime = Stopwatch.StartNew();
 
diff --git a/Engine/Source/Programs/UnrealToolbox/Plugins/HordeAgent/HordeAgentPlugin.cs b/Engine/Source/Programs/UnrealToolbox/Plugins/HordeAgent/HordeAgentPlugin.cs
--- a/Engine/Source/Programs/UnrealToolbox/Plugins/HordeAgent/HordeAgentPlugin.cs
+++ b/Engine/Source/Programs/UnrealToolbox/Plugins/HordeAgent/HordeAgentPlugin.cs
@@ -262,7 +262,7 @@ namespace UnrealToolbox.Plugins.HordeAgent
 
 		async Task StatusTaskAsync(CancellationToken cancellationToken)
 		{
-			SetStatus(AgentStatusMessage.Starting);
+			SetStatus(new AgentStatusMessage(true, 0, AgentStatusMessage.Starting));
 			for (; ; )
 			{
 				try
