Horde: Updating schema documentation

#rnx

diff --git a/Engine/Source/Programs/Horde/Docs/Config/Schema/Globals.md b/Engine/Source/Programs/Horde/Docs/Config/Schema/Globals.md
--- a/Engine/Source/Programs/Horde/Docs/Config/Schema/Globals.md
+++ b/Engine/Source/Programs/Horde/Docs/Config/Schema/Globals.md
@@ -486,6 +486,8 @@ Name | Description
 `azureContainerName` | `string`<br>Name of the container
 `relayServer` | `string`<br>
 `relayToken` | `string`<br>
+`gcsBucketName` | `string`<br>Name of the GCS bucket to use
+`gcsBucketPath` | `string`<br>Base path within the bucket
 
 ## StorageBackendType (Enum)
 
@@ -496,6 +498,7 @@ Name | Description
 `FileSystem` | Local filesystem
 `Aws` | AWS S3
 `Azure` | Azure blob store
+`Gcs` | Google Cloud Storage
 `Memory` | In-memory only (for testing)
 
 ## AwsCredentialsType (Enum)
diff --git a/Engine/Source/Programs/Horde/Docs/Deployment/ServerSettings.md b/Engine/Source/Programs/Horde/Docs/Deployment/ServerSettings.md
--- a/Engine/Source/Programs/Horde/Docs/Deployment/ServerSettings.md
+++ b/Engine/Source/Programs/Horde/Docs/Deployment/ServerSettings.md
@@ -393,6 +393,8 @@ Name | Description
 `azureContainerName` | `string`<br>Name of the container
 `relayServer` | `string`<br>
 `relayToken` | `string`<br>
+`gcsBucketName` | `string`<br>Name of the GCS bucket to use
+`gcsBucketPath` | `string`<br>Base path within the bucket
 
 ## StorageBackendType (Enum)
 
@@ -403,6 +405,7 @@ Name | Description
 `FileSystem` | Local filesystem
 `Aws` | AWS S3
 `Azure` | Azure blob store
+`Gcs` | Google Cloud Storage
 `Memory` | In-memory only (for testing)
 
 ## AwsCredentialsType (Enum)
