Horde: Always show the reason for why compute resources may not be available

diff --git a/Engine/Source/Programs/Horde/Plugins/Compute/HordeServer.Compute/Compute/ComputeService.cs b/Engine/Source/Programs/Horde/Plugins/Compute/HordeServer.Compute/Compute/ComputeService.cs
--- a/Engine/Source/Programs/Horde/Plugins/Compute/HordeServer.Compute/Compute/ComputeService.cs
+++ b/Engine/Source/Programs/Horde/Plugins/Compute/HordeServer.Compute/Compute/ComputeService.cs
@@ -44,15 +44,16 @@ namespace HordeServer.Compute
 		public bool ShowToUser { get; init; } = false;
 
 		/// <inheritdoc/>
-		public ComputeServiceException(string? message, Exception? innerException = null) : base(message, innerException)
+		public ComputeServiceException(string? message, Exception? innerException = null, bool showToUser = false) : base(message, innerException)
 		{
+			ShowToUser = showToUser;
 		}
 	}
 	
 	/// <summary>
 	/// Exception when no compute resources available
 	/// </summary>
-	public class NoComputeResourcesException(string? message, Exception? innerException = null) : ComputeServiceException(message, innerException)
+	public class NoComputeResourcesException(string? message, Exception? innerException = null) : ComputeServiceException(message, innerException, true)
 	{
 		internal const string AllResourcesInUse = "Matching resources available but all are in use";
 		internal const string NoMatchingResources = "No matching resources available (in use or not)";
