Horde: Rename commands for creating and extracting archives.

#rnx

diff --git a/Engine/Source/Programs/Horde/HordeCmd/Commands/Archives/ArchiveCreate.cs b/Engine/Source/Programs/Horde/HordeCmd/Commands/Archives/ArchiveCreate.cs
new file mode 100644
--- /dev/null
+++ b/Engine/Source/Programs/Horde/HordeCmd/Commands/Archives/ArchiveCreate.cs
@@ -0,0 +1,120 @@
+// Copyright Epic Games, Inc. All Rights Reserved.
+
+using System.ComponentModel;
+using System.Diagnostics;
+using EpicGames.Core;
+using EpicGames.Horde.Storage;
+using EpicGames.Horde.Storage.Bundles;
+using EpicGames.Horde.Storage.Nodes;
+using Microsoft.Extensions.Logging;
+
+namespace Horde.Commands.Archives
+{
+	[Command("archive", "create", "Creates a storage archive from a folder on the local hard drive")]
+	class ArchiveCreate : StorageCommandBase
+	{
+		[CommandLine("-File=")]
+		[Description("Output file for the ref. Either -File=.. or -Ref=.. must be set.")]
+		public FileReference? File { get; set; }
+
+		[CommandLine("-Ref=")]
+		[Description("Output ref for the data. Either -File=.. or -Ref=.. must be set.")]
+		public string? Ref { get; set; }
+
+		[CommandLine("-Input=", Required = true)]
+		[Description("Input file or directory")]
+		public string Input { get; set; } = null!;
+
+		[CommandLine("-Filter=")]
+		[Description("Filter for files to include, in P4 syntax (eg. Foo/...).")]
+		public string Filter { get; set; } = "...";
+
+		[CommandLine("-CleanOutput")]
+		[Description("Clean the output folder before writing any data")]
+		public bool CleanOutput { get; set; }
+
+		public ArchiveCreate(HttpStorageClient storageClient, BundleCache bundleCache)
+			: base(storageClient, bundleCache)
+		{
+		}
+
+		public override async Task<int> ExecuteAsync(ILogger logger)
+		{
+			if (CleanOutput && File != null)
+			{
+				logger.LogInformation("Cleaning {Directory}...", File.Directory);
+				FileUtils.ForceDeleteDirectoryContents(File.Directory);
+			}
+
+			if (File != null)
+			{
+				using MemoryMappedFileCache memoryMappedFileCache = new MemoryMappedFileCache();
+				IStorageNamespace store = BundleStorageNamespace.CreateFromDirectory(File.Directory, BundleCache, memoryMappedFileCache, logger);
+				return await ExecuteInternalAsync(store, logger);
+			}
+			else if (!String.IsNullOrEmpty(Ref))
+			{
+				IStorageNamespace store = GetStorageNamespace();
+				return await ExecuteInternalAsync(store, logger);
+			}
+			else
+			{
+				throw new CommandLineArgumentException("Either -File=... or -Ref=... must be specified.");
+			}
+		}
+
+		async Task<int> ExecuteInternalAsync(IStorageNamespace store, ILogger logger)
+		{
+			// Gather the input files
+			DirectoryReference baseDir;
+			List<FileReference> files = new List<FileReference>();
+
+			if (System.IO.File.Exists(Input))
+			{
+				FileReference file = new FileReference(Input);
+				baseDir = file.Directory;
+				files.Add(file);
+			}
+			else if (Directory.Exists(Input))
+			{
+				baseDir = new DirectoryReference(Input);
+				FileFilter filter = new FileFilter(Filter.Split(';'));
+				files.AddRange(filter.ApplyToDirectory(baseDir, true));
+			}
+			else
+			{
+				logger.LogError("{Path} does not exist", Input);
+				return 1;
+			}
+
+			// Create the bundle
+			Stopwatch timer = Stopwatch.StartNew();
+
+			await using (IBlobWriter writer = store.CreateBlobWriter())
+			{
+				ChunkingOptions options = new ChunkingOptions();
+
+				List<FileInfo> fileInfos = files.ConvertAll(x => x.ToFileInfo());
+				UpdateStatsLogger updateStatsLogger = new UpdateStatsLogger(files.Count, fileInfos.Sum(x => x.Length), logger);
+
+				IHashedBlobRef<DirectoryNode> nodeRef = await writer.WriteFilesAsync(baseDir.ToDirectoryInfo(), fileInfos, options, updateStatsLogger, CancellationToken.None);
+
+				await writer.FlushAsync();
+
+				if (File != null)
+				{
+					logger.LogInformation("Writing {File}", File);
+					await FileReference.WriteAllTextAsync(File, nodeRef.GetLocator().ToString()!);
+				}
+				else
+				{
+					logger.LogInformation("Writing ref {Ref}", Ref);
+					await store.WriteRefAsync(new RefName(Ref!), nodeRef);
+				}
+			}
+
+			logger.LogInformation("Time: {Time}", timer.Elapsed.TotalSeconds);
+			return 0;
+		}
+	}
+}
diff --git a/Engine/Source/Programs/Horde/HordeCmd/Commands/Archives/ArchiveDump.cs b/Engine/Source/Programs/Horde/HordeCmd/Commands/Archives/ArchiveDump.cs
new file mode 100644
--- /dev/null
+++ b/Engine/Source/Programs/Horde/HordeCmd/Commands/Archives/ArchiveDump.cs
@@ -0,0 +1,84 @@
+// Copyright Epic Games, Inc. All Rights Reserved.
+
+using System.ComponentModel;
+using EpicGames.Core;
+using EpicGames.Horde.Storage.Bundles.V1;
+using Microsoft.Extensions.Logging;
+
+namespace Horde.Commands.Archives
+{
+	[Command("archive", "dump", "Dumps the contents of a bundle")]
+	internal class ArchiveDump : Command
+	{
+		[CommandLine("-Input=", Required = true)]
+		[Description("Path to the bundle to display information for.")]
+		public FileReference Input { get; set; } = null!;
+
+		[CommandLine("-Verbose")]
+		[Description("Include more detailed information about the bundle.")]
+		public bool Verbose { get; set; }
+
+		public override async Task<int> ExecuteAsync(ILogger logger)
+		{
+			byte[] data = await FileReference.ReadAllBytesAsync(Input);
+			logger.LogInformation("Summary for blob {Location}", Input);
+
+			BundleHeader header = BundleHeader.Read(data);
+
+			string[] types = new string[header.Types.Count];
+			for (int idx = 0; idx < header.Types.Count; idx++)
+			{
+				string name = header.Types[idx].Guid.ToString();
+				types[idx] = $"{name} v{header.Types[idx].Version}";
+			}
+
+			int packetStart = 0;
+
+			logger.LogInformation("");
+			logger.LogInformation("Imports: {NumImports}", header.Imports.Count);
+
+			for (int importIdx = 0; importIdx < header.Imports.Count; importIdx++)
+			{
+				logger.LogInformation("  IMP [{ImportIdx}] = {BlobId}", importIdx, header.Imports[importIdx]);
+			}
+
+			logger.LogInformation("");
+			logger.LogInformation("Exports: {NumExports}", header.Exports.Count);
+
+			int packetIdx = 0;
+			int packetOffset = 0;
+			for (int exportIdx = 0; exportIdx < header.Exports.Count; exportIdx++)
+			{
+				BundleExport export = header.Exports[exportIdx];
+				logger.LogInformation("  EXP [{ExportIdx}] EXP = hash: {ExportHash}, type: {Type}, length: {NumBytes:n0}, packet: {PacketIdx}", exportIdx, export.Hash, types[export.TypeIdx], export.Length, packetIdx);
+
+				if (Verbose)
+				{
+					for (int referenceIdx = 0; referenceIdx < export.References.Count; referenceIdx++)
+					{
+						BundleExportRef exportRef = export.References[referenceIdx];
+						logger.LogInformation("            REF {RefIdx,-3} -> {Node}", referenceIdx, $"{header.Imports[exportRef.ImportIdx]}#{exportRef.NodeIdx}");
+					}
+				}
+
+				packetOffset += export.Length;
+				if (packetOffset >= header.Packets[packetIdx].DecodedLength)
+				{
+					packetIdx++;
+					packetOffset = 0;
+				}
+			}
+
+			logger.LogInformation("");
+			logger.LogInformation("Packets: {NumPackets}", header.Packets.Count);
+			for (int idx = 0; idx < header.Packets.Count; idx++)
+			{
+				BundlePacket packet = header.Packets[idx];
+				logger.LogInformation("  PKT {Idx} (file offset: {Offset:n0}, encoded: {EncodedLength:n0}, decoded: {DecodedLength:n0}, ratio: {Ratio}%)", idx, packetStart, packet.EncodedLength, packet.DecodedLength, (int)(packet.EncodedLength * 100) / packet.DecodedLength);
+				packetStart += packet.EncodedLength;
+			}
+
+			return 0;
+		}
+	}
+}
diff --git a/Engine/Source/Programs/Horde/HordeCmd/Commands/Archives/ArchiveExtract.cs b/Engine/Source/Programs/Horde/HordeCmd/Commands/Archives/ArchiveExtract.cs
new file mode 100644
--- /dev/null
+++ b/Engine/Source/Programs/Horde/HordeCmd/Commands/Archives/ArchiveExtract.cs
@@ -0,0 +1,96 @@
+// Copyright Epic Games, Inc. All Rights Reserved.
+
+using System.ComponentModel;
+using System.Diagnostics;
+using EpicGames.Core;
+using EpicGames.Horde.Storage;
+using EpicGames.Horde.Storage.Backends;
+using EpicGames.Horde.Storage.Bundles;
+using EpicGames.Horde.Storage.Nodes;
+using Microsoft.Extensions.Logging;
+
+namespace Horde.Commands.Archives
+{
+	[Command("archive", "extract", "Extracts data from a bundle to the local hard drive")]
+	internal class ArchiveExtract : StorageCommandBase
+	{
+		[CommandLine("-File=")]
+		[Description("Path to a text file containing the root ref to read. -File=..., -Ref=..., or -Node=... must be specified.")]
+		public FileReference? File { get; set; }
+
+		[CommandLine("-Ref=")]
+		[Description("Name of a ref to read from the default storage namespace. -File=..., -Ref=..., or -Node=... must be specified.")]
+		public string? Ref { get; set; }
+
+		[CommandLine("-Node=")]
+		[Description("Locator for a node to read as the root. -File=..., -Ref=..., or -Node=... must be specified.")]
+		public string? Node { get; set; }
+
+		[CommandLine("-Stats")]
+		[Description("Outputs stats about the extraction process.")]
+		public bool Stats { get; set; }
+
+		[CommandLine("-OutputDir=", Required = true)]
+		[Description("Directory to write extracted files.")]
+		public DirectoryReference OutputDir { get; set; } = null!;
+
+		[CommandLine("-CleanOutput")]
+		[Description("If set, deletes the contents of the output directory before extraction.")]
+		public bool CleanOutput { get; set; }
+
+		public ArchiveExtract(HttpStorageClient storageClient, BundleCache bundleCache)
+			: base(storageClient, bundleCache)
+		{
+		}
+
+		public override async Task<int> ExecuteAsync(ILogger logger)
+		{
+			if (CleanOutput)
+			{
+				logger.LogInformation("Deleting contents of {OutputDir}...", OutputDir);
+				FileUtils.ForceDeleteDirectoryContents(OutputDir);
+			}
+
+			if (File != null)
+			{
+				using MemoryMappedFileCache memoryMappedFileCache = new MemoryMappedFileCache();
+				IStorageNamespace store = BundleStorageNamespace.CreateFromDirectory(File.Directory, BundleCache, memoryMappedFileCache, logger);
+				IBlobRef<DirectoryNode> handle = store.CreateBlobRef<DirectoryNode>(await FileStorageBackend.ReadRefAsync(File));
+				await ExecuteInternalAsync(store, handle, logger);
+			}
+			else if (Ref != null)
+			{
+				IStorageNamespace store = GetStorageNamespace();
+				IBlobRef<DirectoryNode> handle = await store.ReadRefAsync<DirectoryNode>(new RefName(Ref));
+				await ExecuteInternalAsync(store, handle, logger);
+			}
+			else if (Node != null)
+			{
+				IStorageNamespace store = GetStorageNamespace();
+				IBlobRef<DirectoryNode> handle = store.CreateBlobRef<DirectoryNode>(new BlobLocator(Node));
+				await ExecuteInternalAsync(store, handle, logger);
+			}
+			else
+			{
+				throw new CommandLineArgumentException("Either -File=... or -Ref=... must be specified");
+			}
+
+			return 0;
+		}
+
+		protected async Task ExecuteInternalAsync(IStorageNamespace store, IBlobRef<DirectoryNode> handle, ILogger logger)
+		{
+			Stopwatch timer = Stopwatch.StartNew();
+
+			await handle.ExtractAsync(OutputDir.ToDirectoryInfo(), new ExtractStatsLogger(logger), logger, CancellationToken.None);
+
+			logger.LogInformation("Elapsed: {Time}s", timer.Elapsed.TotalSeconds);
+
+			if (Stats)
+			{
+				StorageStats stats = store.GetStats();
+				stats.Print(logger);
+			}
+		}
+	}
+}
diff --git a/Engine/Source/Programs/Horde/HordeCmd/Commands/Bundles/BundleCreate.cs b/Engine/Source/Programs/Horde/HordeCmd/Commands/Bundles/BundleCreate.cs
deleted file mode 100644
diff --git a/Engine/Source/Programs/Horde/HordeCmd/Commands/Bundles/BundleDump.cs b/Engine/Source/Programs/Horde/HordeCmd/Commands/Bundles/BundleDump.cs
deleted file mode 100644
diff --git a/Engine/Source/Programs/Horde/HordeCmd/Commands/Bundles/BundleExtract.cs b/Engine/Source/Programs/Horde/HordeCmd/Commands/Bundles/BundleExtract.cs
deleted file mode 100644
