Horde: Fix dark theme ue5 project logo

diff --git a/Engine/Source/Programs/Horde/HordeServer/Defaults/ue5.project.json b/Engine/Source/Programs/Horde/HordeServer/Defaults/ue5.project.json
--- a/Engine/Source/Programs/Horde/HordeServer/Defaults/ue5.project.json
+++ b/Engine/Source/Programs/Horde/HordeServer/Defaults/ue5.project.json
@@ -1,6 +1,7 @@
 {
 	"name": "UE5",
 	"logo": "ue5.project.png",
+	"logoDarkTheme": "ue5.project-darktheme.png",
 	"streams": [
 		{
 		  "id": "ue5-release-5.5",
