Horde: Show logs event data as generating for running steps when missing

diff --git a/Engine/Source/Programs/Horde/HordeDashboard/src/components/ErrorPane.tsx b/Engine/Source/Programs/Horde/HordeDashboard/src/components/ErrorPane.tsx
--- a/Engine/Source/Programs/Horde/HordeDashboard/src/components/ErrorPane.tsx
+++ b/Engine/Source/Programs/Horde/HordeDashboard/src/components/ErrorPane.tsx
@@ -5,7 +5,7 @@ import { IColumn, List, Stack, Text } from '@fluentui/react';
 import { getFocusStyle, mergeStyleSets } from '@fluentui/react/lib/Styling';
 import React from 'react';
 import { Link, useNavigate } from 'react-router-dom';
-import { EventData, EventSeverity, GetLabelStateResponse } from '../backend/Api';
+import { EventData, EventSeverity, GetLabelStateResponse, JobStepState } from '../backend/Api';
 import { JobDetails } from '../backend/JobDetails';
 import { JobEventHandler } from '../backend/JobEventHandler';
 import { renderLine } from './LogRender';
@@ -119,7 +119,7 @@ export const ErrorPane: React.FC<{ jobDetails: JobDetails; stepId: string; showE
       return (<Stack className={styles.itemCell} styles={{ root: { padding: 8, marginRight: 8 } }}><Stack className={item.severity === EventSeverity.Warning ? styles.gutterWarning : styles.gutter} styles={{ root: { padding: 0, margin: 0 } }}>
          <Stack styles={{ root: { paddingLeft: 14 } }}>
             {!!lines.length && lines}
-            {!lines.length && <Text>Missing Log Data</Text>}
+            {!lines.length && <Text>{step.state === JobStepState.Running ? "Generating Log Data" : "Missing Log Data"}</Text>}
          </Stack>
       </Stack>
       </Stack>
@@ -214,8 +214,8 @@ export const JobEventListPanel: React.FC<{ jobDetails: JobDetails, stepIds: stri
          <Link className="log-link" to={url}>
             <Stack className={styles.itemCell} styles={{ root: { padding: 8, paddingLeft: 24, marginRight: 8 } }}><Stack className={event.severity === EventSeverity.Warning ? styles.gutterWarning : styles.gutter} styles={{ root: { padding: 0, margin: 0 } }}>
                <Stack styles={{ root: { paddingLeft: 8 } }}>
-               {!!lines.length && lines}
-               {!lines.length && <Text>Missing Log Data</Text>}
+                  {!!lines.length && lines}
+                  {!lines.length && <Text>{step.state === JobStepState.Running ? "Generating Log Data" : "Missing Log Data"}</Text>}
                </Stack>
             </Stack>
             </Stack>
diff --git a/Engine/Source/Programs/Horde/HordeDashboard/src/components/jobDetailsV2/StepErrorPanel.tsx b/Engine/Source/Programs/Horde/HordeDashboard/src/components/jobDetailsV2/StepErrorPanel.tsx
--- a/Engine/Source/Programs/Horde/HordeDashboard/src/components/jobDetailsV2/StepErrorPanel.tsx
+++ b/Engine/Source/Programs/Horde/HordeDashboard/src/components/jobDetailsV2/StepErrorPanel.tsx
@@ -6,7 +6,7 @@ import { observer } from "mobx-react-lite";
 import React, { useEffect } from 'react';
 import { Link, useNavigate } from 'react-router-dom';
 import backend from "../../backend";
-import { EventData, EventSeverity } from '../../backend/Api';
+import { EventData, EventSeverity, JobStepState } from '../../backend/Api';
 import { ISideRailLink } from "../../base/components/SideRail";
 import { getHordeStyling } from "../../styles/Styles";
 import { getHordeTheme } from "../../styles/theme";
@@ -226,7 +226,7 @@ const ErrorPane: React.FC<{ jobDetails: JobDetailsV2; view: StepSummaryErrorsVie
       return (<Stack className={styles.itemCell} styles={{ root: { padding: 8, marginRight: 8 } }}><Stack className={item.severity === EventSeverity.Warning ? styles.gutterWarning : styles.gutter} styles={{ root: { padding: 0, margin: 0 } }}>
          <Stack styles={{ root: { paddingLeft: 14 } }}>
             {!!lines.length && lines}
-            {!lines.length && <Text>Missing Log Data</Text>}
+            {!lines.length && <Text>{step.state === JobStepState.Running ? "Generating Log Data" : "Missing Log Data"}</Text>}
          </Stack>
       </Stack>
       </Stack>
