Horde: Delete ephemeral agents that have been offline for more than X hours

This ignores checking for any outdated sessions and simply uses the last online timestamp instead.

diff --git a/Engine/Source/Programs/Horde/Plugins/Build/HordeServer.Build.Tests/Agents/AgentServiceTest.cs b/Engine/Source/Programs/Horde/Plugins/Build/HordeServer.Build.Tests/Agents/AgentServiceTest.cs
--- a/Engine/Source/Programs/Horde/Plugins/Build/HordeServer.Build.Tests/Agents/AgentServiceTest.cs
+++ b/Engine/Source/Programs/Horde/Plugins/Build/HordeServer.Build.Tests/Agents/AgentServiceTest.cs
@@ -203,6 +203,27 @@ public class AgentServiceTest : BuildTestSetup
 		// Once more time has passed, the ephemeral agent marked as deleted is removed from database
 		Assert.IsNull(await AgentService.GetAgentAsync(agent.Id));
 	}
+	
+	[TestMethod]
+	// 5 hour old sessions are deleted
+	[DataRow(5, null, true)]
+	[DataRow(5, AgentStatus.Stopped, true)]
+	[DataRow(5, AgentStatus.Ok, false)]
+	[DataRow(5, null, false, false)] // Non-ephemeral agent should never be deleted
+	
+	// 30 min old sessions are not deleted
+	[DataRow(0.5, null, false)]
+	[DataRow(0.5, AgentStatus.Stopped, false)]
+	[DataRow(0.5, AgentStatus.Ok, false)]
+	public async Task Ephemeral_IsDeletedInBackgroundTick_Async(double lastOnlineAgoHours, AgentStatus? status, bool isDeleted, bool isEphemeral = true)
+	{
+		TimeSpan lastOnlineAgo = TimeSpan.FromHours(lastOnlineAgoHours);
+		IAgent? agent = await CreateAgentAsync(new PoolId("pool1"), ephemeral: isEphemeral, adjustClockBy: -lastOnlineAgo, status: status);
+		agent = await AgentService.GetAgentAsync(agent.Id);
+		Assert.IsNotNull(AgentService.GetAgentAsync(agent!.Id));
+		await ServiceProvider.GetRequiredService<AgentCollection>().DeleteExpiredEphemeralAgentsAsync(CancellationToken.None);
+		Assert.AreEqual(isDeleted, await AgentService.GetAgentAsync(agent!.Id) == null);
+	}
 
 	private static ClaimsPrincipal GetUser(IAgent agent)
 	{
diff --git a/Engine/Source/Programs/Horde/Plugins/Compute/HordeServer.Compute/Agents/AgentCollection.cs b/Engine/Source/Programs/Horde/Plugins/Compute/HordeServer.Compute/Agents/AgentCollection.cs
--- a/Engine/Source/Programs/Horde/Plugins/Compute/HordeServer.Compute/Agents/AgentCollection.cs
+++ b/Engine/Source/Programs/Horde/Plugins/Compute/HordeServer.Compute/Agents/AgentCollection.cs
@@ -324,7 +324,7 @@ namespace HordeServer.Agents
 			await DeleteExpiredEphemeralAgentsAsync(stoppingToken);
 		}
 
-		private async Task TerminateExpiredSessionsAsync(CancellationToken cancellationToken)
+		internal async Task TerminateExpiredSessionsAsync(CancellationToken cancellationToken)
 		{
 			using TelemetrySpan span = _tracer.StartActiveSpan($"{nameof(AgentCollection)}.{nameof(TerminateExpiredSessionsAsync)}");
 
@@ -345,19 +345,25 @@ namespace HordeServer.Agents
 			}
 			span.SetAttribute("NumAgentsTerminated", c);
 		}
-
-		private async Task DeleteExpiredEphemeralAgentsAsync(CancellationToken cancellationToken)
+		
+		internal async Task DeleteExpiredEphemeralAgentsAsync(CancellationToken cancellationToken)
 		{
 			using TelemetrySpan span = _tracer.StartActiveSpan($"{nameof(AgentCollection)}.{nameof(DeleteExpiredEphemeralAgentsAsync)}");
 			int numDeleted = 0;
-
-			List<AgentDocument> deletedDocuments = await _agentCollection.Find(x => x.Deleted).ToListAsync(cancellationToken);
+			TimeSpan maxAge = TimeSpan.FromHours(1);
+			
+			FilterDefinition<AgentDocument> filter = Builders<AgentDocument>.Filter.Or(
+				Builders<AgentDocument>.Filter.Eq(x => x.Deleted, true),
+				Builders<AgentDocument>.Filter.Lt(x => x.LastOnlineTime, _clock.UtcNow - maxAge)
+			);
+			
+			List<AgentDocument> deletedDocuments = await _agentCollection.Find(filter).ToListAsync(cancellationToken);
 			deletedDocuments = await PostLoadAsync(deletedDocuments, cancellationToken);
 
 			foreach (AgentDocument agent in deletedDocuments)
 			{
 				cancellationToken.ThrowIfCancellationRequested();
-				if (agent.LastOnlineTime.HasValue && _clock.UtcNow > agent.LastOnlineTime.Value + TimeSpan.FromHours(1))
+				if (agent is { Ephemeral: true, LastOnlineTime: not null } && _clock.UtcNow > agent.LastOnlineTime.Value + maxAge)
 				{
 					_logger.LogDebug("Deleting ephemeral agent {Agent}", agent.Id);
 
