Horde: Include ue5 project logo in installer build

diff --git a/Engine/Source/Programs/Horde/HordeServer/HordeServer.csproj b/Engine/Source/Programs/Horde/HordeServer/HordeServer.csproj
--- a/Engine/Source/Programs/Horde/HordeServer/HordeServer.csproj
+++ b/Engine/Source/Programs/Horde/HordeServer/HordeServer.csproj
@@ -137,6 +137,9 @@
     <Content Include="Defaults\ue5.project.png">
       <CopyToOutputDirectory>Always</CopyToOutputDirectory>
     </Content>
+    <Content Include="Defaults\ue5.project-darktheme.png">
+      <CopyToOutputDirectory>Always</CopyToOutputDirectory>
+    </Content>
   </ItemGroup>
 
   <ItemGroup>
