Horde: Update feature status doc for 5.5

#rnx

diff --git a/Engine/Source/Programs/Horde/Docs/Features.md b/Engine/Source/Programs/Horde/Docs/Features.md
--- a/Engine/Source/Programs/Horde/Docs/Features.md
+++ b/Engine/Source/Programs/Horde/Docs/Features.md
@@ -5,13 +5,13 @@
 | Feature | Status | Notes/Roadmap |
 | ------- | ------ | ----- |
 | Analytics | 🟧&nbsp;Experimental | - |
-| Artifact&nbsp;Distribution | 🟧&nbsp;Experimental | - |
-| Automation Hub | 🟧&nbsp;Experimental | - |
-| Build Automation | 🟨&nbsp;Beta | Migrate jobs to a more lightweight data model.<br>Programmable state machine for jobs. |
-| Build Health | 🟧&nbsp;Experimental | Migrate more engine code to produce structured logging.<br>Move issue creation into UAT.<br> |
+| Artifact&nbsp;Distribution | 🟩&nbsp;Production Ready | - |
+| Automation Hub | 🟩&nbsp;Production Ready | - |
+| Build Automation | 🟩&nbsp;Production Ready | Migrate jobs to a more lightweight data model.<br>Programmable state machine for jobs. |
+| Build Health | 🟩&nbsp;Production Ready | - |
 | DDC | 🟥&nbsp;Early | Recommend using UnrealCloudDDC rather than Horde. |
-| Remote Execution | 🟧&nbsp;Experimental | - |
-| UnrealGameSync&nbsp;Server | 🟨&nbsp;Beta | - |
+| Remote Execution | 🟩&nbsp;Production Ready | - |
+| UnrealGameSync&nbsp;Server | 🟩&nbsp;Production Ready | - |
 
 ## Key
 
