Horde: Documentation updates; add symbol server, separate config index into plugins.

#rnx

diff --git a/Engine/Source/Programs/Horde/Docs/Config.md b/Engine/Source/Programs/Horde/Docs/Config.md
--- a/Engine/Source/Programs/Horde/Docs/Config.md
+++ b/Engine/Source/Programs/Horde/Docs/Config.md
@@ -10,20 +10,53 @@ This section targets operating and managing Horde installations and shares some
   revision control, and how to set up a schema server.
 * [Permissions](Config/Permissions.md): Understand how to use the Horde permissions model.
 
-## Resources
+## Plugins
 
-* [Agents](Config/Agents.md): Configure worker machines to connect to the Horde server for CI and remote execution workspaces.
+Horde is implemented as a set of plugins over a host framework. For information about enabling or disabling specific plugins, see [Configuration > Plugins](Config/Plugins.md).
+
+### Analytics
+
+Experimental solution for studio-wide analytics and metrics gathering, with
+support integrated into AutomationTool, UnrealBuildTool and Unreal Editor.
+
+* [Analytics](Config/Analytics.md): Setting up a telemetry sink and creating dashboards showing KPIs.
+
+### Build
+
+Mature build automation system similar to Jenkins or TeamCity, designed
+and streamlined for Unreal Engine projects and best practices adoped by Epic.
+
+* [Build Automation](Config/BuildAutomation.md): Introduction to Horde's Build Automation and CI/CD system.
+* [Artifacts](Config/Artifacts.md): Managing and distributing artifacts produced by build steps.
 * [Devices](Config/Devices.md): Adding mobile devices and console development kits as shared resources.
-* [Secrets](Config/Secrets.md): Allowing Horde to manage access to confidential values.
-* [Storage](Config/Storage.md): Setting up different storage backends for Horde.
+* [Automation Hub](Config/AutomationHub.md): Dashboard for surfacing test jamdata and trends across multiple projects and streams.
+* [UGS Metadata Server](Config/UgsMetadataServer.md): Surfaces team metadata to users of UnrealGameSync.
+
+### Compute
+
+Manages worker machines that can be leased out to perform workloads, including build automation and remote execution use cases.
+
+* [Agents](Config/Agents.md): Configure worker machines to connect to the Horde server for CI and remote execution workspaces.
+
+### Secrets
+
+Managed access to sensitive data integrated with Horde's permissions model.Secrets may be stored in Horde itself, or in an external secret store.
+
+* [Secrets](Config/Secrets.md): Configuring secrets and external secret providers.
+
+### Storage
+
+Flexible, low overhead storage abstraction suitable for many different use
+cases. Used by Horde internally, but it can also be used directly by client applications.
+
+* [Storage](Config/Storage.md): Setting up and managing the storage system.  
+
+### Symbols
 
-## Features
+Implements a Windows symbol store indexing artifacts uploaded to Horde's
+storage system.
 
-* [Build Automation](Config/BuildAutomation.md)
-* [Remote Execution](Config/RemoteExecution.md)
-* DDC Server (Coming Soon)
-* [Automation Hub](Config/AutomationHub.md)
-* [UnrealGameSync Metadata Server](Config/UgsMetadataServer.md)
+* [Symbols](Config/Symbols.md): Adding symbol stores and configuring clients to use them.
 
 ## Reference
 
diff --git a/Engine/Source/Programs/Horde/Docs/Config/BuildAutomation.md b/Engine/Source/Programs/Horde/Docs/Config/BuildAutomation.md
--- a/Engine/Source/Programs/Horde/Docs/Config/BuildAutomation.md
+++ b/Engine/Source/Programs/Horde/Docs/Config/BuildAutomation.md
@@ -10,8 +10,8 @@ tracking and transfer of intermediate build artifacts between agents.
 
 It also supports the following features:
 
-* Build Health
-* Perforce metadata caching, workspace management
+* [Build Health](BuildHealth.md)
+* Perforce metadata caching and workspace management
 * Support for autoscaling based on job queues
 * Support for structured logging, with many common UE-types automatically annotated with additional metadata.
 * Profiling and telemetry functionality
@@ -29,7 +29,7 @@ When running a job on Horde, you specify a BuildGraph script, any command line a
 of one or more nodes to evaluate. Horde takes care of provisioning machines, syncing from Perforce, and transferring
 inputs and outputs to temporary storage.
 
-## Enabling CI functionality
+## Configuration
 
 To enable CI functionality in Horde, you will need to perform the following steps:
 
diff --git a/Engine/Source/Programs/Horde/Docs/Config/Orientation.md b/Engine/Source/Programs/Horde/Docs/Config/Orientation.md
--- a/Engine/Source/Programs/Horde/Docs/Config/Orientation.md
+++ b/Engine/Source/Programs/Horde/Docs/Config/Orientation.md
@@ -25,19 +25,6 @@ running with a cached version of the previous configuration until fixed.
 Horde configuration files may include other configuration files by path. For CI use cases, for example, it can be
 convenient to configure each stream within the stream itself.
 
-## Projects and Streams
-
-Most of the Horde dashboard is split into projects and streams. Projects are designed as a top-level way of
-partitioning functionality for different teams working on a shared Horde instance, and streams configure functionality
-relevant for a specific Perforce stream.
-
-Each project and stream typically has its own configuration file. By convention (and for the schema server to work
-correctly), project configuration files have a `.project.json` extension, and stream configuration files have a
-`.stream.json` extension.
-
-Projects and streams need to be set up to use the CI, PerfMem Hub, and Test Hub aspects of Horde. For
-configuring remote execution and DDC use cases, a global configuration file will suffice.
-
 ## Revision Control
 
 Horde supports reading configuration files from Perforce.
diff --git a/Engine/Source/Programs/Horde/Docs/Config/Permissions.md b/Engine/Source/Programs/Horde/Docs/Config/Permissions.md
--- a/Engine/Source/Programs/Horde/Docs/Config/Permissions.md
+++ b/Engine/Source/Programs/Horde/Docs/Config/Permissions.md
@@ -15,7 +15,7 @@ You can configure the mode via the `AuthMethod` setting.
 ### Anonymous
 Horde ships with authorization disabled by default for demonstration purposes and to get started.
 
-> **NOTE:** For production deployment, proper authentication must be configured using either OpenID Connect or built-in user accounts.
+> **NOTE:** For production deployments, proper authentication should be configured using either OpenID Connect or built-in user accounts.
 
 ### OpenID Connect
 Horde can use an external OpenID Connect (OIDC) provider for authorization.
diff --git a/Engine/Source/Programs/Horde/Docs/Config/Plugins.md b/Engine/Source/Programs/Horde/Docs/Config/Plugins.md
new file mode 100644
--- /dev/null
+++ b/Engine/Source/Programs/Horde/Docs/Config/Plugins.md
@@ -0,0 +1,26 @@
+[Horde](../../README.md) > [Configuration](../Config.md) > Plugins
+
+# Plugins
+
+Horde is implemented as a core host process and set of plugins. Most
+plugins correspond to a feature area can be enabled or disabled
+individually, though some plugins build on functionality providied by
+others. The [Storage](../Config.md#storage) and
+[Compute](../Config.md#compute) plugins are notable foundational components
+of Horde that other systems rely on.
+
+For information about how plugins operate and how to add new plugins, see
+[Internals > Plugin Architecture](Internals/PluginArchitecture.md).
+
+## Configuration
+
+Like the server itself, plugins can have a static boot-time configuration
+(through the [server.json](../Deployment/ServerSettings.md) file) and dynamic
+runtime configuration (through the [globals.json](Orientation.md) file).
+
+In each case, settings for each plugin is stored under the `plugins`
+property. All plugins have an `enabled` property which can be set to false
+to disable them.
+
+* [Plugins (server.json)](../Deployment/ServerSettings.md#serverpluginsconfig)
+* [Plugins (globals.json)](Schema/Globals.md#globalpluginsconfig)
diff --git a/Engine/Source/Programs/Horde/Docs/Config/Storage.md b/Engine/Source/Programs/Horde/Docs/Config/Storage.md
--- a/Engine/Source/Programs/Horde/Docs/Config/Storage.md
+++ b/Engine/Source/Programs/Horde/Docs/Config/Storage.md
@@ -14,12 +14,16 @@ in the storage system (some arbitrary block of data) are characterized by being
 outward references to other nodes. These tree-like data structures can be traversed one node at a time, with the
 underlying storage system retrieving and storing individual nodes efficiently.
 
-The entry point to any such data structure is a ref, a user-defined name that keeps a reference to the root of the
+The entry point to any such data structure is a user-defined name that keeps a reference to the root of the
 tree. Any nodes not directly or indirectly referenced by a ref are subject to garbage collection.
 
 This design intends to support a massive range of payload sizes by splitting larger data blocks into smaller chunks 
-and has advantages in caching and incremental client data patching. With this, it becomes reasonable to have a ref 
-point to a single 4kb payload as it is to point to a multi-gigabyte payload split into 128kb chunks.
+and has advantages in caching and incremental client data patching. With this, it becomes as reasonable to reference 
+a single 4kb payload as it is to point to a multi-gigabyte payload partitioned into 128kb chunks.
+
+The [Internals > Storage Architecture](../Internals/StorageArchitecture.md) page
+goes into technical details around implementation of the storage system, and
+describes how to interface with it programmatically.
 
 ## Namespaces and Backends
 
diff --git a/Engine/Source/Programs/Horde/Docs/Config/Symbols.md b/Engine/Source/Programs/Horde/Docs/Config/Symbols.md
new file mode 100644
--- /dev/null
+++ b/Engine/Source/Programs/Horde/Docs/Config/Symbols.md
@@ -0,0 +1,72 @@
+[Horde](../../README.md) > [Configuration](../Config.md) > Symbols
+
+# Symbols
+
+Horde implements a Windows symbol server compatible with Visual Studio, WinDbg, and Microsoft's DIA library.
+
+Symbols are exposed to the symbol store by uploading them as part of an [artifact](Artifacts.md). As such, any build produced by Horde's build automation system can reuse uploaded symbols without any additional storage requirements, and the lifetime of the uploaded symbols is tied to that of the artifact.
+
+Internally, this mapping is implemented using _aliases_ in Horde's storage system which map from symbol server hashes to files in uploaded artifacts. See [Internals > Storage Architecture](../Internals/StorageArchitecture.md) for more information on aliasing.
+
+## Uploading Symbols
+
+To add the metadata used for indexing symbols as part of the artifact upload process, set the `Symbols="True"` attribute on the `CreateArtifact` task from BuildGraph. 
+
+## Adding Source Information
+
+In order to allow source-level debugging, it can be useful to embed information about the source files that are built in the PDB itself. When configured in this way, Visual Studio supports downloading the matching source file for a PDB automatically when stepping through a program or debugging a crash dump.
+
+The functionality is referred to as _source server_ support in Microsoft literature. This terminology can be a little misleading; the source server being referred to is that of an existing source control server, and not a new server deployment. When embedding source information in a PDB file, you effectively embed a small script that Visual Studio can execute to fetch the matching source file if needed - and that script can execute arbitrary commands on the host machine (for example, launching p4.exe with particular arguments). 
+
+The `SrcSrv` BuildGraph task can be used to embed Perforce source server information into PDB files. The `Engine/Build/InstalledEngineBuild.xml` demonstrates adding source information into installed builds as follows:
+
+  ```xml
+  <Tag
+    Files="Engine/Source/...;Engine/Plugins/..."
+    Filter="*.c;*.h;*.cpp;*.hpp;*.inl"
+    Except="Engine/Source/ThirdParty/..."
+    With="#SourceFiles"/>
+
+  <SrcSrv
+    BinaryFiles="#UnrealEditor Win64"
+    SourceFiles="#SourceFiles"
+    Branch="$(Branch)"
+    Change="$(Change)"/>
+  ```
+
+Source server support needs to be manually enabled in Visual Studio through the `Tools > Options` window, then checking the `Debugging > General > Enable Source Server Support` option.
+
+## Configuring a Symbol Store
+
+Configuration for the symbols plugin is through the `SymbolsConfig` type in the [globals.json](Schema/Globals.md#symbolsconfig) file.
+
+Each configured symbol store will index an entire storage namespace. Permission to access a particular store will grant access to any indexed data, regardless of the user's ability to access that storage namespace directly.
+
+  ```json
+  {
+    "stores": [
+        {
+            "id": "my-store-id",
+            "namespaceId": "horde-artifacts",
+            "public": true
+        }
+    ]
+  }
+
+  ```
+
+You can configure Visual Studio to use a symbol store through the `Tools > Options` dialog, selecting the `Debugging > Symbols` tab.
+
+Add Horde as a symbol server using a HTTP URL with a base path containing the symbol store id:
+
+  ```
+  http://my-horde-server/api/v1/symbols/my-store-id
+  ```
+
+## Accessing a Symbol Store with Authentication
+
+Setting the `"public": true` property on a symbol store makes it accessible to anyone via unauthenticated endpoints.
+
+Microsoft's symbol server API does not directly support authentication, but **Unreal Toolbox** implements a proxy server that can provide an unauthenticated bridge to access the Horde server over an authenticated connection via a port bound to `localhost`.
+
+To use this, install Unreal Toolbox from the Tools > Download page in the Horde dashboard, and enable the proxy server plugin from the settings menu. After enabling the plugin, you can configure Visual Studio to search for symbols on http://localhost:13344.
diff --git a/Engine/Source/Programs/Horde/Docs/Internals.md b/Engine/Source/Programs/Horde/Docs/Internals.md
--- a/Engine/Source/Programs/Horde/Docs/Internals.md
+++ b/Engine/Source/Programs/Horde/Docs/Internals.md
@@ -4,54 +4,11 @@
 
 ## Getting Started
 
-### Server
+* [Building Horde](Internals/Building.md)
 
-The source code for the Horde Server is in `Engine/Source/Programs/Horde/Horde.Server`. It is written in C#, using ASP.NET.
+## Features
 
-Horde uses the standard
-[C# coding conventions](https://learn.microsoft.com/en-us/dotnet/csharp/fundamentals/coding-style/coding-conventions)
-published by Microsoft, though we use tabs rather than spaces for legacy reasons. We enable most static analyzer warnings
-that ship with the NET SDK, but you may disable some through `.editorconfig` files.
-
-Horde is configured to support local development by default. You can launch it by opening
-`Engine/Source/Programs/Horde/Horde.sln` and setting Horde.Server as the default project. 
-By default, you can access the server at `http://localhost:5000/account`.
-
-When debugging a local Horde server against a live deployment, setting the `DatabaseReadOnlyMode` property in
-[`Server.json`](Deployment/ServerSettings.md) prevents the server from attempting any operation that modifies the server state. 
-Using a read-only DB account in addition is recommended for safety.
-
-### Dashboard
-
-The Horde dashboard is a frontend client developed in [TypeScript](https://www.typescriptlang.org) using
-[React](https://react.dev/). To set up your machine for developing the dashboard:
-
-1. Install [Node.js](https://nodejs.org/en/download).
-2. From a command line, install Yarn using: `npm install --global yarn`.
-3. Navigate to the dashboard folder at `Engine\Source\Programs\Horde\HordeDashboard`.
-4. Run `yarn install` to install the package dependencies.
-5. Edit package.json setting the proxy property to point at your server URL, for example: `http://localhost:13340`.
-6. Navigate to the admin token endpoint of your server to get an expiring access token, for example:
-   `http://localhost:13340/api/v1/admin/token`.
-7. Create a file called `.env.development.local` in the root HordeDashboard folder and paste the access token in like
-   so: `REACT_APP_HORDE_DEBUG_TOKEN=eyFhbGciziJIUz`.
-8. Run `yarn start` to start the development web server which should open a tab to the local dashboard at
-   `http://localhost:3000`.
-
-## Docker
-
-Horde includes a `Dockerfile` for creating Docker images. However, its position in Unreal Engine source tree requires
-files to be staged beforehand to reduce the size of data copied into build images.
-
-A BuildGraph script to perform these operations is included under `Engine\Source\Programs\Horde\HordeBuild.xml`, which
-can be run as follows:
-
-    RunUAT.bat Engine/Source/Programs/Horde/HordeBuild.xml -Target="Build HordeServer"
-
-## Topics
-
-* [Build Health](Internals/BuildHealth.md)
-* [Storage](Internals/Storage.md)
-* [Leases](Internals/Leases.md)
+* [Compute Architecture](Internals/ComputeArchitecture.md)
+* [Storage Architecture](Internals/StorageArchitecture.md)
 * [Logs](Internals/Logs.md)
 * [Structured Logging](Internals/StructuredLogging.md)
diff --git a/Engine/Source/Programs/Horde/Docs/Internals/Building.md b/Engine/Source/Programs/Horde/Docs/Internals/Building.md
new file mode 100644
--- /dev/null
+++ b/Engine/Source/Programs/Horde/Docs/Internals/Building.md
@@ -0,0 +1,47 @@
+[Horde](../../README.md) > [Internals](../Internals.md) > Building Horde
+
+# Building Horde
+
+## Server
+
+The source code for the Horde Server is in `Engine/Source/Programs/Horde/Horde.Server`. It is written in C#, using ASP.NET.
+
+Horde uses the standard
+[C# coding conventions](https://learn.microsoft.com/en-us/dotnet/csharp/fundamentals/coding-style/coding-conventions)
+published by Microsoft, though we use tabs rather than spaces for legacy reasons. We enable most static analyzer warnings
+that ship with the NET SDK, but you may disable some through `.editorconfig` files.
+
+Horde is configured to support local development by default. You can launch it by opening
+`Engine/Source/Programs/Horde/Horde.sln` and setting Horde.Server as the default project. 
+By default, you can access the server at `http://localhost:5000/account`.
+
+When debugging a local Horde server against a live deployment, setting the `DatabaseReadOnlyMode` property in
+[`Server.json`](Deployment/ServerSettings.md) prevents the server from attempting any operation that modifies the server state. 
+Using a read-only DB account in addition is recommended for safety.
+
+## Dashboard
+
+The Horde dashboard is a frontend client developed in [TypeScript](https://www.typescriptlang.org) using
+[React](https://react.dev/). To set up your machine for developing the dashboard:
+
+1. Install [Node.js](https://nodejs.org/en/download).
+2. From a command line, install Yarn using: `npm install --global yarn`.
+3. Navigate to the dashboard folder at `Engine\Source\Programs\Horde\HordeDashboard`.
+4. Run `yarn install` to install the package dependencies.
+5. Edit package.json setting the proxy property to point at your server URL, for example: `http://localhost:13340`.
+6. Navigate to the admin token endpoint of your server to get an expiring access token, for example:
+   `http://localhost:13340/api/v1/admin/token`.
+7. Create a file called `.env.development.local` in the root HordeDashboard folder and paste the access token in like
+   so: `REACT_APP_HORDE_DEBUG_TOKEN=eyFhbGciziJIUz`.
+8. Run `yarn start` to start the development web server which should open a tab to the local dashboard at
+   `http://localhost:3000`.
+
+## Docker Image
+
+Horde includes a `Dockerfile` for creating Docker images. However, its position in Unreal Engine source tree requires
+files to be staged beforehand to reduce the size of data copied into build images.
+
+A BuildGraph script to perform these operations is included under `Engine\Source\Programs\Horde\HordeBuild.xml`, which
+can be run as follows:
+
+    RunUAT.bat Engine/Source/Programs/Horde/HordeBuild.xml -Target="Build HordeServer"
diff --git a/Engine/Source/Programs/Horde/Docs/Internals/ComputeArchitecture.md b/Engine/Source/Programs/Horde/Docs/Internals/ComputeArchitecture.md
new file mode 100644
--- /dev/null
+++ b/Engine/Source/Programs/Horde/Docs/Internals/ComputeArchitecture.md
@@ -0,0 +1,40 @@
+[Horde](../../README.md) > [Internals](../Internals.md) > Compute Architecture
+
+# Compute Architecture
+
+The mechanism by which Horde communicates with its agents is based heavily
+on Google's Remote Worker API. Work items assigned to agents are known as
+*leases*.
+
+Communication between the agent and server is done through streaming gRPC calls initiated by the agent. The ends of the
+connection exchange copies of what they believe the current state of leases owned by the agent should be, and a state
+machine determines which field is authoritative at different times. The server and agent keep exchanging state objects
+and reconciling differences until they are in sync; at that point, the agent has acknowledged that it is taking any
+new leases added by the server, and the server has acknowledged the completion of any removed leases by the agent.
+
+Within the Horde server, leases are assigned to agents through classes implementing `ITaskSource`.
+
+## Lease Types
+
+Each lease type consists of the following key classes:
+
+* a message defining the lease itself (i.e. what the agent needs to do),
+* a task source on the server that decides to assign a lease to an agent (`ITaskSource`),
+* a lease handler on the agent which does the work (`LeaseHandler`).
+
+Horde ships with the following lease types:
+
+| Message | Task source (server) | Lease handler (agent) | Description |
+| ---------------- | --- | --- | --- |
+| `job_task.proto` | `JobTaskSource` | `JobHandler` | Executes a batch as part of a CI job. |
+| `upgrade_task.proto` | `UpgradeTaskSource` | `UpgradeHandler` | Upgrades the Horde agent software to a newer version. |
+| `conform_task.proto` | `ConformTaskSource` | `ConformHandler` | Syncs all the workspaces on a machine to latest and optionally removes any untracked files. |
+| `restart_task.proto` | `RestartTaskSource` | `RestartHandler` | Restarts the machine. |
+| `shutdown_task.proto` | `ShutdownTaskSource` | `ShutdownHandler` | Powers down the machine. |
+
+## Adding New Lease Types
+
+To add a new lease type, each of the key classes above must be added.
+
+* After adding a task source to the server, register it as a singleton implementing `ITaskSource` in `Startup.cs`.
+* After adding a lease handler to the agent, register it as a singleton implementing `LeaseHandler` in `Program.cs`.
diff --git a/Engine/Source/Programs/Horde/Docs/Internals/Leases.md b/Engine/Source/Programs/Horde/Docs/Internals/Leases.md
deleted file mode 100644
diff --git a/Engine/Source/Programs/Horde/Docs/Internals/Storage.md b/Engine/Source/Programs/Horde/Docs/Internals/Storage.md
deleted file mode 100644
diff --git a/Engine/Source/Programs/Horde/Docs/Internals/StorageArchitecture.md b/Engine/Source/Programs/Horde/Docs/Internals/StorageArchitecture.md
new file mode 100644
--- /dev/null
+++ b/Engine/Source/Programs/Horde/Docs/Internals/StorageArchitecture.md
@@ -0,0 +1,151 @@
+[Horde](../../README.md) > [Internals](../Internals.md) > Storage
+
+# Storage Architecture
+
+Horde's storage platform is designed to support manipulating massive data structures consisting of
+interlinked blobs. Blobs are **immutable**, and consist of an arbitrary block of data and zero
+or more outward references to other blobs.
+
+The entry point to any such data structure is a named reference, which maps a
+user-defined name to a blob at the root of the data structure. Any nodes that are
+not directly or indirectly referenced by a named reference are subject to garbage
+collection.
+
+## History
+
+Horde's storage system can be seen as
+an evolution of the GitSync tool used to download binaries for Epic's GitHub repo;  for each change
+committed to Epic's Perforce server, we upload matching binaries to AWS S3 and commit a manifest to those
+files to the Git repository that can be used to download those files. The GitSync tool uses the manifest to retrieve and
+unpack those files whenever a new commit is checked out via a Git hook that is installed by running the Setup.bat file
+in the root of the repo.
+
+One of the main design goals for GitSync was to offload the hosting of binary data to a proven, scalable third-party
+storage service (AWS S3) without having to maintain an active server capable of supporting many Unreal Engine developers. 
+As such, we kept the ideas of content addressing used by Git but
+packed the content-addressed payloads into non-deterministic packages for more efficient downloads. At upload
+time, we use a heuristic to decide whether to reuse existing packages to download data versus re-packing
+them into new packages to avoid expensive gather operations.
+
+While clients can still model the data as a Git-like Merkle tree - using any locally cached data they may
+have available using a uniquely identifying SHA1 hash - we reduce the chattiness and server-side
+compute load in negotiating data to transfer to clients by putting the data into pre-made, static
+download packages that are arranged to optimize for coherency between blobs we anticipate to be requested together.
+
+This model optimizes for streaming reads and writes while still accomodating point reads when necessary.
+
+## Blobs
+
+A blob in Horde has the following attributes (see `BlobData` class):
+
+* **Type**: Represented by a GUID and integer version number and used to distinguish between payloads that
+may have a particular serialization format but the same data.
+
+* **Payload**: A byte array. Blobs are meant to be fully read into memory, so payloads are typically limited
+to a few hundred kb. Larger payloads can be split into smaller blobs using static or
+content-defined chunking using utility libraries.
+
+* **References**: A set of references to other blobs.
+
+References to blobs are typically manipulated in memory using `IBlobHandle` instances. After flushing to
+storage, blob handles can be converted to and from a `BlobLocator`, an opaque implementation-defined
+string identifier assigned by the storage system.
+
+Horde abstracts the way that blobs are serialized to the underlying storage backend through the `IBlobWriter`
+interface; the programmer requests a buffer to serialize a blob to, writes the data and its references, and
+gets a `IBlobHandle` back to allow retrieving it at any point in the future. The implementation is left to
+decide how to store the data, implementing compression, packing, buffering, and uploading the data as
+necessary.
+
+Multiple `IBlobWriter` instances can be created to write separate streams of related blobs.
+
+## Refs and Aliases
+
+Refs and aliases provide entry points into the storage system. With either, you can assign a user-defined name to a particular blob and retrieve it later.
+
+* **Refs** are strong references into the blob store and act as roots for the garbage collector. Refs can be
+set to expire at fixed times or after not being retrieved for a specific period of time - which can be useful
+for implementing caches.
+
+* **Aliases** are weak references to blobs. Multiple aliases with the same name may exist, and users can
+query for one or more blobs with a particular alias. Aliases have an associated user-specified rank, and
+consumers can query aliases ordered by rank.
+
+## Content Addressing
+
+Horde models blobs in a way that supports content addressing. While hashes are not exposed
+directly through `BlobLocator` strings, hashes can be encoded into a blob's payload, which has a matching
+entry in the references array.
+
+Since references are stored separately from the blob payload, the unique identifier stored through a `BlobLocator`
+does not affect the payload's hash.
+
+The implementation primarily uses `IoHash` for hashing blob data (a truncated 20-byte Blake 3 hash), but
+decoupling the encoding of hashes in the payload from references in the storage system allows any other
+hashing algorithm to be used instead. The underlying storage system can reason about the topology of blob
+trees while still supporting a variety of hashing algorithms.
+
+The `IBlobRef` interface extends the basic `IBlobHandle` interface with an `IoHash` of the target blob.
+
+## Implementation
+
+> **Note:** This section describes the storage system's current implementation details, which may change in future releases.
+
+### Layers
+
+The storage system is implemented through several layers:
+
+* The C# serialization library (`BlobSerializer`, `BlobConverter`, and so on).
+
+* The logical storage model is declared through the `IStorageClient` interface, which is the primary method for
+  interacting with the storage system.
+  At this layer, blobs are manipulated via `IBlobHandle` objects.
+
+  * `BundleStorageClient` is the standard implementation of `IStorageClient`
+  and packs blobs into [bundles](#bundles).
+
+  * `KeyValueStorageClient` implements a client that passes individual
+  blobs to the underlying `IStorageBackend`.
+
+* The physical storage model is declared through the `IStorageBackend` interface, which deals with sending data over
+  the wire to storage service implementations.
+  
+  * `HttpStorageBackend` uploads data to the Horde Server over HTTP.
+
+  * `FileStorageBackend` writes data directly to files on disk.
+
+  * `MemoryStorageBackend` stores data in memory.
+
+* The bulk data store is declared through the `IObjectStore` interface, which interacts with low-level
+  storage services.
+  
+  * `FileObjectStore` writes data to files on disk.
+
+  * `AwsObjectStore` reads and writes data from AWS S3.
+
+  * `AzureObjectStore` reads and writes data from the Azure blob store.
+
+### Bundles
+
+Blobs are designed to be used as a general-purpose storage primitive, so we try to
+efficiently accommodate blob types
+ranging from a handful of bytes to several hundred kilobytes (larger streams of data can be split up into smaller chunks
+along fixed boundaries or using content-defined slicing).
+
+Blobs are packed together into bundles for storage in the underlying object store.
+
+The implementation of bundles and their use within the storage system is mostly hidden from user code, though
+knowing how a stream of blobs will be written to storage may help reasoning about access patterns.
+
+Each bundle consists of a sequence of compressed packets, each of which may contain several blobs. Each packet is
+self-contained, so it may be decoded from a single contiguous ranged read of the bundle data.
+
+### Locators
+
+Locators typically have the following form:
+
+    [path]#pkt=[offset],[length]&exp=[index]
+
+* `[path]`: Path to an object within the underlying object store.
+* `[offset]` and `[length]`: Byte range of the compressed packet data within a bundle.
+* `[index]`: The index of an exported blob within the packet.
