Horde: Indicate that Horde is production ready in documentation.

#rnx

diff --git a/Engine/Source/Programs/Horde/Docs/Home.md b/Engine/Source/Programs/Horde/Docs/Home.md
--- a/Engine/Source/Programs/Horde/Docs/Home.md
+++ b/Engine/Source/Programs/Horde/Docs/Home.md
@@ -27,11 +27,9 @@ Horde's reference documentation is divided into sections by target audience:
 
 ## Status
 
-Horde is under heavy development, and large parts of it are still in flux. While we use aspects of it
-(particularly the CI system) heavily at Epic, we consider it experimental for Unreal Engine licensees
-and offer limited support for it.
+Horde is considered production ready for Unreal Engine licensees.
 
-See also: [Feature Status](Features.md)
+See: [Feature Status](Features.md)
 
 ## Further Reading
 
diff --git a/Engine/Source/Programs/Horde/README.md b/Engine/Source/Programs/Horde/README.md
--- a/Engine/Source/Programs/Horde/README.md
+++ b/Engine/Source/Programs/Horde/README.md
@@ -25,9 +25,7 @@ Read more about our [goals and philosophy](Docs/Goals.md), or check out the [FAQ
 
 ## Status
 
-Horde is under heavy development, and large parts of it are still in flux. While we use aspects of it (particularly
-the CI system) heavily at Epic, we consider it experimental for Unreal Engine licensees and offer limited support
-for it.
+Horde is considered production ready for Unreal Engine licensees.
 
 See also: [Feature Status](Docs/Features.md)
 
