Horde: Fix warnings in documentation.

#rnx

diff --git a/Engine/Source/Programs/Horde/Docs/Config/BuildHealth.md b/Engine/Source/Programs/Horde/Docs/Config/BuildHealth.md
--- a/Engine/Source/Programs/Horde/Docs/Config/BuildHealth.md
+++ b/Engine/Source/Programs/Horde/Docs/Config/BuildHealth.md
@@ -6,7 +6,7 @@ Epic has a large development team working on products that are continously shipp
 use a trunk-based development model with branches for each release, with
 commits to release branches automatically being merged back to the mainline.
 
-Maintaining a stable and healthy codebase is important to our ability to iterate quickly, 
+Maintaining a stable and healthy codebase is important to our ability to iterate quickly,
 so we have a number of policies in place to help:
 
 * Local test suites that developers can run before changes are submitted.
@@ -47,7 +47,8 @@ Issues will be marked as fixed automatically after a successful build, but users
 can indicate that an issue has been addressed via Slack, the Horde Dashboard, or
 by including a `#horde 1234` tag in a separate line in a commit description.
 
-The tag name can be customized using the `IssueFixedTag` setting in the [globals.json](../Config/Schema/Globals.md#buildconfig) file.
+The tag name can be customized using the `IssueFixedTag` setting in the
+[globals.json](../Config/Schema/Globals.md#buildconfig) file.
 
 ## Issues, Spans, and Fingerprints
 
@@ -68,7 +69,7 @@ A fingerprint contains the following data (see `IIssueFingerprint.cs`):
 
 * __Type__: A string identifying the type of issue. Fingerprints only match with other fingerprints of the same type.
   One of the most commonly encountered issue types is `Compile`, indicating a compile error in one or more source files.
-* __SummaryTemplate__: A string that is used to populate an issue's description in the Horde dashboard and notifications. 
+* __SummaryTemplate__: A string that is used to populate an issue's description in the Horde dashboard and notifications.
   May contain placeholder strings that will be expanded when fingerprints are merged together
   (see [Summary Templates](#summary-templates)).
 * __Keys__: Pieces of identifying data used to match and group fingerprints,
@@ -101,9 +102,9 @@ Fingerprints are matched by logic in `IIssueFingerprint.cs`. The pipeline is as
 
 Fingerprints may be created for errors in two ways:
 
-* By including fingerprint information directly in the [structured log event](StructuredLogging.md) at the
+* By including fingerprint information directly in the [structured log event](../Internals/StructuredLogging.md) at the
   point that it is generated.
-* By post-processing [structured log events](StructuredLogging.md) in the Horde server once a build step 
+* By post-processing [structured log events](../Internals/StructuredLogging.md) in the Horde server once a build step
   completes.
 
 Including fingerprints directly in structured log events is preferred, since it
@@ -143,7 +144,7 @@ until a handler is found which returns `true`. Once all events have been process
 returns a list of the matched log events and their fingerprints.
 
 Note that handlers do not have to generate fingerprints for a log event; it can be useful
-for more specific errors (say, a compile error) to mask more vague errors (say, UnrealBuildTool 
+for more specific errors (say, a compile error) to mask more vague errors (say, UnrealBuildTool
 returning a failing exit code). As such, `HandleEvent` can return true for errors that it wants
 to exclude from the remaining handler pipeline.
 
@@ -190,7 +191,8 @@ The sending of reports can be configured through the `reportTimes` and `reportCh
 
 ### Annotations
 
-Handling of issues for individual job steps may be further configured using __node attributes__. Attributes can be specified for a node through the `Annotations` attribute on `Node` elements in BuildGraph scripts:
+Handling of issues for individual job steps may be further configured using __node attributes__. Attributes can be
+specified for a node through the `Annotations` attribute on `Node` elements in BuildGraph scripts:
 
   ```xml
   <Node Name="Compile UnrealEditor Win64" Annotations="Workflow=my-workflow;BuildBlocker=true">
@@ -200,8 +202,12 @@ The supported set of annotations are defined in the `NodeAnnotations.cs` source
 
 * `Workflow`: Workflow to use for triaging issues from this node
 * `CreateIssues`: Allows disabling creation of issues for this node, when set to `false`.
-* `AutoAssign`: Whether to automatically assign issues that could only be caused by one user, or have a well defined correlation with a modified file.
+* `AutoAssign`: Whether to automatically assign issues that could only be caused by one user, or have a well defined
+  correlation with a modified file.
 * `AutoAssignToUser`: Automatically assign any issues in this step to the Perforce user given by the supplied argument.
-* `NotifySubmitters`: Whether to notify all submitters between a build suceeding and failing, allowing them to step forward and take ownership of an issue.
-* `IssueGroup`: Specifies a suffix to be appended to the `type` property of any issues created from this node, preventing them being merged with other issues.
-* `BuildBlocker`: Whether failures in this node should be considered a build blocker. Issues identified as build blockers have a special tag in Slack notifications.
+* `NotifySubmitters`: Whether to notify all submitters between a build suceeding and failing, allowing them to step
+  forward and take ownership of an issue.
+* `IssueGroup`: Specifies a suffix to be appended to the `type` property of any issues created from this node,
+  preventing them being merged with other issues.
+* `BuildBlocker`: Whether failures in this node should be considered a build blocker. Issues identified as build
+  blockers have a special tag in Slack notifications.
diff --git a/Engine/Source/Programs/Horde/Docs/Config/Plugins.md b/Engine/Source/Programs/Horde/Docs/Config/Plugins.md
--- a/Engine/Source/Programs/Horde/Docs/Config/Plugins.md
+++ b/Engine/Source/Programs/Horde/Docs/Config/Plugins.md
@@ -9,9 +9,6 @@ others. The [Storage](../Config.md#storage) and
 [Compute](../Config.md#compute) plugins are notable foundational components
 of Horde that other systems rely on.
 
-For information about how plugins operate and how to add new plugins, see
-[Internals > Plugin Architecture](Internals/PluginArchitecture.md).
-
 ## Configuration
 
 Like the server itself, plugins can have a static boot-time configuration
diff --git a/Engine/Source/Programs/Horde/Docs/Internals/Building.md b/Engine/Source/Programs/Horde/Docs/Internals/Building.md
--- a/Engine/Source/Programs/Horde/Docs/Internals/Building.md
+++ b/Engine/Source/Programs/Horde/Docs/Internals/Building.md
@@ -16,8 +16,8 @@ Horde is configured to support local development by default. You can launch it b
 By default, you can access the server at `http://localhost:5000/account`.
 
 When debugging a local Horde server against a live deployment, setting the `DatabaseReadOnlyMode` property in
-[`Server.json`](Deployment/ServerSettings.md) prevents the server from attempting any operation that modifies the server state. 
-Using a read-only DB account in addition is recommended for safety.
+[`Server.json`](../Deployment/ServerSettings.md) prevents the server from attempting any operation that modifies the
+server state. Using a read-only DB account in addition is recommended for safety.
 
 ## Dashboard
 
