Horde: Add missing file for test

diff --git a/Engine/Source/Programs/Horde/Plugins/Build/HordeServer.Build.Tests/Agents/AgentServiceTest.cs b/Engine/Source/Programs/Horde/Plugins/Build/HordeServer.Build.Tests/Agents/AgentServiceTest.cs
--- a/Engine/Source/Programs/Horde/Plugins/Build/HordeServer.Build.Tests/Agents/AgentServiceTest.cs
+++ b/Engine/Source/Programs/Horde/Plugins/Build/HordeServer.Build.Tests/Agents/AgentServiceTest.cs
@@ -206,7 +206,7 @@ public class AgentServiceTest : BuildTestSetup
 	
 	[TestMethod]
 	// 5 hour old sessions are deleted
-//	[DataRow(5, null, true)]
+	[DataRow(5, null, true)]
 	[DataRow(5, AgentStatus.Stopped, true)]
 	[DataRow(5, AgentStatus.Ok, false)]
 	[DataRow(5, null, false, false)] // Non-ephemeral agent should never be deleted
@@ -218,7 +218,7 @@ public class AgentServiceTest : BuildTestSetup
 	public async Task Ephemeral_IsDeletedInBackgroundTick_Async(double lastOnlineAgoHours, AgentStatus? status, bool isDeleted, bool isEphemeral = true)
 	{
 		TimeSpan lastOnlineAgo = TimeSpan.FromHours(lastOnlineAgoHours);
-		IAgent? agent = await CreateAgentAsync(new PoolId("pool1"), ephemeral: isEphemeral, adjustClockBy: -lastOnlineAgo, status: status ?? AgentStatus.Ok);
+		IAgent? agent = await CreateAgentAsync(new PoolId("pool1"), ephemeral: isEphemeral, adjustClockBy: -lastOnlineAgo, status: status);
 		agent = await AgentService.GetAgentAsync(agent.Id);
 		Assert.IsNotNull(AgentService.GetAgentAsync(agent!.Id));
 		await ServiceProvider.GetRequiredService<AgentCollection>().DeleteExpiredEphemeralAgentsAsync(CancellationToken.None);
diff --git a/Engine/Source/Programs/Horde/Plugins/Build/HordeServer.Build.Tests/BuildTestSetup.cs b/Engine/Source/Programs/Horde/Plugins/Build/HordeServer.Build.Tests/BuildTestSetup.cs
--- a/Engine/Source/Programs/Horde/Plugins/Build/HordeServer.Build.Tests/BuildTestSetup.cs
+++ b/Engine/Source/Programs/Horde/Plugins/Build/HordeServer.Build.Tests/BuildTestSetup.cs
@@ -289,7 +289,7 @@ namespace HordeServer.Tests
 			string? awsInstanceId = null,
 			CreateLeaseOptions? lease = null,
 			bool ephemeral = false,
-			AgentStatus status = AgentStatus.Ok)
+			AgentStatus? status = AgentStatus.Ok)
 		{
 			DateTime now = Clock.UtcNow;
 			if (adjustClockBy != null)
@@ -308,14 +308,17 @@ namespace HordeServer.Tests
 
 			agent = await agent.TryUpdateAsync(new UpdateAgentOptions { Enabled = enabled, ExplicitPools = poolId != null ? [poolId.Value] : [] });
 			Assert.IsNotNull(agent);
-
-			agent = await AgentService.CreateSessionAsync(agent, new RpcAgentCapabilities(tempProps), null);
-			Assert.IsNotNull(agent);
-
-			if (status != AgentStatus.Ok)
+			
+			if (status != null)
 			{
-				agent = await agent.TryUpdateSessionAsync(new UpdateSessionOptions { Status = status });
+				agent = await AgentService.CreateSessionAsync(agent, new RpcAgentCapabilities(tempProps), null);
 				Assert.IsNotNull(agent);
+				
+				if (status != AgentStatus.Ok)
+				{
+					agent = await agent.TryUpdateSessionAsync(new UpdateSessionOptions { Status = status });
+					Assert.IsNotNull(agent);
+				}
 			}
 
 			if (workspaces is { Count: > 0 })
@@ -336,11 +339,40 @@ namespace HordeServer.Tests
 			Clock.UtcNow = now;
 			return agent;
 		}
-
-		protected async Task<IPool> CreatePoolAsync(PoolConfig poolConfig)
+		
+		protected Task<IPool> CreatePoolAsync(PoolConfig poolConfig, bool viaConfig = true)
+		{
+			return viaConfig ? CreatePoolViaConfigAsync(poolConfig) : CreatePoolViaDatabaseAsync(poolConfig);
+		}
+		
+		protected async Task<IPool> CreatePoolViaConfigAsync(PoolConfig poolConfig)
 		{
 			UpdateConfig(config => config.Plugins.GetComputeConfig().Pools.Add(poolConfig));
-			return await PoolCollection.GetAsync(poolConfig.Id) ?? throw new NotImplementedException();
+			return await PoolCollection.GetAsync(poolConfig.Id) ?? throw new Exception($"Unable to get pool ID {poolConfig.Id}");
+		}
+		
+		protected async Task<IPool> CreatePoolViaDatabaseAsync(PoolConfig poolConfig)
+		{
+#pragma warning disable CS0618 // Type or member is obsolete
+			await PoolCollection.CreateConfigAsync(poolConfig.Id, poolConfig.Name, new CreatePoolConfigOptions()
+			{
+				Condition = poolConfig.Condition,
+				EnableAutoscaling = poolConfig.EnableAutoscaling,
+				MinAgents = poolConfig.MinAgents,
+				NumReserveAgents = poolConfig.NumReserveAgents,
+				ConformInterval = poolConfig.ConformInterval,
+				ScaleOutCooldown = poolConfig.ScaleOutCooldown,
+				ScaleInCooldown = poolConfig.ScaleInCooldown,
+				SizeStrategies = poolConfig.SizeStrategies,
+				FleetManagers = poolConfig.FleetManagers,
+				SizeStrategy = poolConfig.SizeStrategy,
+				LeaseUtilizationSettings = poolConfig.LeaseUtilizationSettings,
+				JobQueueSettings = poolConfig.JobQueueSettings,
+				ComputeQueueAwsMetricSettings = poolConfig.ComputeQueueAwsMetricSettings,
+				Properties = poolConfig.Properties,
+			});
+			return await PoolCollection.GetAsync(poolConfig.Id) ?? throw new Exception($"Unable to get pool ID {poolConfig.Id}");
+#pragma warning restore CS0618 // Type or member is obsolete
 		}
 	}
 }
\ No newline at end of file
