Horde: Updating Release Notes
#rnx

diff --git a/Engine/Source/Programs/Horde/Docs/ReleaseNotes.md b/Engine/Source/Programs/Horde/Docs/ReleaseNotes.md
--- a/Engine/Source/Programs/Horde/Docs/ReleaseNotes.md
+++ b/Engine/Source/Programs/Horde/Docs/ReleaseNotes.md
@@ -1,5 +1,19 @@
 # Release Notes
 
+## 2024-10-17
+
+* Delete ephemeral agents that have been offline for more than X hours. This ignores checking for any outdated sessions and simply uses the last online timestamp instead. (37172088)
+* Show logs event data as generating for running steps when missing (37125109)
+* Always show the reason for why compute resources may not be available (37114936)
+* Do not retry HTTP requests which throw an HttpRequestException. These generally represent a non-transient error, unlike HTTP status responses. In the case of bad server addresses or ports, requests may take a long time to timeout. Retrying them does not make sense, and creates very long delays before the result is reported. (37093049)
+* Cloud DDC - Changed server timing metrics into a concurrent bag instead of list with locks to reduce lock time a bit. (37085986)
+* Fix for calling Remote UFUNCTIONs from other Remote UFUNCTIONs. (37066038)
+* Instrument ManagedWorkspace with OpenTelemetry tracing (37059773)
+* Remove use of resource.name in spans resource.name is Datadog concept and doesn't play well presenting traces from OpenTelemetry in the Web UI. (37051252)
+* Fix bug with OpenTelemetry settings not being base64 encoded (37049676)
+* Add Google Cloud Storage support (37048260)
+* Pass OpenTelemetry settings to agent driver subprocess (37047438)
+
 ## 2024-10-10
 
 * EpicGames.Perforce: Assume and mark workspace having untracked files prior to starting a sync (36972802)
