Horde: Fix agent.json configuration not being read from data directory

diff --git a/Engine/Source/Programs/Horde/HordeAgent/AgentApp.cs b/Engine/Source/Programs/Horde/HordeAgent/AgentApp.cs
--- a/Engine/Source/Programs/Horde/HordeAgent/AgentApp.cs
+++ b/Engine/Source/Programs/Horde/HordeAgent/AgentApp.cs
@@ -14,6 +14,7 @@ using HordeAgent.Services;
 using HordeAgent.Utility;
 using Microsoft.Extensions.Configuration;
 using Microsoft.Extensions.DependencyInjection;
+using Microsoft.Extensions.FileProviders;
 using Microsoft.Extensions.Logging;
 using Microsoft.Win32;
 using Polly;
@@ -110,12 +111,11 @@ namespace HordeAgent
 				configOverrides.Add($"{AgentSettings.SectionName}:{nameof(AgentSettings.WorkingDir)}", workingDirOverride);
 			}
 			
-			List<string> configSources;
-			List<string> configSourcesInstalled = new();
+			List<string> configSourcesInstalled = [];
 			
 			// Create the base configuration data by just reading from the application directory. We need to check some settings before
 			// being able to read user configuration files.
-			IConfiguration configuration = CreateConfig(false, null, configOverrides, out configSources);
+			IConfiguration configuration = CreateConfig(false, null, configOverrides, out _);
 			AgentSettings settings = BindSettings(configuration);
 
 			if (settings.Installed)
@@ -138,9 +138,8 @@ namespace HordeAgent
 			using ILoggerFactory loggerFactory = Logging.CreateLoggerFactory(configuration);
 			
 			{
-				configSources.AddRange(configSourcesInstalled);
 				ILogger logger = loggerFactory.CreateLogger(typeof(AgentApp));
-				foreach (string configSource in configSources)
+				foreach (string configSource in configSourcesInstalled)
 				{
 					logger.LogInformation("Config source: {ConfigSource}", configSource);
 				}
@@ -276,7 +275,14 @@ namespace HordeAgent
 			
 			if (agentConfigFile != null)
 			{
-				builder = builder.AddJsonFile(agentConfigFile.FullName, optional: true, reloadOnChange: true);
+				// Adding a JSON file outside the common base path requires adding a completely separate JsonConfigurationSource entry
+				builder.Add(new JsonConfigurationSource()
+				{
+					Path = agentConfigFile.GetFileName(),
+					Optional = true,
+					ReloadOnChange = true,
+					FileProvider = new PhysicalFileProvider(agentConfigFile.Directory.FullName),
+				});
 			}
 			
 			foreach (IConfigurationSource source in builder.Sources)
@@ -284,7 +290,10 @@ namespace HordeAgent
 				switch (source)
 				{
 					case JsonConfigurationSource jcs:
-						configSources.Add("JSON file: " + Path.Join(basePath, jcs.Path));
+						string path = jcs.FileProvider is PhysicalFileProvider pfp
+							? Path.Join(pfp.Root, jcs.Path)
+							: Path.Join(basePath, jcs.Path); // Use shared base path if file provider is missing
+						configSources.Add("JSON file: " + path);
 						break;
 					case RegistryConfigurationSource rcs:
 						if (OperatingSystem.IsWindows())
