Horde: Fix exception when specifying custom permissions for artifact types. Also add tests for artifact type permission inheritance.

#rnx

diff --git a/Engine/Source/Programs/Horde/Plugins/Build/HordeServer.Build.Tests/Artifacts/ArtifactAclTests.cs b/Engine/Source/Programs/Horde/Plugins/Build/HordeServer.Build.Tests/Artifacts/ArtifactAclTests.cs
new file mode 100644
--- /dev/null
+++ b/Engine/Source/Programs/Horde/Plugins/Build/HordeServer.Build.Tests/Artifacts/ArtifactAclTests.cs
@@ -0,0 +1,109 @@
+// Copyright Epic Games, Inc. All Rights Reserved.
+
+using System.Security.Claims;
+using EpicGames.Horde.Acls;
+using EpicGames.Horde.Artifacts;
+using EpicGames.Horde.Projects;
+using EpicGames.Horde.Streams;
+using HordeServer.Acls;
+using HordeServer.Artifacts;
+using HordeServer.Projects;
+using HordeServer.Server;
+using HordeServer.Streams;
+using Microsoft.Extensions.DependencyInjection;
+using Microsoft.Extensions.Options;
+
+namespace HordeServer.Tests.Artifacts
+{
+	[TestClass]
+	public class ArtifactAclTests : BuildTestSetup
+	{
+		const string ClaimType = "artifact-test-claim";
+
+		static StreamId NormalStream { get; } = new StreamId("normal-stream");
+		static StreamId SecretStream { get; } = new StreamId("secret-stream");
+
+		static ArtifactType TestArtifactType { get; } = new ArtifactType("normal-build");
+
+		const string NormalReadRole = "normal-read";
+		const string SecretReadRole = "secret-read";
+
+		public ArtifactAclTests()
+		{
+			UpdateConfig(SetupFixture);
+		}
+
+		void SetupFixture(GlobalConfig globalConfig)
+		{
+			StreamConfig streamConfig = new StreamConfig { Id = SecretStream };
+			{
+				ArtifactTypeConfig artifactTypeConfig = new ArtifactTypeConfig();
+				artifactTypeConfig.Type = TestArtifactType;
+				artifactTypeConfig.Acl = new AclConfig();
+				artifactTypeConfig.Acl.Inherit = false;
+				artifactTypeConfig.Acl.Entries =
+					[
+						CreateAclEntry(SecretReadRole, ArtifactAclAction.ReadArtifact)
+					];
+				streamConfig.ArtifactTypes.Add(artifactTypeConfig);
+			}
+
+			ProjectConfig projectConfig = new ProjectConfig();
+			projectConfig.Id = new ProjectId("ue5");
+			projectConfig.Streams.Add(new StreamConfig { Id = NormalStream });
+			projectConfig.Streams.Add(streamConfig);
+			{
+				ArtifactTypeConfig artifactTypeConfig = new ArtifactTypeConfig();
+				artifactTypeConfig.Type = TestArtifactType;
+				artifactTypeConfig.Acl = new AclConfig();
+				artifactTypeConfig.Acl.Entries =
+					[
+						CreateAclEntry(NormalReadRole, ArtifactAclAction.ReadArtifact),
+						CreateAclEntry(SecretReadRole, ArtifactAclAction.ReadArtifact)
+					];
+				projectConfig.ArtifactTypes.Add(artifactTypeConfig);
+			}
+
+			BuildConfig buildConfig = globalConfig.Plugins.GetBuildConfig();
+			buildConfig.Projects.Add(projectConfig);
+		}
+
+		static AclEntryConfig CreateAclEntry(string role, AclAction action)
+		{
+			AclEntryConfig entry = new AclEntryConfig();
+			entry.Claim = new AclClaimConfig(ClaimType, role);
+			entry.Actions = [action];
+			return entry;
+		}
+
+		ClaimsPrincipal CreateUser(string role)
+		{
+			ClaimsIdentity identity = new ClaimsIdentity([new Claim(ClaimType, role)]);
+			return new ClaimsPrincipal(identity);
+		}
+
+		[TestMethod]
+		public void TestAcls()
+		{
+			BuildConfig buildConfig = ServiceProvider.GetRequiredService<IOptionsMonitor<BuildConfig>>().CurrentValue;
+
+			// Test that both roles can read from the default stream
+			{
+				ClaimsPrincipal secretUser = CreateUser(SecretReadRole);
+				Assert.IsTrue(buildConfig.AuthorizeArtifact(TestArtifactType, NormalStream, ArtifactAclAction.ReadArtifact, secretUser));
+
+				ClaimsPrincipal normalUser = CreateUser(NormalReadRole);
+				Assert.IsTrue(buildConfig.AuthorizeArtifact(TestArtifactType, NormalStream, ArtifactAclAction.ReadArtifact, normalUser));
+			}
+
+			// Test that only the secret role can read from the secret stream
+			{
+				ClaimsPrincipal secretUser = CreateUser(SecretReadRole);
+				Assert.IsTrue(buildConfig.AuthorizeArtifact(TestArtifactType, SecretStream, ArtifactAclAction.ReadArtifact, secretUser));
+
+				ClaimsPrincipal normalUser = CreateUser(NormalReadRole);
+				Assert.IsFalse(buildConfig.AuthorizeArtifact(TestArtifactType, SecretStream, ArtifactAclAction.ReadArtifact, normalUser));
+			}
+		}
+	}
+}
diff --git a/Engine/Source/Programs/Horde/Plugins/Build/HordeServer.Build/Artifacts/ArtifactTypeConfig.cs b/Engine/Source/Programs/Horde/Plugins/Build/HordeServer.Build/Artifacts/ArtifactTypeConfig.cs
--- a/Engine/Source/Programs/Horde/Plugins/Build/HordeServer.Build/Artifacts/ArtifactTypeConfig.cs
+++ b/Engine/Source/Programs/Horde/Plugins/Build/HordeServer.Build/Artifacts/ArtifactTypeConfig.cs
@@ -1,5 +1,6 @@
 // Copyright Epic Games, Inc. All Rights Reserved.
 
+using EpicGames.Horde.Acls;
 using EpicGames.Horde.Artifacts;
 using EpicGames.Horde.Storage;
 using HordeServer.Acls;
@@ -45,5 +46,13 @@ namespace HordeServer.Artifacts
 		/// Storage namespace to use for this artifact types
 		/// </summary>
 		public NamespaceId NamespaceId { get; set; } = new NamespaceId("horde-artifacts");
+
+		/// <summary>
+		/// Fixup an artifact type after reading the config
+		/// </summary>
+		public void PostLoad(AclConfig parentAcl)
+		{
+			Acl?.PostLoad(parentAcl, $"artifact:{Type}", AclConfig.GetActions([typeof(ArtifactTypeConfig)]));
+		}
 	}
 }
diff --git a/Engine/Source/Programs/Horde/Plugins/Build/HordeServer.Build/BuildConfig.cs b/Engine/Source/Programs/Horde/Plugins/Build/HordeServer.Build/BuildConfig.cs
--- a/Engine/Source/Programs/Horde/Plugins/Build/HordeServer.Build/BuildConfig.cs
+++ b/Engine/Source/Programs/Horde/Plugins/Build/HordeServer.Build/BuildConfig.cs
@@ -109,6 +109,11 @@ namespace HordeServer
 				}
 			}
 
+			foreach (ArtifactTypeConfig artifactTypeConfig in ArtifactTypes)
+			{
+				artifactTypeConfig.PostLoad(Acl);
+			}
+
 			UpdateWorkspacesForPools(configOptions.Plugins);
 		}
 
diff --git a/Engine/Source/Programs/Horde/Plugins/Build/HordeServer.Build/Projects/ProjectConfig.cs b/Engine/Source/Programs/Horde/Plugins/Build/HordeServer.Build/Projects/ProjectConfig.cs
--- a/Engine/Source/Programs/Horde/Plugins/Build/HordeServer.Build/Projects/ProjectConfig.cs
+++ b/Engine/Source/Programs/Horde/Plugins/Build/HordeServer.Build/Projects/ProjectConfig.cs
@@ -141,6 +141,7 @@ namespace HordeServer.Projects
 			}
 			foreach (ArtifactTypeConfig artifactTypeConfig in ArtifactTypes)
 			{
+				artifactTypeConfig.PostLoad(Acl);
 				_artifactTypeLookup[artifactTypeConfig.Type] = artifactTypeConfig;
 			}
 
diff --git a/Engine/Source/Programs/Horde/Plugins/Build/HordeServer.Build/Streams/StreamConfig.cs b/Engine/Source/Programs/Horde/Plugins/Build/HordeServer.Build/Streams/StreamConfig.cs
--- a/Engine/Source/Programs/Horde/Plugins/Build/HordeServer.Build/Streams/StreamConfig.cs
+++ b/Engine/Source/Programs/Horde/Plugins/Build/HordeServer.Build/Streams/StreamConfig.cs
@@ -354,6 +354,7 @@ namespace HordeServer.Streams
 			}
 			foreach (ArtifactTypeConfig artifactTypeConfig in ArtifactTypes)
 			{
+				artifactTypeConfig.PostLoad(Acl);
 				_artifactTypeLookup[artifactTypeConfig.Type] = artifactTypeConfig;
 			}
 
