Horde: Add documentation for the C# and HTTP REST APIs.

#rnx

diff --git a/Engine/Source/Programs/Horde/Docs/Internals.md b/Engine/Source/Programs/Horde/Docs/Internals.md
--- a/Engine/Source/Programs/Horde/Docs/Internals.md
+++ b/Engine/Source/Programs/Horde/Docs/Internals.md
@@ -6,6 +6,11 @@
 
 * [Building Horde](Internals/Building.md)
 
+## Interfacing with Horde
+
+* [C# API](Internals/CsApi.md)
+* [REST API](Internals/RestApi.md)
+
 ## Features
 
 * [Compute Architecture](Internals/ComputeArchitecture.md)
diff --git a/Engine/Source/Programs/Horde/Docs/Internals/CsApi.md b/Engine/Source/Programs/Horde/Docs/Internals/CsApi.md
new file mode 100644
--- /dev/null
+++ b/Engine/Source/Programs/Horde/Docs/Internals/CsApi.md
@@ -0,0 +1,93 @@
+[Horde](../../README.md) > [Internals](../Internals.md) > C# API
+
+# C# API
+
+Horde is implemented in C#, so the C# API is able to share a lot
+of implementation boilerplate with the server. That makes it an ideal
+choice for scripting, and makes the integration with other Unreal
+Engine tools such as AutomationTool, UnrealBuildTool,
+UnrealGameSync and Unreal Toolbox more convenient.
+
+## Overview
+
+There are two primary mechanisms for interacting with the server;
+through the [HTTP REST API](RestApi.md) (`HordeHttpClient`), which
+provides typed request and response classes, or via the native
+C# API (`IHordeClient`). The `IHordeClient` interface is the main
+entry point, but provides a method for getting a HTTP client if necessary (`IHordeClient.CreateHttpClient`).
+
+While the HTTP API is more complete, we intent to move all functionality
+into the native C# API over time. Doing so allows client tooling to write
+idiomatic C# code without having to worry about marshalling and transport
+concerns, and allows clients to write code against the same public APIs
+exposed by plugins on the server without having to concern themselves
+with implementation details.
+
+This approach gives more flexibility for prototyping code outside
+the server then moving it into the server as it makes sense, and
+gives flexibility if we need to split the server monolith into
+microservices in the future. Where possible, use functionality exposed
+through `IHordeClient` directly rather than using `HordeHttpClient`.
+
+## Configuring Horde
+
+There is a per-user setting for the preferred Horde server.
+
+On Windows, this is stored in the registry under the following keys:
+
+  ```
+  HKEY_CURRENT_USER\SOFTWARE\Epic Games\Horde\Url (preferred)
+  HKEY_LOCAL_MACHINE\SOFTWARE\Epic Games\Horde\Url
+  ```
+
+On MacOS and Linux, this is stored in the following JSON file:
+
+  ```
+  ~/.horde.json
+  ```
+
+The `HordeOptions.GetDefaultServerUrl()` method will retrieve the user's
+preferred server.
+
+## Connecting to Horde
+
+Horde uses .NET's standard dependency injection pattern.
+
+AutomationTool creates a service collection including Horde by default,
+allowing you to get a Horde client interface as follows:
+
+  ```cs
+  IHordeClient hordeClient = CommandUtils.ServiceProvider.GetRequiredService<IHordeClient>();
+  ```
+
+From other C# tools, it may be necessary to construct the required
+service container manually.
+
+  ```cs
+  ServiceCollection serviceCollection = new ServiceCollection();
+  serviceCollection.AddHorde();
+
+  await using ServiceProvider serviceProvider = serviceCollection.BuildServiceProvider();
+  IHordeClient hordeClient = 
+     serviceProvider.GetRequiredService<IHordeClient>();
+  ```
+
+The server URL can be overridden for a particular tool using the
+standard [.NET options pattern](https://learn.microsoft.com/en-us/aspnet/core/fundamentals/configuration/options?view=aspnetcore-8.0):
+
+  ```cs
+  serviceCollection.Configure<HordeOptions>(options => 
+  {
+    options.ServerUrl = new Uri("http://my-horde-server.com");
+  });
+  ```
+
+Or through the `AddHorde()` method:
+
+  ```cs
+  serviceCollection.AddHorde(options => 
+  {
+    options.ServerUrl = new Uri("http://my-horde-server.com");
+  });
+  ```
+  
\ No newline at end of file
diff --git a/Engine/Source/Programs/Horde/Docs/Internals/RestApi.md b/Engine/Source/Programs/Horde/Docs/Internals/RestApi.md
new file mode 100644
--- /dev/null
+++ b/Engine/Source/Programs/Horde/Docs/Internals/RestApi.md
@@ -0,0 +1,30 @@
+[Horde](../../README.md) > [Internals](../Internals.md) > REST API
+
+# REST API
+
+Horde's REST API is exposed through paths under `/api` on the
+server.
+
+Requests and responses support JSON and UE Compact Binary
+(by using the `application/x-ue-cb` MIME type), and responses support
+compression using Gzip, Zstd and Brotli.
+
+## Documentation
+
+A list of endpoints, and their parameters, is available
+on the server under:
+
+  ```
+  http://my-horde-server.com/swagger
+  ```
+
+This page also supports sending requests to the server directly.
+
+## Filtered responses
+
+Many Horde HTTP methods support a filter parameter for filtering
+responses. The syntax for filters includes simple wildcard expressions, such as:
+
+  ```
+  http://my-horde-server.com/api/v1/jobs?filter=id,batches.steps.id
+  ```
