Horde: Fix bug where agent fails to start if data dir is missing

This can happen during the automatic agent upgrade process. Likely caused by CL 37227886.

diff --git a/Engine/Source/Programs/Horde/HordeAgent/AgentApp.cs b/Engine/Source/Programs/Horde/HordeAgent/AgentApp.cs
--- a/Engine/Source/Programs/Horde/HordeAgent/AgentApp.cs
+++ b/Engine/Source/Programs/Horde/HordeAgent/AgentApp.cs
@@ -273,7 +273,7 @@ namespace HordeAgent
 				.AddJsonFile($"appsettings.{environment}.json", optional: true) // environment variable overrides, also used in k8s setups with Helm
 				.AddJsonFile("appsettings.User.json", optional: true);
 			
-			if (agentConfigFile != null)
+			if (agentConfigFile != null && Path.Exists(agentConfigFile.FullName))
 			{
 				// Adding a JSON file outside the common base path requires adding a completely separate JsonConfigurationSource entry
 				builder.Add(new JsonConfigurationSource()
