Adding robomerge track change messages to slack messages.

diff --git a/Engine/Source/Programs/Horde/Plugins/Build/HordeServer.Build/BuildResponseFilter.cs b/Engine/Source/Programs/Horde/Plugins/Build/HordeServer.Build/BuildResponseFilter.cs
--- a/Engine/Source/Programs/Horde/Plugins/Build/HordeServer.Build/BuildResponseFilter.cs
+++ b/Engine/Source/Programs/Horde/Plugins/Build/HordeServer.Build/BuildResponseFilter.cs
@@ -32,6 +32,11 @@ namespace HordeServer
 					dashboardConfigResponse.PerforceSwarmUrl = _staticBuildConfig.P4SwarmUrl.ToString().TrimEnd('/');
 				}
 
+				if (_staticBuildConfig.RobomergeUrl != null)
+				{
+					dashboardConfigResponse.RobomergeUrl = _staticBuildConfig.RobomergeUrl.ToString().TrimEnd('/');
+				}
+
 				dashboardConfigResponse.DeviceProblemCooldownMinutes = _staticBuildConfig.DeviceProblemCooldownMinutes;
 			}
 		}
diff --git a/Engine/Source/Programs/Horde/Plugins/Build/HordeServer.Build/BuildServerConfig.cs b/Engine/Source/Programs/Horde/Plugins/Build/HordeServer.Build/BuildServerConfig.cs
--- a/Engine/Source/Programs/Horde/Plugins/Build/HordeServer.Build/BuildServerConfig.cs
+++ b/Engine/Source/Programs/Horde/Plugins/Build/HordeServer.Build/BuildServerConfig.cs
@@ -40,6 +40,11 @@ namespace HordeServer
 		/// </summary>
 		public Uri? P4SwarmUrl { get; set; }
 
+		/// <summary>
+		/// Url of Robomergem installation
+		/// </summary>
+		public Uri? RobomergeUrl { get; set; }
+
 		/// <summary>
 		/// The Jira service account user name
 		/// </summary>
diff --git a/Engine/Source/Programs/Horde/Plugins/Build/HordeServer.Build/Notifications/Sinks/SlackNotificationSink.cs b/Engine/Source/Programs/Horde/Plugins/Build/HordeServer.Build/Notifications/Sinks/SlackNotificationSink.cs
--- a/Engine/Source/Programs/Horde/Plugins/Build/HordeServer.Build/Notifications/Sinks/SlackNotificationSink.cs
+++ b/Engine/Source/Programs/Horde/Plugins/Build/HordeServer.Build/Notifications/Sinks/SlackNotificationSink.cs
@@ -1,5 +1,6 @@
 // Copyright Epic Games, Inc. All Rights Reserved.
 
+using System;
 using System.Diagnostics;
 using System.Diagnostics.CodeAnalysis;
 using System.Globalization;
@@ -1245,14 +1246,33 @@ namespace HordeServer.Notifications.Sinks
 								suspectMessage += $" (<{swarmLink}|View changes>)";
 							}
 
+							Uri? robomergeTrackChangesLink = GetRobomergeTrackChangesUrlForSpan(span);
+							if (robomergeTrackChangesLink != null)
+							{
+								suspectMessage += $" (<{robomergeTrackChangesLink}|Track changes>)";
+							}
+
 							await _slackClient.PostMessageToThreadAsync(threadId, suspectMessage, cancellationToken: cancellationToken);
 						}
 						else
 						{
+							string suspectMessage = String.Empty;
+
 							Uri? swarmLink = GetSwarmLinkForSpan(span);
 							if (swarmLink != null)
 							{
-								await _slackClient.PostMessageToThreadAsync(threadId, $"<{swarmLink}|View changes in Swarm>", cancellationToken);
+								suspectMessage +=  $"<{swarmLink}|View changes in Swarm>";
+							}
+							
+							Uri? robomergeTrackChangesLink = GetRobomergeTrackChangesUrlForSpan(span);
+							if (robomergeTrackChangesLink != null)
+							{
+								suspectMessage += $" (<{robomergeTrackChangesLink}|Track changes>)";
+							}
+
+							if (!String.IsNullOrWhiteSpace(suspectMessage))
+							{
+								await _slackClient.PostMessageToThreadAsync(threadId, suspectMessage, cancellationToken);
 							}
 
 							notifyTriageAlias = true;
@@ -1296,6 +1316,13 @@ namespace HordeServer.Notifications.Sinks
 					{
 						string fixedEventId = $"issue_{issue.Id}_fixed_{issue.FixCommitId}";
 						string fixedMessage = $"Marked as fixed in {FormatChange(issue.FixCommitId)}";
+
+						Uri? robomergeTrackChangesLink = GetRobomergeTrackChangesUrlForSpan(span);
+						if (robomergeTrackChangesLink != null)
+						{
+							fixedMessage += $" (<{robomergeTrackChangesLink}|Track changes>)";
+						}
+
 						await PostSingleMessageToThreadAsync(triageChannel, fixedEventId, threadId, fixedMessage, cancellationToken);
 					}
 					else
@@ -1439,6 +1466,18 @@ namespace HordeServer.Notifications.Sinks
 			}
 		}
 
+		Uri? GetRobomergeTrackChangesUrlForSpan(IIssueSpan span)
+		{
+			if (_settings.RobomergeUrl == null || span.LastSuccess == null)
+			{
+				return null;
+			}
+			else
+			{
+				return new Uri(_settings.RobomergeUrl, $"trackchange/{span.FirstFailure.CommitId}");
+			}
+		}
+
 		async Task PostSingleMessageToThreadAsync(string recipient, string eventId, SlackMessageId threadId, string message, CancellationToken cancellationToken)
 		{
 			(MessageStateDocument state, bool isNew) = await AddOrUpdateMessageStateAsync(recipient, eventId, null, "", null, cancellationToken);
@@ -1542,6 +1581,12 @@ namespace HordeServer.Notifications.Sinks
 				{
 					Uri fixFailedUrl = new Uri(_serverInfo.DashboardUrl, $"job/{fixFailedStep.JobId}?step={fixFailedStep.StepId}&issue={issue.Id}");
 					text = $":cross: Marked fixed in *CL {issue.FixCommitId}*, but seen again at *<{fixFailedUrl}|CL {fixFailedStep.CommitId}>*";
+
+					string rbTrackChangeText = FormatRobomergeTrackChangeLinkChange(issue.FixCommitId);
+					if (!String.IsNullOrWhiteSpace(rbTrackChangeText))
+					{
+						text += rbTrackChangeText;
+					}
 				}
 				else
 				{
@@ -1691,14 +1736,23 @@ namespace HordeServer.Notifications.Sinks
 
 		string FormatChange(CommitId change)
 		{
-			if (_settings.P4SwarmUrl != null)
-			{
-				return $"<{new Uri(_settings.P4SwarmUrl, $"changes/{change}")}|CL {change}>";
-			}
-			else
+			string output = _settings.P4SwarmUrl != null
+				? $"<{new Uri(_settings.P4SwarmUrl, $"changes/{change}")}|CL {change}>"
+				: $"CL {change}";
+
+			if (_settings.RobomergeUrl != null)
 			{
-				return $"CL {change}";
+				output += FormatRobomergeTrackChangeLinkChange(change);
 			}
+
+			return output;
+		}
+
+		string FormatRobomergeTrackChangeLinkChange(CommitId change)
+		{
+			return _settings.RobomergeUrl != null
+				? $" (<{new Uri(_settings.RobomergeUrl, $"trackchange/{change}")}|Track changes>)"
+				: String.Empty;
 		}
 
 		string FormatJobStep(IIssueStep step, string name) => $"<{GetJobUrl(step.JobId)}|{step.JobName}> / <{GetStepUrl(step.JobId, step.StepId)}|{name}>";
diff --git a/Engine/Source/Programs/Shared/EpicGames.Horde/Dashboard/DashboardMessages.cs b/Engine/Source/Programs/Shared/EpicGames.Horde/Dashboard/DashboardMessages.cs
--- a/Engine/Source/Programs/Shared/EpicGames.Horde/Dashboard/DashboardMessages.cs
+++ b/Engine/Source/Programs/Shared/EpicGames.Horde/Dashboard/DashboardMessages.cs
@@ -30,6 +30,11 @@ namespace EpicGames.Horde.Dashboard
 		/// </summary>
 		public string? PerforceSwarmUrl { get; set; }
 
+		/// <summary>
+		/// Url of Robomergem installation
+		/// </summary>
+		public string? RobomergeUrl { get; set; }
+
 		/// <summary>
 		/// Help email address that users can contact with issues
 		/// </summary>
