Horde: Add initial documentation for compute plugin

Also describe pools under agents page.

#skipci

diff --git a/Engine/Source/Programs/Horde/Docs/Config.md b/Engine/Source/Programs/Horde/Docs/Config.md
--- a/Engine/Source/Programs/Horde/Docs/Config.md
+++ b/Engine/Source/Programs/Horde/Docs/Config.md
@@ -24,7 +24,7 @@ support integrated into AutomationTool, UnrealBuildTool and Unreal Editor.
 ### Build
 
 Mature build automation system similar to Jenkins or TeamCity, designed
-and streamlined for Unreal Engine projects and best practices adoped by Epic.
+and streamlined for Unreal Engine projects and best practices adopted by Epic.
 
 * [Build Automation](Config/BuildAutomation.md): Introduction to Horde's Build Automation and CI/CD system.
 * [Artifacts](Config/Artifacts.md): Managing and distributing artifacts produced by build steps.
diff --git a/Engine/Source/Programs/Horde/Docs/Config/Agents.md b/Engine/Source/Programs/Horde/Docs/Config/Agents.md
--- a/Engine/Source/Programs/Horde/Docs/Config/Agents.md
+++ b/Engine/Source/Programs/Horde/Docs/Config/Agents.md
@@ -9,8 +9,9 @@ For information about deploying new agents, see [Horde > Deployment > Agent](../
 ## Pools
 
 Pools are groups of machines that can be used interchangeably, typically due to being a particular platform or
-hardware class. Pools simplify the management of build pipelines by allowing DevOps engineers to configure a mapping
-from agent types to physical machines.
+hardware class. At Epic, certain pools are dedicated to (incremental) compilation, cooking and Unreal Build Accelerator (UBA). 
+Pools simplify the management of build pipelines by allowing infrastructure engineers to configure a mapping
+from agent types to physical machines. An agent can belong to multiple pools.
 
 Pools are defined in the [.globals.json](Schema/Globals.md) file, via the `pools` property. Agents may be added to a
 pool manually through the Horde Dashboard or automatically by matching a particular condition. For example,
@@ -23,6 +24,24 @@ the following configuration block defines a pool that automatically includes all
 
 See also: [Condition expression syntax](Conditions.md)
 
+### Auto-scaling
+
+Pools in Horde serve as the primary configuration point for auto-scaling,
+a feature particularly valuable in cloud environments where hardware resources can be rented on an hourly or per-second basis.
+This functionality allows for resource allocation based on factors such as time of week or the current phase of your game project,
+optimizing cost-efficiency. For on-premise setups where hardware is statically allocated or owned, this feature is less useful.
+
+Each pool can be configured to automatically adjust its size based on specific metrics.
+The pool size strategy defines the ideal number of agents, with two primary approaches: `JobQueue` and `LeaseUtilization`.
+`JobQueue`, the preferred method, proactively examines the queue of pending jobs waiting to run.
+In contrast, `LeaseUtilization` reactively adjusts based on average CPU usage across all agents in the pool, which may result in a slight lag.
+
+Fleet managers handle cloud-specific implementations for hardware allocation,
+such as the `AwsRecycle` manager that controls AWS EC2 instances.
+
+To fine-tune scaling behavior, time-based cooldowns for scale-in and scale-out operations can be set,
+with scale-out typically configured more aggressively to accommodate rapid bursts of activity without excessive wait times.
+
 ## Remoting to Agents
 
 If you have a fleet of machines that require identical login credentials, you can configure UnrealGameSync to open
diff --git a/Engine/Source/Programs/Horde/Docs/Config/Compute.md b/Engine/Source/Programs/Horde/Docs/Config/Compute.md
new file mode 100644
--- /dev/null
+++ b/Engine/Source/Programs/Horde/Docs/Config/Compute.md
@@ -0,0 +1,107 @@
+[Horde](../../README.md) > [Configuration](../Config.md) > Compute
+
+# Compute
+
+Compute is a core plugin in Horde that manages compute resources in the form of agents.
+These agents are essential for two primary functions:
+powering build automation jobs and running remote execution tasks.
+
+To structure management of these agents, Horde allows the creation of pools, which can be configured to include different sets of agents.
+In addition to this, compute clusters serve a similar purpose but are specifically designed for handling remote compute tasks.
+
+More details about pools can be found on the [Agents](Agents.md) page.
+
+## Configuring
+The two most important points of configuration are pools and compute clusters. 
+
+### Pools
+
+Below is a pool configured for AWS using the job queue for appropriate sizing.
+`ScaleOutFactor`/`ScaleInFactor` are parameters for the `JobQueueStrategy` and
+describe how aggressive scaling should be (at which rate/how many EC2 instances to add or remove).
+```json
+{
+  "plugins": {
+    "compute": {
+      "pools": [
+        {
+          "id": "win-ue5",
+          "name": "Win-UE5", // Human-readable name
+          "enableAutoscaling": true,
+          "conformInterval": "1.00:00:00", // Automatically conform every 24 hours
+          "sizeStrategy": "JobQueue",
+          "sizeStrategies": [
+            {
+              "type": "JobQueue",
+              "condition": "true",
+              "config": "{\"ScaleOutFactor\": 1, \"ScaleInFactor\": 0.1}",
+              "extraAgentCount": 1
+            }
+          ],
+          "fleetManagers": [{ "type": "AwsRecycle", "config": "{}" }]
+        }
+      ]
+    }
+  }
+}
+```
+For a complete of fields a pool can be configured with, see [PoolConfig](Schema/Globals.md#poolconfig) schema.
+
+### Compute Clusters
+
+A compute cluster in Horde is a logical grouping of agents designed to help organize remote execution,
+manage permissions, and storage of compute tasks.
+
+Compute tasks are the core units of work within a cluster. Each task describes an OS process to be executed remotely on an agent.
+These descriptions include specific resource requirements such as CPU or RAM allocation, files that need to be transferred,
+the executable entrypoint, arguments, and any necessary environment variables.
+Typically, compute tasks are designed to be short-lived, with most completing within a 1-15 minute timeframe.
+
+The Unreal Build Tool together with Unreal Build Accelerator leverages compute tasks in Horde to distribute CPU-intensive operations.
+This includes C++ and shader compilation allowing for more efficient use of available resources beyond what's available on single a local machine.
+This distributed computing approach is one of the cornerstones for accelerating build processes,
+whether initiated through Horde jobs or from individual developer workstations.
+By leveraging the power of the cluster, teams can significantly reduce build times and improve overall development efficiency.
+
+When setting up a compute cluster, it's important to note that permissions are not granted by default.
+Each user or service that requires access to the cluster must be explicitly granted permissions.
+This ensures a secure environment and allows for fine-grained control over who can utilize the cluster's resources.
+
+An example compute cluster may be configured as follows in the global config:
+
+```json
+{
+  "plugins": {
+    "compute": {
+      "clusters": [
+        {
+          // Arbitrary name for referencing this cluster
+          "id": "uba",
+          
+          // Use condition to select which agents to include
+          "condition": "pool == 'UBA-HQ' && osfamily == 'Linux'",
+
+          // Each user or service that needs access must be explicitly granted permission
+          // In almost all cases, this means the "AddComputeTasks" action
+          "acl": {
+            "entries": [
+              { "claim": {"type": "http://epicgames.com/ue/horde/user", "value": "jane.doe"}, "actions": ["AddComputeTasks"] },
+              { "claim": {"type": "http://schemas.microsoft.com/ws/2008/06/identity/claims/role", "value": "programmers"}, "actions": ["AddComputeTasks"] }
+            ]
+          }
+        }
+      ]
+    }
+  }
+}
+```
+
+## Running remote compute tasks
+
+For those looking to understand the basic functionality of compute tasks,
+the Horde CLI tool provides a simple example through its `compute run` command.
+This command demonstrates how to initiate and manage a basic compute task,
+serving as a starting point for developers new to Horde.
+
+Unreal Build Tool implements Horde's compute client in C# for allocating tasks for UBA.
+See the [remote C++ tutorial](../Tutorials/RemoteCompilation.md) for getting started.
