Horde: Add docs for configuring tools.

#rnx

diff --git a/Engine/Source/Programs/Horde/Docs/Config.md b/Engine/Source/Programs/Horde/Docs/Config.md
--- a/Engine/Source/Programs/Horde/Docs/Config.md
+++ b/Engine/Source/Programs/Horde/Docs/Config.md
@@ -2,61 +2,92 @@
 
 # Configuration
 
-This section targets operating and managing Horde installations and shares some best practices learned from supporting it for teams at Epic.
+This section targets operating and managing Horde installations
+and shares some best practices learned from supporting it for
+teams at Epic.
 
 ## General
 
-* [Orientation](Config/Orientation.md): Get familiar with how Horde is configured, how to store configuration data in
+* [Orientation](Config/Orientation.md): Get familiar with how
+  Horde is configured, how to store configuration data in
   revision control, and how to set up a schema server.
-* [Permissions](Config/Permissions.md): Understand how to use the Horde permissions model.
+* [Permissions](Config/Permissions.md): Understand how to use the
+  Horde permissions model.
 
 ## Plugins
 
-Horde is implemented as a set of plugins over a host framework. For information about enabling or disabling specific plugins, see [Configuration > Plugins](Config/Plugins.md).
+Horde is implemented as a set of plugins on a host server. For
+information about enabling or disabling specific plugins, see
+[Configuration > Plugins](Config/Plugins.md).
 
 ### Analytics
 
-Experimental solution for studio-wide analytics and metrics gathering, with
-support integrated into AutomationTool, UnrealBuildTool and Unreal Editor.
+Experimental solution for studio-wide analytics and metrics
+gathering, with support integrated into AutomationTool,
+UnrealBuildTool and Unreal Editor.
 
-* [Analytics](Config/Analytics.md): Setting up a telemetry sink and creating dashboards showing KPIs.
+* [Analytics](Config/Analytics.md): Setting up a telemetry sink
+  and creating dashboards showing KPIs.
 
 ### Build
 
-Mature build automation system similar to Jenkins or TeamCity, designed
-and streamlined for Unreal Engine projects and best practices adopted by Epic.
-
-* [Build Automation](Config/BuildAutomation.md): Introduction to Horde's Build Automation and CI/CD system.
-* [Artifacts](Config/Artifacts.md): Managing and distributing artifacts produced by build steps.
-* [Devices](Config/Devices.md): Adding mobile devices and console development kits as shared resources.
-* [Automation Hub](Config/AutomationHub.md): Dashboard for surfacing test jamdata and trends across multiple projects and streams.
-* [UGS Metadata Server](Config/UgsMetadataServer.md): Surfaces team metadata to users of UnrealGameSync.
+Mature build automation system similar to Jenkins or TeamCity,
+designed and streamlined for Unreal Engine projects and best
+practices adopted by Epic.
+
+* [Build Automation](Config/BuildAutomation.md): Introduction to
+  Horde's Build Automation and CI/CD system.
+* [Artifacts](Config/Artifacts.md): Managing and distributing
+  artifacts produced by build steps.
+* [Devices](Config/Devices.md): Adding mobile devices and console
+  development kits as shared resources.
+* [Automation Hub](Config/AutomationHub.md): Dashboard for
+  surfacing test jamdata and trends across multiple projects and
+  streams.
+* [UGS Metadata Server](Config/UgsMetadataServer.md): Surfaces
+  team metadata to users of UnrealGameSync.
 
 ### Compute
 
-Manages worker machines that can be leased out to perform workloads, including build automation and remote execution use cases.
+Manages worker machines that can be leased out to perform
+workloads, including build automation and remote execution use
+cases.
 
-* [Agents](Config/Agents.md): Configure worker machines to connect to the Horde server for CI and remote execution workspaces.
+* [Agents](Config/Agents.md): Configure worker machines to connect
+  to the Horde server for CI and remote execution workspaces.
 
 ### Secrets
 
-Managed access to sensitive data integrated with Horde's permissions model.Secrets may be stored in Horde itself, or in an external secret store.
+Managed access to sensitive data integrated with Horde's
+permissions model.Secrets may be stored in Horde itself, or in an
+external secret store.
 
-* [Secrets](Config/Secrets.md): Configuring secrets and external secret providers.
+* [Secrets](Config/Secrets.md): Configuring secrets and external
+  secret providers.
 
 ### Storage
 
-Flexible, low overhead storage abstraction suitable for many different use
-cases. Used by Horde internally, but it can also be used directly by client applications.
+Flexible, low overhead storage abstraction suitable for many
+different use cases. Used by Horde internally, but it can also be
+used directly by client applications.
 
-* [Storage](Config/Storage.md): Setting up and managing the storage system.  
+* [Storage](Config/Storage.md): Setting up and managing the
+  storage system.  
 
 ### Symbols
 
-Implements a Windows symbol store indexing artifacts uploaded to Horde's
-storage system.
+Implements a Windows symbol store indexing artifacts uploaded to
+Horde's storage system.
+
+* [Symbols](Config/Symbols.md): Adding symbol stores and
+  configuring clients to use them.
+
+### Tools
+
+Hosts tools and utilities that can be downloaded by users using
+UnrealGameSync or Unreal Toolbox.
 
-* [Symbols](Config/Symbols.md): Adding symbol stores and configuring clients to use them.
+* [Tools](Config/Tools.md): Configure tools hosted by Horde.
 
 ## Reference
 
diff --git a/Engine/Source/Programs/Horde/Docs/Config/Tools.md b/Engine/Source/Programs/Horde/Docs/Config/Tools.md
new file mode 100644
--- /dev/null
+++ b/Engine/Source/Programs/Horde/Docs/Config/Tools.md
@@ -0,0 +1,59 @@
+[Horde](../../README.md) > [Configuration](../Config.md) > Tools
+
+# Tools
+
+Unreal Engine includes a number of tools to streamline certain
+workflows. Horde implements functionality for hosting tools for
+distribution to your team, and allows versioning and incrementally
+rolling out new builds to users.
+
+## Unreal Toolbox
+
+**Unreal Toolbox** is an application which can manage tools
+downloaded from Horde. It provides a simple interface for
+downloading, installing, configuring, and upgrading Unreal Engine
+tools, as well as managing access tokens for connecting to the Horde
+Server that can be used by other tools (such as UnrealBuildTool or
+the Unreal Editor).
+
+By default, a Windows installer for Unreal Toolbox is distributed
+with the Horde Server, which can be downloaded from the Tools > Download
+page on the dashboard.
+
+## Adding tools
+
+Tools are configured through the [Tools](Schema/Globals.md#toolsconfig)
+list in the [global.json](Schema/Globals.md) file. Licensees are encouraged
+to extend the list of tools hosted by Horde with their own tools, and to
+use Unreal Toolbox for distributing them.
+
+New tool deployments can be created through the `DeployTool` BuildGraph
+task:
+
+  ```xml
+  <DeployTool Id="horde-agent" Version="1.0" Directory="$(RootDir)/ToolDirectory"/>
+  ```
+
+Placing a Toolbox.json file in the root directory of a tool deployment
+can configure commands to run on install and uninstall of the tool by
+Unreal Toolbox, as well as adding additional commands to be made available
+through the Toolbox popup menu. See
+`Engine/Source/Programs/UnrealToolbox/ToolConfig.cs` for the syntax of
+this file.
+
+## Staggered Deployments
+
+Each tool can have a number of deployments, and the Tools > Dashboard page
+on the Horde dashboard can be used to start, pause, and revert revisions.
+The system is designed to support staggered deployments; clients
+choose a pseudo-random 'phase' value from 0-1 indicating when during a
+rollout they should receive the deployment, and the server returns the
+appropriate build information based on the start time and configured duration.
+
+## Bundled tools
+
+Tools bundled with the Horde installer by default are configured
+separately, via the
+[BundledTools](../Deployment/ServerSettings.md#toolsserverconfig) section
+of the server's config file. Data for these tools is installed alongside
+the server.
