Horde: Fix spelling mistakes and linter warnings in docs.

#rnx

diff --git a/Engine/Source/Programs/Horde/Docs/Config/Artifacts.md b/Engine/Source/Programs/Horde/Docs/Config/Artifacts.md
--- a/Engine/Source/Programs/Horde/Docs/Config/Artifacts.md
+++ b/Engine/Source/Programs/Horde/Docs/Config/Artifacts.md
@@ -2,7 +2,11 @@
 
 # Artifacts
 
-An artifact in Horde refers to a set of archived files produced by a build process. The archive corresponds to a particular commit in VCS, allowing for searching for artifacts produced between in a particular stream across a range of time, and includes other indexed keys and metadata that can help it to be manipulated by external tools.
+An artifact in Horde refers to a set of archived files produced by a build
+process. The archive corresponds to a particular commit in VCS, allowing
+for searching for artifacts produced between in a particular stream across
+a range of time, and includes other indexed keys and metadata that can
+help it to be manipulated by external tools.
 
 Amongst other things, artifacts can be used to store:
 
@@ -14,24 +18,44 @@ Amongst other things, artifacts can be used to store:
 
 ## Configuration
 
-Settings for artifacts are configured by their **artifact type**, which are specified at upload time. Artifact types are arbitrary string identifiers (case insensitive alpha-numeric strings, or the characters `-`, `_`, `.`).
+Settings for artifacts are configured by their **artifact type**, which
+are specified at upload time. Artifact types are arbitrary string
+identifiers (case insensitive alpha-numeric strings, or the characters
+`-`, `_`, `.`).
 
-Artifact types are customizable, and new types may be added as needed. They must be declared through Horde config files before Horde will allow clients to upload artifacts of that type.
+Artifact types are customizable, and new types may be added as needed.
+They must be declared through Horde config files before Horde will allow
+clients to upload artifacts of that type.
 
-Types may be configured through the `artifactTypes` property at the [global](Schema/Globals.md#buildconfig), [project](Schema/Projects.md), or
-[stream](Schema/Streams.md) scopes.
+Types may be configured through the `artifactTypes` property at the
+[global](Schema/Globals.md#buildconfig), [project](Schema/Projects.md),
+or [stream](Schema/Streams.md) scopes.
 
 ### Standard Artifact Types
 
 Horde produces certain artifact types by default:
 
-- `step-output`: Output from a build step which can be used by other build steps. These artifacts are produced based on dependencies between steps configured through BuildGraph.
+- `step-output`: Output from a build step which can be used by other build
+steps. These artifacts are produced based on dependencies between steps
+configured through BuildGraph.
 
-- `step-saved`: Contents of the AutomationTool log folder, harvested from the agent executing a job once it completes. The Horde Agent sets the path to this output folder for jobs via the `uebp_LogFolder` environment variable (normally `Engine/Programs/AutomationTool/Saved/Logs`).
+- `step-saved`: Contents of the AutomationTool log folder, harvested from
+the agent executing a job once it completes. The Horde Agent sets the path
+to this output folder for jobs via the `uebp_LogFolder` environment
+variable (normally `Engine/Programs/AutomationTool/Saved/Logs`).
 
-- `step-trace`: Telemetry and trace data for a build step, harvested from the directory given by the `UE_TELEMETRY_DIR` environment variable (normally `Engine/Programs/AutomationTool/Saved/Telemetry`). AutomationTool and UnrealBuildTool write timing information for their execution to this folder using the in JSON format, which are replicated to the configured telemetry provider.
+- `step-trace`: Telemetry and trace data for a build step, harvested from
+the directory given by the `UE_TELEMETRY_DIR` environment variable
+(normally `Engine/Programs/AutomationTool/Saved/Telemetry`).
+AutomationTool and UnrealBuildTool write timing information for their
+execution to this folder using the in JSON format, which are replicated to
+the configured telemetry provider.
 
-- `step-testdata`: Data written to the directory given by the `UE_TESTDATA_DIR` environment variable by a build step (normally `Engine/Programs/AutomationTool/Saved/TestData`). Json files in this directory will also be harvested for test data results and surfaced to Automation Hub.
+- `step-testdata`: Data written to the directory given by the
+`UE_TESTDATA_DIR` environment variable by a build step (normally
+`Engine/Programs/AutomationTool/Saved/TestData`). Json files in this
+directory will also be harvested for test data results and surfaced to
+Automation Hub.
 
 ### Expiration
 
@@ -42,17 +66,22 @@ Artifacts may be expired automatically using the following settings:
 
 If both settings are specified, the artifact will be retained until neither retention policy applies.
 
-If the configuration for an artifact is removed from Horde - because a stream is deleted, for example - any orphanned artifacts will be expired according to the last configured settings.
+If the configuration for an artifact is removed from Horde - because a
+stream is deleted, for example - any orphanned artifacts will be expired
+according to the last configured settings.
 
 ## Uploading Artifacts
 
-Artifacts can be produced by build processes using the `CreateArtifact` BuildGraph task within a node, or using the `Artifact` element to indicate
+Artifacts can be produced by build processes using the `CreateArtifact`
+BuildGraph task within a node, or using the `Artifact` element to indicate
 that a produced tag should be published as an artifact.
 
 Artifacts produced by build jobs will be listed at the top of the dashboard page
 for that job.
 
-An example of using the `CreateArtifact` task can be found in the `Engine/Build/InstalledEngineBuild.xml` script, which uploads the complete installed build as an `installed-build` artifact type:
+An example of using the `CreateArtifact` task can be found in the
+`Engine/Build/InstalledEngineBuild.xml` script, which uploads the complete
+installed build as an `installed-build` artifact type:
 
   ```xml
   <Node Name="Upload To Horde Win64" Requires="Make Installed Build Win64">
@@ -65,7 +94,11 @@ An example of using the `CreateArtifact` task can be found in the `Engine/Build/
   </Node>
   ```
 
-An example of using the `Artifact` element can be found in the `Engine/Build/Graph/Examples/BuildEditorAndTools.xml` script, which uploads a set of files as a `ugs-pcb` artifact type. UnrealGameSync searches for this artifact type to determine which archives are available, along with their associated project and archive type.
+An example of using the `Artifact` element can be found in the
+`Engine/Build/Graph/Examples/BuildEditorAndTools.xml` script, which
+uploads a set of files as a `ugs-pcb` artifact type. UnrealGameSync
+searches for this artifact type to determine which archives are available,
+along with their associated project and archive type.
 
   ```xml
   <Artifact
@@ -82,9 +115,13 @@ Artifacts may also be uploaded programmatically using the `EpicGames.Horde` libr
 
 ## Downloading Artifacts
 
-The easiest way to download an artifact is via the Horde dashboard. While convenient, downloading a zip file directly from the dashboard can be slow and inefficient due to the work the server has to do to re-pack data in that format.
+The easiest way to download an artifact is via the Horde dashboard. While
+convenient, downloading a zip file directly from the dashboard can be slow
+and inefficient due to the work the server has to do to re-pack data in
+that format.
 
-To download artifacts in their native format - which is typically much faster and more efficient - the following methods are available:
+To download artifacts in their native format - which is typically much
+faster and more efficient - the following methods are available:
 
 - **UnrealGameSync**: The Horde dashboard allows downloading an artifact as a
   `.uartifact` file; a JSON file containing metadata about an artifact.
@@ -93,12 +130,13 @@ To download artifacts in their native format - which is typically much faster an
 
 - **Command line**: The Horde command line tool allows downloading artifacts to an output folder using the following syntax:
 
-    ```
+    ```bat
     horde artifact download -Id=12345678abcd12345678abcd -OutputDir=C:\Folder
     ```
+
   Artifacts can be enumerated using the `artifact find` command:
 
-    ```
+    ```bat
     horde artifact find -streamid=ue5-main -type=step-saved -jobid=670fce208961cf36f3855ff0 -stepid=4994
     ```
 
diff --git a/Engine/Source/Programs/Horde/Docs/Config/AutomationHub.md b/Engine/Source/Programs/Horde/Docs/Config/AutomationHub.md
--- a/Engine/Source/Programs/Horde/Docs/Config/AutomationHub.md
+++ b/Engine/Source/Programs/Horde/Docs/Config/AutomationHub.md
@@ -2,14 +2,26 @@
 
 # Automation Hub
 
-The Horde Automation Hub surfaces individual and suite [Gauntlet](https://docs.unrealengine.com/5.3/en-US/gauntlet-automation-framework-overview-in-unreal-engine) test results.  Horde efficiently generates searchable metadata for streams, platforms, configurations, rendering apis, etc.  The automation hub is used at Epic by QA, release managers, and code owners to quickly view and investigate the latest test results across platforms and streams.  It provides historical data and views which drill into specific test events which can include screenshots, logging, and callstacks.   
-
-In addition to enabling Horde [build automation](BuildAutomation.md), the configuration required to surface test results is simply adding the  
-`-WriteTestResultsForHorde` Gauntlet test command line argument.  Please see the BuildGraph [example](#buildgraph-example) below for details. 
+The Horde Automation Hub surfaces individual and suite [Gauntlet](https://docs.unrealengine.com/5.3/en-US/gauntlet-automation-framework-overview-in-unreal-engine)
+test results.  Horde efficiently generates searchable metadata for
+streams, platforms, configurations, rendering apis, etc.  The automation
+hub is used at Epic by QA, release managers, and code owners to quickly
+view and investigate the latest test results across platforms and
+streams.  It provides historical data and views which drill into specific
+test events which can include screenshots, logging, and callstacks.
+
+In addition to enabling Horde [build automation](BuildAutomation.md), the
+configuration required to surface test results is simply adding the  
+`-WriteTestResultsForHorde` Gauntlet test command line argument.  Please
+see the BuildGraph [example](#buildgraph-example) below for details.
 
 ## Automation Filters
 
-The automation UI is data driven and automatically populated from test metadata.  Test results can be filtered by project, tests, streams, platforms, configurations, targets, rendering hardware interface, and variations.  Detailed selections may also be linked for sharing or to be bookmarked.  
+The automation UI is data driven and automatically populated from test
+metadata.  Test results can be filtered by project, tests, streams,
+platforms, configurations, targets, rendering hardware interface, and
+variations.  Detailed selections may also be linked for sharing or to be
+bookmarked.  
 
 ![Automation Selection](../Images/AutomationHub-LeftPanel.png)
 
@@ -19,7 +31,9 @@ Test results are presented as tiles which surface relative test health based on
 
 ![Test Results](../Images/AutomationHub-TestResults.png)
 
-Test tiles can be expanded to view further details such as platforms and changelists, which are linked to individual [Horde CI](BuildAutomation.md) job steps to assist in issue investigation.  
+Test tiles can be expanded to view further details such as platforms and
+changelists, which are linked to individual [Horde CI](BuildAutomation.md)
+job steps to assist in issue investigation.  
 
 ![Test Card](../Images/AutomationHub-TestPanel.png)
 
@@ -29,11 +43,16 @@ Test history graphs and detailed failure reports are also available which includ
 
 ## Test Suites
 
-Gauntlet suite tests may contains thousands of individual unit tests.  The automation hub can drill into historical data for each unit test with cross stream comparison.  
+Gauntlet suite tests may contains thousands of individual unit tests.  The
+automation hub can drill into historical data for each unit test with
+cross stream comparison.  
 
 ![Suite Results](../Images/AutomationHub-SuiteResults.png)
 
-Suite tests generate test events which can  be viewed to help diagnose issues.  Test events can include additional data such as logging and screenshots to catch regressions.  Alternate platforms of the unit test may also be selected for comparison purposes.  
+Suite tests generate test events which can  be viewed to help diagnose
+issues.  Test events can include additional data such as logging and
+screenshots to catch regressions.  Alternate platforms of the unit test
+may also be selected for comparison purposes.  
 
 ![Screenshot Comparison](../Images/AutomationHub-ScreenShotCompare.png)
 
@@ -42,14 +61,16 @@ Suite tests generate test events which can  be viewed to help diagnose issues.
 The following [**BuildGraph**](https://docs.unrealengine.com/5.0/en-US/buildgraph-for-unreal-engine/) fragment declares:
 
 * `HordeDeviceService` and `HordeDevicePool` properties that specify your Horde server and which device pool to use.
-* Adds a `BootTest Android` node which specifies `-WriteTestResultsForHorde` andh will automatically generate test data to be ingested by Horde, parsed to efficient meta data, and surfaces by the automation hub
-
----
-
-	<Property Name="HordeDeviceService" Value="http://localhost:13440" />
-	<Property Name="HordeDevicePool" Value="UE5" />
-		
-	<Node Name="BootTest Android">
-		<Command Name="RunUnreal" Arguments="-test=UE.BootTest -platform=Android " -deviceurl=&quot;$(HordeDeviceService)&quot; -devicepool=&quot;$(HordeDevicePool)&quot; -WriteTestResultsForHorde/>
-	</Node>
----
+* Adds a `BootTest Android` node which specifies
+`-WriteTestResultsForHorde` and will automatically generate test data to
+be ingested by Horde, parsed to efficient meta data, and surfaces by the
+automation hub
+
+```xml
+    <Property Name="HordeDeviceService" Value="http://localhost:13440" />
+    <Property Name="HordeDevicePool" Value="UE5" />
+
+    <Node Name="BootTest Android">
+        <Command Name="RunUnreal" Arguments="-test=UE.BootTest -platform=Android -deviceurl=&quot;$(HordeDeviceService)&quot; -devicepool=&quot;$(HordeDevicePool)&quot; -WriteTestResultsForHorde"/>
+    </Node>
+```
diff --git a/Engine/Source/Programs/Horde/Docs/Config/BuildHealth.md b/Engine/Source/Programs/Horde/Docs/Config/BuildHealth.md
--- a/Engine/Source/Programs/Horde/Docs/Config/BuildHealth.md
+++ b/Engine/Source/Programs/Horde/Docs/Config/BuildHealth.md
@@ -2,7 +2,7 @@
 
 # Build Health
 
-Epic has a large development team working on products that are continously shipping. We
+Epic has a large development team working on products that are continuously shipping. We
 use a trunk-based development model with branches for each release, with
 commits to release branches automatically being merged back to the mainline.
 
@@ -139,7 +139,7 @@ Post-processing is handled by the `IssueHandler` derived classes in the `EpicGam
 Handlers are enumerated by the server by searching for classes with the `[IssueHandler]` attribute. A new
 instance of each handler is created for each build step containing events to be processed.
 
-Log events are passed to `IssueHandler.HandleEvent` in decreasing order of their `IssueHandler.Priorty` value,
+Log events are passed to `IssueHandler.HandleEvent` in decreasing order of their `IssueHandler.Priority` value,
 until a handler is found which returns `true`. Once all events have been processed, `IssueHandler.GetIssues`
 returns a list of the matched log events and their fingerprints.
 
@@ -205,7 +205,7 @@ The supported set of annotations are defined in the `NodeAnnotations.cs` source
 * `AutoAssign`: Whether to automatically assign issues that could only be caused by one user, or have a well defined
   correlation with a modified file.
 * `AutoAssignToUser`: Automatically assign any issues in this step to the Perforce user given by the supplied argument.
-* `NotifySubmitters`: Whether to notify all submitters between a build suceeding and failing, allowing them to step
+* `NotifySubmitters`: Whether to notify all submitters between a build succeeding and failing, allowing them to step
   forward and take ownership of an issue.
 * `IssueGroup`: Specifies a suffix to be appended to the `type` property of any issues created from this node,
   preventing them being merged with other issues.
diff --git a/Engine/Source/Programs/Horde/Docs/Config/Devices.md b/Engine/Source/Programs/Horde/Docs/Config/Devices.md
--- a/Engine/Source/Programs/Horde/Docs/Config/Devices.md
+++ b/Engine/Source/Programs/Horde/Docs/Config/Devices.md
@@ -2,7 +2,9 @@
 
 # Device Manager
 
-The Horde device manager administers mobile and console development kit resources. The service is used extensively at Epic and is quite mature.  It features:
+The Horde device manager administers mobile and console development kit
+resources. The service is used extensively at Epic and is quite mature.
+It features:
 
 * Configurable device platforms and pools
 * Shared remote device resources with user checkouts
@@ -11,11 +13,17 @@ The Horde device manager administers mobile and console development kit resource
 * Device usage history, telemetry, and pool health reporting
 * Integration with the Gauntlet automation framework
 
-##  Shared Pools
+## Shared Pools
 
-Horde users use shared pools to check out remote device resources for development and testing. Users check out devices through the dashboard, with devices being returned to the pool either by an explicit check-in or when a configurable duration has been exceeded.  
+Horde users use shared pools to check out remote device resources for
+development and testing. Users check out devices through the dashboard,
+with devices being returned to the pool either by an explicit check-in or
+when a configurable duration has been exceeded.
 
-There is a notification sink that notifies users 24 hours before the check-in will expire so they can renew if desired. There is a subsequent notification once the checkout has expired and the device is returned to the pool.
+There is a notification sink that notifies users 24 hours before the
+check-in will expire so they can renew if desired. There is a subsequent
+notification once the checkout has expired and the device is returned to
+the pool.
 
 It is also possible to set up automation jobs targeting the checked-out devices to install builds, for example.
 
@@ -28,48 +36,52 @@ However, you can also drive custom solutions using the same REST API.
 
 ## Platform and Pool Configuration
 
-Device platform hardware is partitioned into pools for use by automated tests and users. Device platforms and pools are configured using the `Devices` section of the globals.json file (see [DeviceConfig](Schema/Globals.md#deviceconfig)).
+Device platform hardware is partitioned into pools for use by automated
+tests and users. Device platforms and pools are configured using the
+`Devices` section of the globals.json file (see
+[DeviceConfig](Schema/Globals.md#deviceconfig)).
 
-#### Example
+### Example
 
 The following config fragment declares a device manager configuration that:
 
 * Adds an `Android` device platform specifying several models
-* Adds two device pools, an `Automation` pool which automated tests can use, and a `Shared` pool from which users can check out remote device hardware
-
----
-
-	"devices": {
-		"platforms": [
-			{
-				"id": "android",
-				"name": "Android",
-				"models": [
-					"Pixel4",
-					"Pixel5",
-					"Pixel8"
-				]
-			}
-		],
-		"pools": [
-			{
-				"id": "ue5",
-				"name": "UE5",
-				"poolType": "Automation"
-			},
-			{
-				"id": "remote-ue5",
-				"name": "Remote UE5",
-				"poolType": "Shared"
-			}
-		]
-	}
-
----
+* Adds two device pools, an `Automation` pool which automated tests can
+use, and a `Shared` pool from which users can check out remote device
+hardware
+
+```json
+    "devices": {
+        "platforms": [
+            {
+                "id": "android",
+                "name": "Android",
+                "models": [
+                    "Pixel4",
+                    "Pixel5",
+                    "Pixel8"
+                ]
+            }
+        ],
+        "pools": [
+            {
+                "id": "ue5",
+                "name": "UE5",
+                "poolType": "Automation"
+            },
+            {
+                "id": "remote-ue5",
+                "name": "Remote UE5",
+                "poolType": "Shared"
+            }
+        ]
+    }
+```
 
 ## Device Configuration
 
-Shared and automation devices are added and managed through the Horde Dashboard by navigating to `SERVER => Resources => Devices`.
+Shared and automation devices are added and managed through the Horde
+Dashboard by navigating to `SERVER => Resources => Devices`.
 
 This includes support for:
 
@@ -82,21 +94,27 @@ This includes support for:
 
 ## Gauntlet Integration
 
-[Gauntlet](https://docs.unrealengine.com/en-US/gauntlet-automation-framework-overview-in-unreal-engine) tests can reserve hardware from the Horde device manager. This integration includes features such as reporting problems with devices and recovery with new devices. It also supports reservation blocks, which can be used to reuse a device with an installed build across a series of automation tests.  
+[Gauntlet](https://docs.unrealengine.com/en-US/gauntlet-automation-framework-overview-in-unreal-engine)
+tests can reserve hardware from the Horde device manager. This integration
+includes features such as reporting problems with devices and recovery
+with new devices. It also supports reservation blocks, which can be used
+to reuse a device with an installed build across a series of automation
+tests.  
 
-#### Example
+### Example
 
 The following [**BuildGraph**](https://docs.unrealengine.com/en-US/buildgraph-for-unreal-engine/) fragment declares:
 
 * `HordeDeviceManager` and `HordeDevicePool` properties that specify your Horde server and which pool to use.
 * Adds a `BootTest Android` node, which will reserve an Android Pixel 8 for the test
 
----
+  ```xml
+  <Property Name="HordeDeviceManager" Value="https://horde.yourdomain.com" />
+  <Property Name="HordeDevicePool" Value="UE5" />
+
+  <Node Name="BootTest Android">
+      <Command Name="RunUnreal" Arguments="-test=UE.BootTest -platform=Android " -deviceurl=&quot;$(HordeDeviceManager)&quot; -devicepool=&quot;$(HordeDevicePool)&quot; -PerfModel=Pixel8/>
+  </Node>
+  ```
 
-	<Property Name="HordeDeviceManager" Value="https://horde.yourdomain.com" />
-	<Property Name="HordeDevicePool" Value="UE5" />
-		
-	<Node Name="BootTest Android">
-		<Command Name="RunUnreal" Arguments="-test=UE.BootTest -platform=Android " -deviceurl=&quot;$(HordeDeviceManager)&quot; -devicepool=&quot;$(HordeDevicePool)&quot; -PerfModel=Pixel8/>
-	</Node>
 ---
diff --git a/Engine/Source/Programs/Horde/Docs/Config/Plugins.md b/Engine/Source/Programs/Horde/Docs/Config/Plugins.md
--- a/Engine/Source/Programs/Horde/Docs/Config/Plugins.md
+++ b/Engine/Source/Programs/Horde/Docs/Config/Plugins.md
@@ -4,7 +4,7 @@
 
 Horde is implemented as a core host process and set of plugins. Most
 plugins correspond to a feature area can be enabled or disabled
-individually, though some plugins build on functionality providied by
+individually, though some plugins build on functionality provided by
 others. The [Storage](../Config.md#storage) and
 [Compute](../Config.md#compute) plugins are notable foundational components
 of Horde that other systems rely on.
diff --git a/Engine/Source/Programs/Horde/Docs/Deployment/Agent.md b/Engine/Source/Programs/Horde/Docs/Deployment/Agent.md
--- a/Engine/Source/Programs/Horde/Docs/Deployment/Agent.md
+++ b/Engine/Source/Programs/Horde/Docs/Deployment/Agent.md
@@ -25,15 +25,19 @@ page and clicking on the **Get agent software download token** link.
 
 #### Windows (PowerShell)
 
+```ps1
     Invoke-WebRequest -Uri https://[HORDE-SERVER-URL]/api/v1/agentsoftware/default/zip -OutFile C:\Horde\HordeAgent.zip -Headers @{ 'Authorization' = 'Bearer [AUTH-TOKEN]' }
     Expand-Archive -LiteralPath C:\Horde\HordeAgent.zip -DestinationPath C:\Horde -Force
+```
 
 The -Headers parameter and value are not required when using an unauthenticated server.
 
 #### Mac & Linux
 
+```sh
     curl https://[HORDE-SERVER-URL]/api/v1/agentsoftware/default/zip --output ~/Horde/HordeAgent.zip -H "Authorization: Bearer [AUTH-TOKEN]"
     unzip -o ~/Horde/HordeAgent.zip -d ~/Horde/
+```
 
 The -H parameter and value are not required when using an unauthenticated server.
 
@@ -63,7 +67,9 @@ When scripting agent deployment, you can either modify the build hosted by the s
 configuration by default, or use the agent's `SetServer` command to modify the configuration file after downloading it.
 This command can be invoked as:
 
+```bat
     dotnet HordeAgent.dll SetServer -Name=.. -Url=.. -Token=...
+```
 
 Adding the `-Default` argument will configure this server to be used by default. Run with the `-Help` argument for a
 full list of available options.
@@ -96,19 +102,24 @@ Running the MSI installer will configure the Horde Agent to run as a background
 the agent directly from the server and configuring it manually, a service can be registered by running the following
 command:
 
+```bat
     dotnet HordeAgent.dll service install [-UserName=..] [-Password=..]
+```
 
 Where `-UserName` and `-Password` specify credentials for the account to run the service under.
 
 The service may be uninstalled using the following command:
 
+```bat
     dotnet HordeAgent.dll service uninstall
+```
 
 #### Mac
 
 Create a `/Library/LaunchAgents/epic.hordeagent.plist` file describing the daemon configuration (substituting the
 `{{ HORDE_SERVICE_ACCOUNT }}` variables as appropriate).
 
+```xml
     <?xml version="1.0" encoding="UTF-8"?>
     <!DOCTYPE plist PUBLIC "-//Apple//DTD PLIST 1.0//EN" "http://www.apple.com/DTDs/PropertyList-1.0.dtd">
     <plist version="1.0">
@@ -154,18 +165,25 @@ Create a `/Library/LaunchAgents/epic.hordeagent.plist` file describing the daemo
         <integer>10</integer>
     </dict>
     </plist>
+```
 
 Adjust `/etc/newsyslog.conf` Log file out size limit (optional):
 
+```sh
     append '/Library/Logs/hordeagent_error.log 700 2 1000 * J'
+```
 
 Set any Horde agent environment variables you want to be defined outside of the plist (optional):
 
+```sh
     launchctl setenv Horde:WorkingDirectory {horde_working_directory}
+```
 
 Launch the daemon:
 
+```sh
     launchctl load -w /Library/LaunchAgents/epic.hordeagent.plist
+```
 
 #### Linux
 
@@ -176,6 +194,7 @@ user must have `sudo` access to restart/shutdown/autoscale Horde agents.
 Create a service descriptor file in `/etc/systemd/system/horde-agent.service` (substitute the `{{ HORDE_PATH }}`,
 `{{ HORDE_WORKING_DIRECTORY }}` and `{{ HORDE_SERVICE_ACCOUNT }}` variables as appropriate):
 
+```ini
     [Unit]
     Description=Horde Agent
 
@@ -194,10 +213,13 @@ Create a service descriptor file in `/etc/systemd/system/horde-agent.service` (s
 
     [Install]
     WantedBy=multi-user.target
+```
 
 Launch the daemon:
 
+```sh
     systemctl daemon-reload
+```
 
 ### Working Directory
 
diff --git a/Engine/Source/Programs/Horde/Docs/Internals/CsApi.md b/Engine/Source/Programs/Horde/Docs/Internals/CsApi.md
--- a/Engine/Source/Programs/Horde/Docs/Internals/CsApi.md
+++ b/Engine/Source/Programs/Horde/Docs/Internals/CsApi.md
@@ -16,12 +16,11 @@ provides typed request and response classes, or via the native
 C# API (`IHordeClient`). The `IHordeClient` interface is the main
 entry point, but provides a method for getting a HTTP client if necessary (`IHordeClient.CreateHttpClient`).
 
-While the HTTP API is more complete, we intent to move all functionality
+While the HTTP API is more complete, we intend to move all functionality
 into the native C# API over time. Doing so allows client tooling to write
-idiomatic C# code without having to worry about marshalling and transport
-concerns, and allows clients to write code against the same public APIs
-exposed by plugins on the server without having to concern themselves
-with implementation details.
+idiomatic C# code without having to worry about marshalling and transport,
+and allows clients to write code against the same public APIs
+exposed by plugins on the server.
 
 This approach gives more flexibility for prototyping code outside
 the server then moving it into the server as it makes sense, and
@@ -47,13 +46,14 @@ On MacOS and Linux, this is stored in the following JSON file:
   ```
 
 The `HordeOptions.GetDefaultServerUrl()` method will retrieve the user's
-preferred server.
+preferred server for the current platform.
 
 ## Connecting to Horde
 
-Horde uses .NET's standard dependency injection pattern.
+Horde is typically configured using .NET's standard dependency injection
+library.
 
-AutomationTool creates a service collection including Horde by default,
+AutomationTool creates a service provider including Horde by default,
 allowing you to get a Horde client interface as follows:
 
   ```cs
@@ -82,7 +82,7 @@ standard [.NET options pattern](https://learn.microsoft.com/en-us/aspnet/core/fu
   });
   ```
 
-Or through the `AddHorde()` method:
+Or through a delegate passed to the `AddHorde()` method:
 
   ```cs
   serviceCollection.AddHorde(options => 
@@ -90,4 +90,3 @@ Or through the `AddHorde()` method:
     options.ServerUrl = new Uri("http://my-horde-server.com");
   });
   ```
-  
\ No newline at end of file
