Horde: Set agent's default setting for EC2 support to disabled

Should be opt-in for when agent runs on AWS.

diff --git a/Engine/Source/Programs/Horde/HordeAgent/AgentSettings.cs b/Engine/Source/Programs/Horde/HordeAgent/AgentSettings.cs
--- a/Engine/Source/Programs/Horde/HordeAgent/AgentSettings.cs
+++ b/Engine/Source/Programs/Horde/HordeAgent/AgentSettings.cs
@@ -162,7 +162,7 @@ namespace HordeAgent
 		/// <summary>
 		/// Queries information about the current agent through the AWS EC2 interface
 		/// </summary>
-		public bool EnableAwsEc2Support { get; set; } = true;
+		public bool EnableAwsEc2Support { get; set; } = false;
 
 		/// <summary>
 		/// Option to use a local storage client rather than connecting through the server. Primarily for convenience when debugging / iterating locally.
