Horde: Ensure all uses of agent settings are hot-reloadable

Use IOptionsMonitor everywhere that AgentSettings is injected.

diff --git a/Engine/Source/Programs/Horde/HordeAgent.Tests/Services/AwsInstanceLifecycleServiceTests.cs b/Engine/Source/Programs/Horde/HordeAgent.Tests/Services/AwsInstanceLifecycleServiceTests.cs
--- a/Engine/Source/Programs/Horde/HordeAgent.Tests/Services/AwsInstanceLifecycleServiceTests.cs
+++ b/Engine/Source/Programs/Horde/HordeAgent.Tests/Services/AwsInstanceLifecycleServiceTests.cs
@@ -92,28 +92,19 @@ public class FakeAwsImds
 	}
 }
 
+// Stub for fulfilling IOptionsMonitor interface during testing
+// Copied from HordeServerTests until a good way to share code between these is decided.
+internal class TestOptionsMonitor<T>(T value) : IOptionsMonitor<T>
+	where T : class, new()
+{
+	public T CurrentValue { get; } = value;
+	public T Get(string? name) => CurrentValue;
+	public IDisposable? OnChange(Action<T, string?> listener) => null;
+}
+
 [TestClass]
 public sealed class AwsInstanceLifecycleServiceTests : IDisposable
 {
-	// Stub for fulfilling IOptionsMonitor interface during testing
-	// Copied from HordeServerTests until a good way to share code between these is decided.
-	class TestOptionsMonitor<T> : IOptionsMonitor<T>
-		where T : class, new()
-	{
-		public TestOptionsMonitor(T value)
-		{
-			CurrentValue = value;
-		}
-
-		public T CurrentValue { get; }
-
-		public T Get(string? name)
-			=> CurrentValue;
-
-		public IDisposable? OnChange(Action<T, string?> listener)
-			=> null;
-	}
-
 	private readonly WorkerServiceStub _workerService = new();
 	private readonly HttpClient _httpClient;
 	private readonly FakeAwsImds _fakeImds = new();
@@ -134,7 +125,7 @@ public sealed class AwsInstanceLifecycleServiceTests : IDisposable
 		_terminationSignalFile = settings.GetTerminationSignalFile();
 
 		_httpClient = _fakeImds.GetHttpClient();
-		_service = new AwsInstanceLifecycleService(_workerService, _httpClient, new OptionsWrapper<AgentSettings>(settings), loggerFactory.CreateLogger<AwsInstanceLifecycleService>());
+		_service = new AwsInstanceLifecycleService(_workerService, _httpClient, new TestOptionsMonitor<AgentSettings>(settings), loggerFactory.CreateLogger<AwsInstanceLifecycleService>());
 		_service._timeToLiveAsg = TimeSpan.FromMilliseconds(10);
 		_service._timeToLiveSpot = TimeSpan.FromMilliseconds(20);
 		_service._terminationBufferTime = TimeSpan.FromMilliseconds(2);
diff --git a/Engine/Source/Programs/Horde/HordeAgent.Tests/Services/TelemetryServiceTest.cs b/Engine/Source/Programs/Horde/HordeAgent.Tests/Services/TelemetryServiceTest.cs
--- a/Engine/Source/Programs/Horde/HordeAgent.Tests/Services/TelemetryServiceTest.cs
+++ b/Engine/Source/Programs/Horde/HordeAgent.Tests/Services/TelemetryServiceTest.cs
@@ -6,7 +6,6 @@ using System.Threading;
 using System.Threading.Tasks;
 using HordeAgent.Services;
 using Microsoft.Extensions.Logging.Abstractions;
-using Microsoft.Extensions.Options;
 using Microsoft.VisualStudio.TestTools.UnitTesting;
 
 namespace HordeAgent.Tests.Services;
@@ -20,8 +19,7 @@ public sealed class TelemetryServiceTest : IDisposable
 	public TelemetryServiceTest()
 	{
 		AgentSettings settings = new() { Server = "Test", ServerProfiles = { ["Test"] = new ServerProfile() { Name = "Test", Url = new Uri("http://localhost:1234") } } };
-		OptionsWrapper<AgentSettings> settingsOpt = new(settings);
-
+		TestOptionsMonitor<AgentSettings> settingsOpt = new(settings);
 		using WorkerService workerService = new WorkerService(null!, null!, null!, null!);
 		GrpcService grpcService = new(settingsOpt, NullLogger<GrpcService>.Instance, _loggerFactory);
 
diff --git a/Engine/Source/Programs/Horde/HordeAgent/Leases/Handlers/ComputeHandler.cs b/Engine/Source/Programs/Horde/HordeAgent/Leases/Handlers/ComputeHandler.cs
--- a/Engine/Source/Programs/Horde/HordeAgent/Leases/Handlers/ComputeHandler.cs
+++ b/Engine/Source/Programs/Horde/HordeAgent/Leases/Handlers/ComputeHandler.cs
@@ -22,16 +22,16 @@ namespace HordeAgent.Leases.Handlers
 	class ComputeHandler : LeaseHandler<ComputeTask>
 	{
 		readonly ComputeListenerService _listenerService;
-		readonly AgentSettings _settings;
+		readonly IOptionsMonitor<AgentSettings> _settings;
 
 		/// <summary>
 		/// Constructor
 		/// </summary>
-		public ComputeHandler(RpcLease lease, ComputeListenerService listenerService, IOptions<AgentSettings> settings)
+		public ComputeHandler(RpcLease lease, ComputeListenerService listenerService, IOptionsMonitor<AgentSettings> settings)
 			: base(lease)
 		{
 			_listenerService = listenerService;
-			_settings = settings.Value;
+			_settings = settings;
 		}
 
 		/// <inheritdoc/>
@@ -85,12 +85,12 @@ namespace HordeAgent.Leases.Handlers
 								Dictionary<string, string?> newEnvVars = new ()
 								{
 									["UE_HORDE_SHARED_DIR"] = sharedDir.FullName,
-									["UE_HORDE_TERMINATION_SIGNAL_FILE"] = _settings.GetTerminationSignalFile().FullName,
-									["UE_HORDE_CPU_COUNT"] = Convert.ToString(_settings.CpuCount),
-									["UE_HORDE_CPU_MULTIPLIER"] = Convert.ToString(_settings.CpuMultiplier, CultureInfo.InvariantCulture),
+									["UE_HORDE_TERMINATION_SIGNAL_FILE"] = _settings.CurrentValue.GetTerminationSignalFile().FullName,
+									["UE_HORDE_CPU_COUNT"] = Convert.ToString(_settings.CurrentValue.CpuCount),
+									["UE_HORDE_CPU_MULTIPLIER"] = Convert.ToString(_settings.CurrentValue.CpuMultiplier, CultureInfo.InvariantCulture),
 								};
 								
-								AgentMessageHandler worker = new AgentMessageHandler(sandboxDir, newEnvVars, false, _settings.WineExecutablePath, _settings.ContainerEngineExecutablePath, logger);
+								AgentMessageHandler worker = new AgentMessageHandler(sandboxDir, newEnvVars, false, _settings.CurrentValue.WineExecutablePath, _settings.CurrentValue.ContainerEngineExecutablePath, logger);
 								await worker.RunAsync(socket, cts.Token);
 								await socket.CloseAsync(cts.Token);
 								return LeaseResult.Success;
@@ -147,7 +147,7 @@ namespace HordeAgent.Leases.Handlers
 
 		private void ClearTerminationSignalFile(ILogger logger)
 		{
-			string path = _settings.GetTerminationSignalFile().FullName;
+			string path = _settings.CurrentValue.GetTerminationSignalFile().FullName;
 			try
 			{
 				File.Delete(path);
@@ -164,9 +164,9 @@ namespace HordeAgent.Leases.Handlers
 	class ComputeHandlerFactory : LeaseHandlerFactory<ComputeTask>
 	{
 		readonly ComputeListenerService _listenerService;
-		readonly IOptions<AgentSettings> _settings;
+		readonly IOptionsMonitor<AgentSettings> _settings;
 
-		public ComputeHandlerFactory(ComputeListenerService listenerService, IOptions<AgentSettings> settings)
+		public ComputeHandlerFactory(ComputeListenerService listenerService, IOptionsMonitor<AgentSettings> settings)
 		{
 			_listenerService = listenerService;
 			_settings = settings;
diff --git a/Engine/Source/Programs/Horde/HordeAgent/Leases/Handlers/JobHandler.cs b/Engine/Source/Programs/Horde/HordeAgent/Leases/Handlers/JobHandler.cs
--- a/Engine/Source/Programs/Horde/HordeAgent/Leases/Handlers/JobHandler.cs
+++ b/Engine/Source/Programs/Horde/HordeAgent/Leases/Handlers/JobHandler.cs
@@ -17,9 +17,9 @@ namespace HordeAgent.Leases.Handlers
 {
 	class JobHandler : LeaseHandler<ExecuteJobTask>
 	{
-		private readonly IOptions<AgentSettings> _settings;
+		private readonly IOptionsMonitor<AgentSettings> _settings;
 		
-		public JobHandler(RpcLease lease, IOptions<AgentSettings> settings)
+		public JobHandler(RpcLease lease, IOptionsMonitor<AgentSettings> settings)
 			: base(lease)
 		{
 			_settings = settings;
@@ -55,7 +55,7 @@ namespace HordeAgent.Leases.Handlers
 				environment[HordeHttpClient.HordeUrlEnvVarName] = session.HordeClient.ServerUrl.ToString();
 				environment[HordeHttpClient.HordeTokenEnvVarName] = executeTask.Token;
 				environment["UE_LOG_JSON_TO_STDOUT"] = "1";
-				environment["UE_HORDE_OTEL_SETTINGS"] = OpenTelemetrySettingsExtensions.Serialize(_settings.Value.OpenTelemetry, true);
+				environment["UE_HORDE_OTEL_SETTINGS"] = OpenTelemetrySettingsExtensions.Serialize(_settings.CurrentValue.OpenTelemetry, true);
 
 				int exitCode = await RunDotNetProcessAsync(driverAssembly, arguments, environment, AgentApp.IsSelfContained, logger, cancellationToken);
 				logger.LogInformation("Driver finished with exit code {ExitCode}", exitCode);
@@ -77,9 +77,9 @@ namespace HordeAgent.Leases.Handlers
 
 	class JobHandlerFactory : LeaseHandlerFactory<ExecuteJobTask>
 	{
-		private readonly IOptions<AgentSettings> _settings;
+		private readonly IOptionsMonitor<AgentSettings> _settings;
 		
-		public JobHandlerFactory(IOptions<AgentSettings> settings)
+		public JobHandlerFactory(IOptionsMonitor<AgentSettings> settings)
 		{
 			_settings = settings;
 		}
diff --git a/Engine/Source/Programs/Horde/HordeAgent/Leases/LeaseLoggerFactory.cs b/Engine/Source/Programs/Horde/HordeAgent/Leases/LeaseLoggerFactory.cs
--- a/Engine/Source/Programs/Horde/HordeAgent/Leases/LeaseLoggerFactory.cs
+++ b/Engine/Source/Programs/Horde/HordeAgent/Leases/LeaseLoggerFactory.cs
@@ -14,7 +14,7 @@ namespace HordeAgent.Leases
 	{
 		static TimeSpan MaxAge { get; } = TimeSpan.FromDays(3.0);
 
-		readonly AgentSettings _settings;
+		readonly IOptionsMonitor<AgentSettings> _settings;
 		readonly DirectoryReference _logDir;
 		readonly BackgroundTask _backgroundTask;
 		readonly ILogger _logger;
@@ -22,10 +22,10 @@ namespace HordeAgent.Leases
 		/// <summary>
 		/// Constructor
 		/// </summary>
-		public LeaseLoggerFactory(IOptions<AgentSettings> settings, ILogger<LeaseLoggerFactory> logger)
+		public LeaseLoggerFactory(IOptionsMonitor<AgentSettings> settings, ILogger<LeaseLoggerFactory> logger)
 		{
-			_settings = settings.Value;
-			_logDir = DirectoryReference.Combine(settings.Value.LogsDir, "Leases");
+			_settings = settings;
+			_logDir = DirectoryReference.Combine(settings.CurrentValue.LogsDir, "Leases");
 			_logger = logger;
 			_backgroundTask = BackgroundTask.StartNew(BackgroundCleanupAsync);
 		}
@@ -90,7 +90,7 @@ namespace HordeAgent.Leases
 			return LoggerFactory.Create(builder =>
 			{
 				builder.AddProvider(Logging.CreateFileLoggerProvider(_logDir, leaseId.ToString()));
-				if (_settings.WriteStepOutputToLogger)
+				if (_settings.CurrentValue.WriteStepOutputToLogger)
 				{
 					builder.AddProvider(new LoggerProvider(_logger));
 				}
diff --git a/Engine/Source/Programs/Horde/HordeAgent/Leases/LeaseManager.cs b/Engine/Source/Programs/Horde/HordeAgent/Leases/LeaseManager.cs
--- a/Engine/Source/Programs/Horde/HordeAgent/Leases/LeaseManager.cs
+++ b/Engine/Source/Programs/Horde/HordeAgent/Leases/LeaseManager.cs
@@ -89,7 +89,7 @@ namespace HordeAgent.Leases
 		readonly ISystemMetrics _systemMetrics;
 		readonly Dictionary<string, LeaseHandlerFactory> _typeUrlToLeaseHandler;
 		readonly LeaseLoggerFactory _leaseLoggerFactory;
-		readonly IOptions<AgentSettings> _settings;
+		readonly IOptionsMonitor<AgentSettings> _settings;
 		readonly Tracer _tracer;
 		readonly ILogger _logger;
 
@@ -102,7 +102,7 @@ namespace HordeAgent.Leases
 			ISystemMetrics systemMetrics,
 			IEnumerable<LeaseHandlerFactory> leaseHandlerFactories,
 			LeaseLoggerFactory leaseLoggerFactory,
-			IOptions<AgentSettings> settings,
+			IOptionsMonitor<AgentSettings> settings,
 			Tracer tracer,
 			ILogger logger)
 		{
@@ -124,7 +124,7 @@ namespace HordeAgent.Leases
 				serviceProvider.GetRequiredService<ISystemMetrics>(),
 				serviceProvider.GetRequiredService<IEnumerable<LeaseHandlerFactory>>(),
 				serviceProvider.GetRequiredService<LeaseLoggerFactory>(),
-				serviceProvider.GetRequiredService<IOptions<AgentSettings>>(),
+				serviceProvider.GetRequiredService<IOptionsMonitor<AgentSettings>>(),
 				serviceProvider.GetRequiredService<Tracer>(),
 				serviceProvider.GetRequiredService<ILogger<LeaseManager>>())
 		{
@@ -201,7 +201,7 @@ namespace HordeAgent.Leases
 			if (graceful)
 			{
 				TimeSpan ttl = TimeSpan.FromSeconds(1);
-				await WriteTerminationSignalFileAsync(_settings.Value.GetTerminationSignalFile().FullName, reason, DateTime.UtcNow + ttl, ttl, CancellationToken.None);
+				await WriteTerminationSignalFileAsync(_settings.CurrentValue.GetTerminationSignalFile().FullName, reason, DateTime.UtcNow + ttl, ttl, CancellationToken.None);
 				await Task.Delay(TimeSpan.FromSeconds(6)); // Allow workload some time to act on signal file
 			}
 			
diff --git a/Engine/Source/Programs/Horde/HordeAgent/Services/AwsInstanceLifecycleService.cs b/Engine/Source/Programs/Horde/HordeAgent/Services/AwsInstanceLifecycleService.cs
--- a/Engine/Source/Programs/Horde/HordeAgent/Services/AwsInstanceLifecycleService.cs
+++ b/Engine/Source/Programs/Horde/HordeAgent/Services/AwsInstanceLifecycleService.cs
@@ -93,7 +93,7 @@ class AwsInstanceLifecycleService : BackgroundService
 	/// <summary>
 	/// Constructor
 	/// </summary>
-	public AwsInstanceLifecycleService(IWorkerService workerService, HttpClient httpClient, IOptions<AgentSettings> settings, ILogger<AwsInstanceLifecycleService> logger)
+	public AwsInstanceLifecycleService(IWorkerService workerService, HttpClient httpClient, IOptionsMonitor<AgentSettings> settings, ILogger<AwsInstanceLifecycleService> logger)
 	{
 		_workerService = workerService;
 		_httpClient = httpClient;
@@ -101,7 +101,7 @@ class AwsInstanceLifecycleService : BackgroundService
 		_terminationWarningCallback = OnTerminationWarningAsync;
 		_terminationCallback = OnTerminationAsync;
 		_logger = logger;
-		_terminationSignalFile = settings.Value.GetTerminationSignalFile();
+		_terminationSignalFile = settings.CurrentValue.GetTerminationSignalFile();
 	}
 
 	private async Task<Ec2InstanceState> GetStateAsync(CancellationToken cancellationToken)
diff --git a/Engine/Source/Programs/Horde/HordeAgent/Services/CapabilitiesService.cs b/Engine/Source/Programs/Horde/HordeAgent/Services/CapabilitiesService.cs
--- a/Engine/Source/Programs/Horde/HordeAgent/Services/CapabilitiesService.cs
+++ b/Engine/Source/Programs/Horde/HordeAgent/Services/CapabilitiesService.cs
@@ -30,16 +30,16 @@ namespace HordeAgent.Services
 		static readonly DateTimeOffset s_startTime = DateTimeOffset.Now;
 		static readonly DateTimeOffset s_bootTime = DateTimeOffset.Now - TimeSpan.FromTicks(Environment.TickCount64 * TimeSpan.TicksPerMillisecond);
 
-		readonly AgentSettings _settings;
+		readonly IOptionsMonitor<AgentSettings> _settings;
 		readonly Tracer _tracer;
 		readonly ILogger _logger;
 
 		/// <summary>
 		/// Constructor
 		/// </summary>
-		public CapabilitiesService(IOptions<AgentSettings> settings, Tracer tracer, ILogger<CapabilitiesService> logger)
+		public CapabilitiesService(IOptionsMonitor<AgentSettings> settings, Tracer tracer, ILogger<CapabilitiesService> logger)
 		{
-			_settings = settings.Value;
+			_settings = settings;
 			_tracer = tracer;
 			_logger = logger;
 		}
@@ -207,7 +207,7 @@ namespace HordeAgent.Services
 				}
 
 				// Add EC2 properties if needed
-				if (_settings.EnableAwsEc2Support)
+				if (_settings.CurrentValue.EnableAwsEc2Support)
 				{
 					await AddAwsPropertiesAsync(capabilities.Properties, logger);
 				}
@@ -230,12 +230,12 @@ namespace HordeAgent.Services
 				capabilities.Properties.Add($"{KnownPropertyNames.OsFamilyCompatibility}=Linux");
 
 				// Add EC2 properties if needed
-				if (_settings.EnableAwsEc2Support)
+				if (_settings.CurrentValue.EnableAwsEc2Support)
 				{
 					await AddAwsPropertiesAsync(capabilities.Properties, logger);
 				}
 				
-				if (_settings.WineExecutablePath != null)
+				if (_settings.CurrentValue.WineExecutablePath != null)
 				{
 					capabilities.Properties.Add($"{KnownPropertyNames.OsFamilyCompatibility}=Windows");
 					capabilities.Properties.Add($"{KnownPropertyNames.WineEnabled}=true");
@@ -400,13 +400,13 @@ namespace HordeAgent.Services
 			}
 			
 			IPAddress? ip;
-			if (_settings.ComputeIp != null)
+			if (_settings.CurrentValue.ComputeIp != null)
 			{
-				ip = IPAddress.Parse(_settings.ComputeIp);
+				ip = IPAddress.Parse(_settings.CurrentValue.ComputeIp);
 			}
 			else
 			{
-				ip = await GetLocalIpAddressAsync(_settings.GetCurrentServerProfile().Url.Host);
+				ip = await GetLocalIpAddressAsync(_settings.CurrentValue.GetCurrentServerProfile().Url.Host);
 			}
 			 
 			if (ip == null)
@@ -415,10 +415,10 @@ namespace HordeAgent.Services
 			}
 
 			// Add the compute configuration
-			if (ip != null && _settings.ComputePort != 0)
+			if (ip != null && _settings.CurrentValue.ComputePort != 0)
 			{
 				capabilities.Properties.Add($"ComputeIp={ip}");
-				capabilities.Properties.Add($"ComputePort={_settings.ComputePort}");
+				capabilities.Properties.Add($"ComputePort={_settings.CurrentValue.ComputePort}");
 			}
 
 			// Get the time that the machine booted
@@ -471,7 +471,7 @@ namespace HordeAgent.Services
 			capabilities.Properties.Add($"DotNetFramework={RuntimeInformation.FrameworkDescription}");
 
 			// Add any additional properties from the config file
-			capabilities.Properties.AddRange(_settings.Properties.Select(kvp => $"{kvp.Key}={kvp.Value}"));
+			capabilities.Properties.AddRange(_settings.CurrentValue.Properties.Select(kvp => $"{kvp.Key}={kvp.Value}"));
 			return capabilities;
 		}
 
diff --git a/Engine/Source/Programs/Horde/HordeAgent/Services/ComputeListenerService.cs b/Engine/Source/Programs/Horde/HordeAgent/Services/ComputeListenerService.cs
--- a/Engine/Source/Programs/Horde/HordeAgent/Services/ComputeListenerService.cs
+++ b/Engine/Source/Programs/Horde/HordeAgent/Services/ComputeListenerService.cs
@@ -44,7 +44,7 @@ namespace HordeAgent.Services
 			}
 		}
 
-		readonly AgentSettings _settings;
+		readonly IOptionsMonitor<AgentSettings> _settings;
 		readonly ILogger _logger;
 
 		Task _serverTask = Task.CompletedTask;
@@ -57,9 +57,9 @@ namespace HordeAgent.Services
 		/// <summary>
 		/// Constructor
 		/// </summary>
-		public ComputeListenerService(IOptions<AgentSettings> settings, ILogger<ComputeListenerService> logger)
+		public ComputeListenerService(IOptionsMonitor<AgentSettings> settings, ILogger<ComputeListenerService> logger)
 		{
-			_settings = settings.Value;
+			_settings = settings;
 			_logger = logger;
 		}
 
@@ -113,9 +113,9 @@ namespace HordeAgent.Services
 		/// <inheritdoc/>
 		public Task StartAsync(CancellationToken cancellationToken)
 		{
-			if (_settings.ComputePort != 0)
+			if (_settings.CurrentValue.ComputePort != 0)
 			{
-				Start(IPAddress.Any, _settings.ComputePort);
+				Start(IPAddress.Any, _settings.CurrentValue.ComputePort);
 			}
 			return Task.CompletedTask;
 		}
diff --git a/Engine/Source/Programs/Horde/HordeAgent/Services/GrpcService.cs b/Engine/Source/Programs/Horde/HordeAgent/Services/GrpcService.cs
--- a/Engine/Source/Programs/Horde/HordeAgent/Services/GrpcService.cs
+++ b/Engine/Source/Programs/Horde/HordeAgent/Services/GrpcService.cs
@@ -18,7 +18,7 @@ namespace HordeAgent.Services
 	/// </summary>
 	class GrpcService
 	{
-		private readonly IOptions<AgentSettings> _settings;
+		private readonly IOptionsMonitor<AgentSettings> _settings;
 		private readonly ServerProfile _serverProfile;
 		private readonly ILogger _logger;
 		private readonly ILoggerFactory _loggerFactory;
@@ -34,10 +34,10 @@ namespace HordeAgent.Services
 		/// <param name="settings"></param>
 		/// <param name="logger"></param>
 		/// <param name="loggerFactory"></param>
-		public GrpcService(IOptions<AgentSettings> settings, ILogger<GrpcService> logger, ILoggerFactory loggerFactory)
+		public GrpcService(IOptionsMonitor<AgentSettings> settings, ILogger<GrpcService> logger, ILoggerFactory loggerFactory)
 		{
 			_settings = settings;
-			_serverProfile = settings.Value.GetCurrentServerProfile();
+			_serverProfile = settings.CurrentValue.GetCurrentServerProfile();
 			_logger = logger;
 			_loggerFactory = loggerFactory;
 		}
@@ -151,7 +151,7 @@ namespace HordeAgent.Services
 			CallInvoker invoker = channel.Intercept(headers =>
 			{
 				headers.Add("Horde-Agent-Version", AgentApp.Version);
-				headers.Add("Horde-Agent-Name", _settings.Value.GetAgentName());
+				headers.Add("Horde-Agent-Name", _settings.CurrentValue.GetAgentName());
 				return headers;
 			});
 
diff --git a/Engine/Source/Programs/Horde/HordeAgent/Services/SessionFactory.cs b/Engine/Source/Programs/Horde/HordeAgent/Services/SessionFactory.cs
--- a/Engine/Source/Programs/Horde/HordeAgent/Services/SessionFactory.cs
+++ b/Engine/Source/Programs/Horde/HordeAgent/Services/SessionFactory.cs
@@ -95,9 +95,9 @@ namespace HordeAgent.Services
 		/// <summary>
 		/// Creates a new agent session
 		/// </summary>
-		public static async Task<Session> CreateAsync(CapabilitiesService capabilitiesService, GrpcService grpcService, StatusService statusService, IOptions<AgentSettings> settings, HordeClientFactory hordeClientFactory, ILogger logger, CancellationToken cancellationToken)
+		public static async Task<Session> CreateAsync(CapabilitiesService capabilitiesService, GrpcService grpcService, StatusService statusService, IOptionsMonitor<AgentSettings> settings, HordeClientFactory hordeClientFactory, ILogger logger, CancellationToken cancellationToken)
 		{
-			AgentSettings currentSettings = settings.Value;
+			AgentSettings currentSettings = settings.CurrentValue;
 
 			// Get the working directory
 			if (currentSettings.WorkingDir == null)
@@ -343,13 +343,13 @@ namespace HordeAgent.Services
 		readonly GrpcService _grpcService;
 		readonly StatusService _statusService;
 		readonly HordeClientFactory _hordeClientFactory;
-		readonly IOptions<AgentSettings> _agentSettings;
+		readonly IOptionsMonitor<AgentSettings> _agentSettings;
 		readonly ILogger _logger;
 
 		/// <summary>
 		/// Constructor
 		/// </summary>
-		public SessionFactory(CapabilitiesService capabilitiesService, GrpcService grpcService, StatusService statusService, HordeClientFactory hordeClientFactory, IOptions<AgentSettings> agentSettings, ILogger<SessionFactory> logger)
+		public SessionFactory(CapabilitiesService capabilitiesService, GrpcService grpcService, StatusService statusService, HordeClientFactory hordeClientFactory, IOptionsMonitor<AgentSettings> agentSettings, ILogger<SessionFactory> logger)
 		{
 			_capabilitiesService = capabilitiesService;
 			_grpcService = grpcService;
