Horde: Add agent IPC message for setting CPU count and multiplier

diff --git a/Engine/Source/Programs/Horde/HordeAgent/IpcMessages.cs b/Engine/Source/Programs/Horde/HordeAgent/IpcMessages.cs
--- a/Engine/Source/Programs/Horde/HordeAgent/IpcMessages.cs
+++ b/Engine/Source/Programs/Horde/HordeAgent/IpcMessages.cs
@@ -52,6 +52,16 @@ namespace HordeAgent
 		/// Returns information about the server we're connected to.
 		/// </summary>
 		GetSettingsResponse = 5,
+		
+		/// <summary>
+		/// Set agent settings request
+		/// </summary>
+		SetSettingsRequest = 6,
+		
+		/// <summary>
+		/// Set agent settings response
+		/// </summary>
+		SetSettingsResponse = 7,
 	}
 
 	/// <summary>
@@ -155,13 +165,20 @@ namespace HordeAgent
 	/// Settings for the agent
 	/// </summary>
 	/// <param name="ServerUrl">Url of the Horde server</param>
-	public record class AgentSettingsMessage(Uri? ServerUrl);
+	public record AgentSettingsMessage(Uri? ServerUrl);
 
 	/// <summary>
 	/// Sets the enabled state for the agent
 	/// </summary>
 	/// <param name="IsEnabled">Whether the agent is enabled or not</param>
-	public record class AgentEnabledMessage(bool IsEnabled);
+	public record AgentEnabledMessage(bool IsEnabled);
+	
+	/// <summary>
+	/// Set settings request message
+	/// </summary>
+	/// <param name="CpuCount">Maximum number of logical CPU cores workloads should use</param>
+	/// <param name="CpuMultiplier">CPU core multiplier applied to CPU core count setting</param>
+	public record AgentSetSettingsRequest(int CpuCount, double CpuMultiplier);
 
 	/// <summary>
 	/// Current status of the agent
diff --git a/Engine/Source/Programs/Horde/HordeAgent/Services/StatusService.cs b/Engine/Source/Programs/Horde/HordeAgent/Services/StatusService.cs
--- a/Engine/Source/Programs/Horde/HordeAgent/Services/StatusService.cs
+++ b/Engine/Source/Programs/Horde/HordeAgent/Services/StatusService.cs
@@ -208,6 +208,11 @@ namespace HordeAgent.Services
 							response.Set(AgentMessageType.GetStatusResponse, Current);
 							await response.SendAsync(pipeServer, cancellationToken);
 							break;
+						case AgentMessageType.SetSettingsRequest:
+							AgentSetSettingsRequest req = request.Parse<AgentSetSettingsRequest>();
+							_settings.CurrentValue.CpuCount = req.CpuCount;
+							_settings.CurrentValue.CpuMultiplier = req.CpuMultiplier;
+							break;
 						case AgentMessageType.GetSettingsRequest:
 							response.Set(AgentMessageType.GetSettingsResponse, GetSettingsMessage());
 							await response.SendAsync(pipeServer, cancellationToken);
