Horde: Fix various issues with device view permalinks

diff --git a/Engine/Source/Programs/Horde/HordeDashboard/src/components/DeviceView.tsx b/Engine/Source/Programs/Horde/HordeDashboard/src/components/DeviceView.tsx
--- a/Engine/Source/Programs/Horde/HordeDashboard/src/components/DeviceView.tsx
+++ b/Engine/Source/Programs/Horde/HordeDashboard/src/components/DeviceView.tsx
@@ -1,4 +1,4 @@
-import { DefaultButton, DetailsList, DetailsListLayoutMode, Dropdown, FocusZone, FocusZoneDirection, IColumn, Icon, IconButton, IGroup, Label, Link as FluentLink, mergeStyleSets, Modal, Pivot, PivotItem, PrimaryButton, ScrollablePane, ScrollbarVisibility, SelectionMode, Spinner, SpinnerSize, Stack, Text, TextField, } from "@fluentui/react";
+import { DefaultButton, DetailsList, DetailsListLayoutMode, Dropdown, FocusZone, FocusZoneDirection, IColumn, Icon, IconButton, IGroup, Label, Link as FluentLink, mergeStyleSets, Modal, Pivot, PivotItem, PrimaryButton, ScrollablePane, ScrollbarVisibility, SelectionMode, Spinner, SpinnerSize, Stack, Text, TextField, themeRulesStandardCreator, } from "@fluentui/react";
 import { observer } from "mobx-react-lite";
 import { observable, action, makeObservable } from "mobx";
 import moment from "moment";
@@ -54,7 +54,7 @@ class LocalState {
    search: URLSearchParams = new URLSearchParams(window.location.search);
 
    updateSearch(): boolean {
-
+      
       const state = { ...this.searchState } as DevicesSearchState;
 
       state.filterPools = state.filterPools?.sort((a, b) => a.localeCompare(b));
@@ -96,11 +96,12 @@ class LocalState {
    }
 
    resetState() {
+      this.initialized = false;
+      this.search = new URLSearchParams();
       this.searchState = {};
    }
-
-   @action
-   stateFromSearch() {
+   
+   private stateFromSearch() {
 
       const state: DevicesSearchState = {};
 
@@ -181,6 +182,16 @@ class LocalState {
       this.searchState.historyItem = undefined;
       this.updateSearch();
    }
+
+   initialize() 
+   {
+      if (!this.initialized) {
+         this.initialized = true;
+         this.stateFromSearch();
+      }
+   }
+
+   private initialized = false;
 }
 
 export const StatusNames = new Map<DeviceStatus, string>([
@@ -256,20 +267,13 @@ export const SearchUpdate: React.FC = observer(() => {
 const DevicePanel: React.FC = observer(() => {
 
    const [editState, setEditState] = useState<{ shown?: boolean, infoShown?: boolean, editNote?: boolean, device?: DeviceItem | undefined }>({});
-   const [pivotState, setPivotState] = useState<{ key: string, poolFilter: Set<string> }>({ key: pivotKeyShared, poolFilter: new Set() });
-   const [platformState, setPlatformState] = useState<Set<string>>(new Set());
+   const [pivotState, setPivotState] = useState<{ key: string, poolFilter: Set<string> }>({ key: localState?.searchState?.pivotKey ?? pivotKeyShared, poolFilter: new Set(localState?.searchState?.filterPools ?? []) });
+   const [platformState, setPlatformState] = useState<Set<string>>(new Set(localState?.searchState?.filterPlatforms ?? []));
    const [checkoutState, setCheckoutState] = useState<{ checkoutId?: string, checkinId?: string, showConfirm?: "in" | "out" | "error" }>({});
    const [telemState, setTelemState] = useState(false);
-   const [initDeviceUpdater, setInitDeviceUpdater] = useState({ initPage: false, initDevices: false });
-
+   
    const navigate = useNavigate();
 
-   useEffect(() => {
-
-      localState.setPivotKey(pivotState.key);
-
-   }, [pivotState.key])
-
    useEffect(() => {
 
       handler.start();
@@ -762,35 +766,6 @@ const DevicePanel: React.FC = observer(() => {
 
    const checkedOut = handler.getUserDeviceCheckouts(dashboard.userId);
 
-
-   if (handler.updated && !initDeviceUpdater.initDevices) { // Devices have been loaded, so now we can check for any history dialogs that need to be shown
-      if (localState.searchState.historyItem) {
-         let selectedDevice: GetDeviceResponse | undefined = undefined;
-         filteredDevices.every((item) => {
-            if (item.id === localState.searchState.historyItem) {
-               selectedDevice = item;
-               return false; // effectively breaks out of the 'every' function
-            }
-            return true;
-         });
-         if (selectedDevice) {
-            setEditState({ infoShown: true, device: { device: selectedDevice } });
-         }
-      }
-      setInitDeviceUpdater({ ...initDeviceUpdater, initDevices: true });
-      return null;
-   }
-
-   if (!initDeviceUpdater.initPage) {
-      if (localState.search) {
-         localState.stateFromSearch();
-         setPivotState({ key: localState.searchState.pivotKey ?? pivotState.key, poolFilter: new Set(localState.searchState.filterPools) });
-         setPlatformState(new Set(localState.searchState.filterPlatforms));
-      }
-      setInitDeviceUpdater({ ...initDeviceUpdater, initPage: true });
-      return null;
-   };
-
    const hasDevices = !!handler.getDevices().length
 
    return (<Stack>
@@ -818,6 +793,7 @@ const DevicePanel: React.FC = observer(() => {
                      placeholder="Device Search"
                      spellCheck={false}
                      autoComplete="off"
+                     validateOnLoad={false}
                      deferredValidationTime={1000}
                      defaultValue={localState.searchState.filterString}
                      styles={{
@@ -1044,6 +1020,18 @@ export const DeviceView: React.FC = () => {
 
    const { hordeClasses, modeColors } = getHordeStyling();
 
+   localState.initialize();
+
+   useEffect(() => {
+
+      return () => {
+         localState.resetState();
+      };
+
+   }, []);
+
+
+
    return <Stack className={hordeClasses.horde}>
       <TopNav />
       <Breadcrumbs items={[{ text: 'Devices' }]} />
