Horde: Fix issue with uri escaping in artifact modal

diff --git a/Engine/Source/Programs/Horde/HordeDashboard/src/components/artifacts/ArtifactsModal.tsx b/Engine/Source/Programs/Horde/HordeDashboard/src/components/artifacts/ArtifactsModal.tsx
--- a/Engine/Source/Programs/Horde/HordeDashboard/src/components/artifacts/ArtifactsModal.tsx
+++ b/Engine/Source/Programs/Horde/HordeDashboard/src/components/artifacts/ArtifactsModal.tsx
@@ -183,8 +183,9 @@ class ArtifactsHandler {
       if (this.baseSearch) {
          let url = `${window.location.pathname}${this.baseSearch}`;
          if (path?.length) {
-            url += `&artifactPath=${encodeURI(path)}`;
+            url += `&artifactPath=${encodeURIComponent(path.replaceAll("+", "%2B"))}`;
          }
+         
          navigate(url, { replace: true });
          console.log(this.artifact.id, path);
       }
@@ -638,7 +639,8 @@ const JobDetailArtifactsInner: React.FC<{ jobId: string; stepId: string, artifac
          return undefined;
       }
 
-      const path = encodeURI((handler.path ? handler.path + "/" : "") + item.text);
+      let path = (handler.path ? handler.path + "/" : "") + encodeURIComponent(item.text);
+      path = path.replaceAll("+", "%2B");
       const server = backend.serverUrl;
       return `${server}/api/v2/artifacts/${handler.artifact!.id}/file?path=${path}`;
 
