Horde: Fix hanging chrome when missing job template, also surface job error

diff --git a/Engine/Source/Programs/Horde/HordeDashboard/src/components/LogRender.tsx b/Engine/Source/Programs/Horde/HordeDashboard/src/components/LogRender.tsx
--- a/Engine/Source/Programs/Horde/HordeDashboard/src/components/LogRender.tsx
+++ b/Engine/Source/Programs/Horde/HordeDashboard/src/components/LogRender.tsx
@@ -175,7 +175,6 @@ const renderTags = (navigate: NavigateFunction, line: LogLine, lineNumber: numbe
          }
 
          if (type === "LogId") {
-            debugger;
             tagType = TagType.LogId;
          }
 
diff --git a/Engine/Source/Programs/Horde/HordeDashboard/src/components/jobDetailsV2/JobDetailViewV2.tsx b/Engine/Source/Programs/Horde/HordeDashboard/src/components/jobDetailsV2/JobDetailViewV2.tsx
--- a/Engine/Source/Programs/Horde/HordeDashboard/src/components/jobDetailsV2/JobDetailViewV2.tsx
+++ b/Engine/Source/Programs/Horde/HordeDashboard/src/components/jobDetailsV2/JobDetailViewV2.tsx
@@ -53,11 +53,10 @@ const JobBreadCrumbs: React.FC<{ jobDetails: JobDetailsV2 }> = observer(({ jobDe
       if (!jobDetails.jobError) {
          return <Breadcrumbs items={[{ text: "Loading Job" }]} title={"Loading Job"} spinner={true} />
       }
-      else {
-         console.error(`Unable to load job ${jobDetails.jobId}: ${jobDetails.jobError}`)
-      }
+   } 
 
-      return null;
+   if (jobDetails.jobError) {
+      return <Breadcrumbs items={[{ text: jobDetails.jobError }]} title={jobDetails.jobError} spinner={false} />               
    }
 
    const crumbItems: BreadcrumbItem[] = [];
diff --git a/Engine/Source/Programs/Horde/HordeDashboard/src/components/jobDetailsV2/JobDetailsViewCommon.tsx b/Engine/Source/Programs/Horde/HordeDashboard/src/components/jobDetailsV2/JobDetailsViewCommon.tsx
--- a/Engine/Source/Programs/Horde/HordeDashboard/src/components/jobDetailsV2/JobDetailsViewCommon.tsx
+++ b/Engine/Source/Programs/Horde/HordeDashboard/src/components/jobDetailsV2/JobDetailsViewCommon.tsx
@@ -543,87 +543,98 @@ export class JobDetailsV2 extends PollBase {
 
    async poll() {
 
-      if (!this.jobId || this.jobError) {
-         return;
-      }
-
-      let requests:any[] = [];
+      try {
 
-      const initialRequest = !this.jobData;
+         if (!this.jobId || this.jobError) {
+            return;
+         }
 
-      requests.push(backend.getJob(this.jobId, undefined, true));
+         let requests: any[] = [];
 
-      let results: any;
+         const initialRequest = !this.jobData;
 
-      await Promise.all(requests as any).then(r => results = r).catch(reason => {
-         console.error(reason);
-         this.jobError = reason;
-      });
+         requests.push(backend.getJob(this.jobId, undefined, true));
 
-      if (!this.jobError && !results?.length) {
-         this.jobError = "Not Found";
-      }
+         let results: any;
 
-      if (this.jobError) {
-         this.setRootUpdated();
-         return;
-      }
+         await Promise.all(requests as any).then(r => results = r).catch(reason => {
+            console.error(reason);
+            this.jobError = reason;
+         });
 
-      const lastUpdateTime = this.jobData?.updateTime;
+         if (!this.jobError && !results?.length) {
+            this.jobError = "Not Found";
+         }
 
-      const jobData = this.jobData = results[0] as JobData;
+         if (this.jobError) {
+            this.setRootUpdated();
+            return;
+         }
 
-      if (!this.stream) {
-         this.stream = projectStore.streamById(jobData.streamId);
-      }
+         const lastUpdateTime = this.jobData?.updateTime;
 
-      if (!this.stream) {
-         throw new Error(`Unable to get stream for job ${this.jobId}`)
-      }
+         const jobData = this.jobData = results[0] as JobData;
 
-      let forceUpdate = false;
-      if (initialRequest) {
-         const templates = await TemplateCache.getStreamTemplates(this.stream);
-         this.template = templates.find(t => t.id === jobData.templateId);
-         if (!this.template) {
-            throw new Error(`Unable to get stream template for job ${this.jobId}`)
+         if (!this.stream) {
+            this.stream = projectStore.streamById(jobData.streamId);
          }
-      } else {
 
-         if (!this.timing && this.jobId) {
-            forceUpdate = true;
-            this.timing = await backend.getJobTiming(this.jobId);
-            this.syncTiming();
+         if (!this.stream) {
+            throw new Error(`Unable to get stream for job ${this.jobId}`)
          }
-      }
 
+         let forceUpdate = false;
+         if (initialRequest) {
+            const templates = await TemplateCache.getStreamTemplates(this.stream);
+            this.template = templates.find(t => t.id === jobData.templateId);
+            if (!this.template) {
+               throw new Error(`Error: Missing template ${this.stream?.fullname ?? jobData.streamId} - ${jobData.templateId}`)
+            }
+         } else {
 
-      if (jobData.updateTime !== lastUpdateTime) {
-         this.processGraph();
-         this.syncTiming();
-         this.setUpdated();
-      }
+            if (!this.timing && this.jobId) {
+               forceUpdate = true;
+               this.timing = await backend.getJobTiming(this.jobId);
+               this.syncTiming();
+            }
+         }
 
-      if (initialRequest) {
-         this.filter.updateFilterItems();
-         this.filterUpdated();
 
-         // poll again for initial request, to save delay on 2nd pass
-         clearTimeout(this.timeoutId);
-         this.timeoutId = setTimeout(() => { this.update(); }, 500);
+         if (jobData.updateTime !== lastUpdateTime) {
+            this.processGraph();
+            this.syncTiming();
+            this.setUpdated();
+         }
 
-      }
+         if (initialRequest) {
+            this.filter.updateFilterItems();
+            this.filterUpdated();
 
-      this.views.forEach(v => {
-         v.tick();
-      });
+            // poll again for initial request, to save delay on 2nd pass
+            clearTimeout(this.timeoutId);
+            this.timeoutId = setTimeout(() => { this.update(); }, 500);
 
-      this.views.forEach(v => {
-         if (forceUpdate || (v.updateTime !== jobData.updateTime)) {
-            v.updateTime = jobData.updateTime;
-            v.detailsUpdated();
          }
-      });
+
+         this.views.forEach(v => {
+            v.tick();
+         });
+
+         this.views.forEach(v => {
+            if (forceUpdate || (v.updateTime !== jobData.updateTime)) {
+               v.updateTime = jobData.updateTime;
+               v.detailsUpdated();
+            }
+         });
+         
+      } catch (reason) {
+         console.error(reason);
+         this.jobError = reason?.message ?? "Error loading job";
+         this.setRootUpdated();
+         
+      } finally {
+
+      }
 
    }
 
