Horde: Tweaks to compact views to support hiding/showing names

diff --git a/Engine/Source/Programs/Horde/HordeDashboard/src/backend/Api.ts b/Engine/Source/Programs/Horde/HordeDashboard/src/backend/Api.ts
--- a/Engine/Source/Programs/Horde/HordeDashboard/src/backend/Api.ts
+++ b/Engine/Source/Programs/Horde/HordeDashboard/src/backend/Api.ts
@@ -2705,6 +2705,8 @@ export type GetStreamTabResponse = {
 	type: TabType;
 
 	style: TabStyle;
+
+	showNames?: boolean;
 };
 
 /**Describes a job page */
diff --git a/Engine/Source/Programs/Horde/HordeDashboard/src/components/JobViewIncremental.tsx b/Engine/Source/Programs/Horde/HordeDashboard/src/components/JobViewIncremental.tsx
--- a/Engine/Source/Programs/Horde/HordeDashboard/src/components/JobViewIncremental.tsx
+++ b/Engine/Source/Programs/Horde/HordeDashboard/src/components/JobViewIncremental.tsx
@@ -45,23 +45,27 @@ type JobGroup = {
 
 const commitCache = new CommitCache();
 
+ 
 const buildColumns = (jobTab: JobsTabData, streamId: string): IColumn[] => {
 
+   const startedByName = jobTab.showNames ? "Started By" : "Change";
+
    const fixedWidths: Record<string, number | undefined> = {
       "Status": 16, // note this doesn't have header text, and need to match the status dot width of 20 after sizing
-      "Time": 60,
-      "Started By": 180,
+      "Time": 60,      
    };
 
+   fixedWidths[startedByName] = jobTab.showNames ? 180 : 32;
+
    const minWidths: Record<string, number | undefined> = {};
 
-   let cnames = ["Status", "Started By"];
+   let cnames = ["Status", startedByName];
 
    if (jobTab.columns) {
 
       let total = 0;
       jobTab.columns.forEach(c => total += c.relativeWidth ?? 1);
-      const w = (940) / total;
+      const w = (jobTab.showNames ? 940 : 940 + 144) / total;
 
       jobTab.columns.forEach(c => { minWidths[c.heading] = w * (c.relativeWidth ?? 1); cnames.push(c.heading); });
    } else {
@@ -220,6 +224,8 @@ const JobList: React.FC<{ tab: string; filter: JobFilterSimple, controller: Call
       return <div />;
    }
 
+   const startedByName = jobTab.showNames ? "Started By" : "Change";
+
    const templateNames: Set<string> = new Set();
    jobTab.templates?.forEach(name => {
       templateNames.add(name);
@@ -499,7 +505,7 @@ const JobList: React.FC<{ tab: string; filter: JobFilterSimple, controller: Call
 
       const commit = commitCache.getCommit(streamId!, item.job!.change!);
 
-      if (column!.key === "Started By") {
+      if (column!.key === startedByName) {
 
          let authorName = item.job.startedByUserInfo?.name;
 
@@ -542,7 +548,7 @@ const JobList: React.FC<{ tab: string; filter: JobFilterSimple, controller: Call
                <Stack verticalAlign="center" verticalFill={true} horizontalAlign="start"> <div style={{ paddingBottom: "1px" }}>
                   <span style={{ padding: "2px 6px 2px 6px", height: "18px", cursor: "pointer", color: "#FFFFFF", backgroundColor: item.job.startedByUserInfo ? "#0288ee" : "#035ca1" }} className={item.job.startedByUserInfo ? "cl-callout-button-user" : "cl-callout-button"} >{change}</span>
                </div></Stack>
-               < Text variant="small">{authorName}</Text>
+               {!!jobTab.showNames &&  <Text variant="small">{authorName}</Text>}
             </Stack>
          </Stack>;
       }
@@ -562,7 +568,7 @@ const JobList: React.FC<{ tab: string; filter: JobFilterSimple, controller: Call
 
    function onRenderItemColumn(item: JobItem, index?: number, column?: IColumn) {
 
-      if (column?.key === "Started By") {
+      if (column?.key === startedByName) {
          return onRenderItemColumnInner(item, index, column);
       }
 
