Horde: Device view fixes
#skipci

diff --git a/Engine/Source/Programs/Horde/HordeDashboard/src/components/DeviceView.tsx b/Engine/Source/Programs/Horde/HordeDashboard/src/components/DeviceView.tsx
--- a/Engine/Source/Programs/Horde/HordeDashboard/src/components/DeviceView.tsx
+++ b/Engine/Source/Programs/Horde/HordeDashboard/src/components/DeviceView.tsx
@@ -153,6 +153,9 @@ class LocalState {
 
    @action
    setPivotKey(pivotKey: string) {
+      this.searchState.filterPlatforms = undefined;
+      this.searchState.filterPools = undefined;
+      this.searchState.filterString = undefined;
       this.searchState.pivotKey = pivotKey;
       this.updateSearch();
    }
@@ -267,8 +270,7 @@ export const SearchUpdate: React.FC = observer(() => {
 const DevicePanel: React.FC = observer(() => {
 
    const [editState, setEditState] = useState<{ shown?: boolean, infoShown?: boolean, editNote?: boolean, device?: DeviceItem | undefined }>({});
-   const [pivotState, setPivotState] = useState<{ key: string, poolFilter: Set<string> }>({ key: localState?.searchState?.pivotKey ?? pivotKeyShared, poolFilter: new Set(localState?.searchState?.filterPools ?? []) });
-   const [platformState, setPlatformState] = useState<Set<string>>(new Set(localState?.searchState?.filterPlatforms ?? []));
+   const [pivotState, setPivotState] = useState<{ key: string, poolFilter: Set<string>, platformFilter: Set<string> }>({ key: localState?.searchState?.pivotKey ?? pivotKeyShared, platformFilter: new Set(localState?.searchState?.filterPlatforms) ?? [], poolFilter: new Set(localState?.searchState?.filterPools ?? []) });
    const [checkoutState, setCheckoutState] = useState<{ checkoutId?: string, checkinId?: string, showConfirm?: "in" | "out" | "error" }>({});
    const [telemState, setTelemState] = useState(false);
    
@@ -606,7 +608,7 @@ const DevicePanel: React.FC = observer(() => {
             return false;
          }
 
-         if (platformState.size && !platformState.has(d.platformId)) {
+         if (pivotState.platformFilter.size && !pivotState.platformFilter.has(d.platformId)) {
             return false;
          }
 
@@ -783,7 +785,8 @@ const DevicePanel: React.FC = observer(() => {
                      linkSize="normal"
                      linkFormat="links"
                      onLinkClick={(item) => {
-                        setPivotState({ key: item!.props.itemKey!, poolFilter: new Set() })
+                        localState.setPivotKey(item!.props.itemKey!);
+                        setPivotState({ key: item!.props.itemKey!, poolFilter: new Set(), platformFilter: new Set() })
                      }}>
                      {pivotItems}
                   </Pivot>
@@ -817,20 +820,20 @@ const DevicePanel: React.FC = observer(() => {
                      placeholder="Filter Platforms"
                      style={{ width: 200 }}
                      styles={dropDownStyle}
-                     selectedKeys={Array.from(platformState).map(id => `dropdown_${id}_key`)}
+                     selectedKeys={Array.from(pivotState.platformFilter).map(id => `dropdown_${id}_key`)}
                      multiSelect
                      options={platformItems}
                      onChange={(event, option, index) => {
 
                         if (option) {
                            if (option.selected) {
-                              platformState.add((option as any).platformId);
+                              pivotState.platformFilter.add((option as any).platformId);
                            } else {
-                              platformState.delete((option as any).platformId);
+                              pivotState.platformFilter.delete((option as any).platformId);
                            }
 
-                           localState.setPlatformFilters(platformState);
-                           setPlatformState(platformState);
+                           localState.setPlatformFilters(pivotState.platformFilter);
+                           setPivotState({...pivotState, platformFilter: new Set(pivotState.platformFilter)})
                         }
                      }}
                   />
