Horde: Update dashboard install docs for vite
#skipci

diff --git a/Engine/Source/Programs/Horde/Docs/Internals/Building.md b/Engine/Source/Programs/Horde/Docs/Internals/Building.md
--- a/Engine/Source/Programs/Horde/Docs/Internals/Building.md
+++ b/Engine/Source/Programs/Horde/Docs/Internals/Building.md
@@ -25,16 +25,15 @@ The Horde dashboard is a frontend client developed in [TypeScript](https://www.t
 [React](https://react.dev/). To set up your machine for developing the dashboard:
 
 1. Install [Node.js](https://nodejs.org/en/download).
-2. From a command line, install Yarn using: `npm install --global yarn`.
-3. Navigate to the dashboard folder at `Engine\Source\Programs\Horde\HordeDashboard`.
-4. Run `yarn install` to install the package dependencies.
-5. Edit package.json setting the proxy property to point at your server URL, for example: `http://localhost:13340`.
-6. Navigate to the admin token endpoint of your server to get an expiring access token, for example:
+2. Navigate to the dashboard folder at `Engine\Source\Programs\Horde\HordeDashboard`.
+3. Run `npm install --legacy-peer-deps` to install the package dependencies.
+4. Edit vite.config.ts setting the proxyTarget variable to point at your server URL, for example: `http://localhost:13340`.
+5. Navigate to the admin token endpoint of your server to get an expiring access token, for example:
    `http://localhost:13340/api/v1/admin/token`.
-7. Create a file called `.env.development.local` in the root HordeDashboard folder and paste the access token in like
-   so: `REACT_APP_HORDE_DEBUG_TOKEN=eyFhbGciziJIUz`.
-8. Run `yarn start` to start the development web server which should open a tab to the local dashboard at
-   `http://localhost:3000`.
+6. Create a file called `.env.local` in the root HordeDashboard folder and paste the access token in like
+   so: `VITE_HORDE_DEBUG_TOKEN=eyFhbGciziJIUz`.
+7. Run `npm run dev` to start the development web server which should open a tab to the local dashboard at
+   `http://localhost:5173`.
 
 ## Docker Image
 
