Horde: Fix undefined access in job vew
#skipci

diff --git a/Engine/Source/Programs/Horde/HordeDashboard/src/components/jobDetailsV2/JobDetailViewV2.tsx b/Engine/Source/Programs/Horde/HordeDashboard/src/components/jobDetailsV2/JobDetailViewV2.tsx
--- a/Engine/Source/Programs/Horde/HordeDashboard/src/components/jobDetailsV2/JobDetailViewV2.tsx
+++ b/Engine/Source/Programs/Horde/HordeDashboard/src/components/jobDetailsV2/JobDetailViewV2.tsx
@@ -52,6 +52,8 @@ const JobBreadCrumbs: React.FC<{ jobDetails: JobDetailsV2 }> = observer(({ jobDe
    if (!jobData) {
       if (!jobDetails.jobError) {
          return <Breadcrumbs items={[{ text: "Loading Job" }]} title={"Loading Job"} spinner={true} />
+      } else {
+         return <Breadcrumbs items={[{ text: jobDetails.jobError }]} title={jobDetails.jobError} spinner={false} />               
       }
    } 
 
