Add Device Spike and Load report

#rnx

diff --git a/Engine/Source/Programs/Horde/Plugins/Build/HordeServer.Build.Tests/Devices/DeviceServiceTests.cs b/Engine/Source/Programs/Horde/Plugins/Build/HordeServer.Build.Tests/Devices/DeviceServiceTests.cs
--- a/Engine/Source/Programs/Horde/Plugins/Build/HordeServer.Build.Tests/Devices/DeviceServiceTests.cs
+++ b/Engine/Source/Programs/Horde/Plugins/Build/HordeServer.Build.Tests/Devices/DeviceServiceTests.cs
@@ -16,6 +16,7 @@ using HordeServer.Users;
 using HordeServer.Utilities;
 using Microsoft.AspNetCore.Http;
 using Microsoft.AspNetCore.Mvc;
+using MongoDB.Bson;
 using Moq;
 
 namespace HordeServer.Tests.Devices
@@ -687,5 +688,102 @@ namespace HordeServer.Tests.Devices
 				}
 			}
 		}
+		class TestDeviceTelemetry : IDeviceTelemetry
+		{
+			public DateTime CreateTimeUtc { get; init; }
+			public DeviceId DeviceId { get; init; }
+			public string? StreamId { get; init; }
+			public string? JobId { get; init; }
+			public string? JobName { get; init; }
+			public string? StepId { get; init; }
+			public string? StepName { get; init; }
+			public ObjectId? ReservationId { get; init; }
+			public DateTime? ReservationStartUtc { get; init; }
+			public DateTime? ReservationFinishUtc { get; init; }
+			public DateTime? ProblemTimeUtc { get; init; }
+		}
+
+		static void AddLoadSpikes(List<IDeviceTelemetry> telemetry, List<IDevice> devices, DevicePoolId pool, DevicePlatformId platform, ref DateTime referenceTime, int nSpikes, TimeSpan duration)
+		{
+			IEnumerable<DeviceId> selectedDevices = devices.Where(d => d.PoolId == pool && d.PlatformId == platform).Select(d => d.Id);
+			for (int i = 0; i < nSpikes; i++)
+			{
+				foreach (DeviceId id in selectedDevices)
+				{
+					telemetry.Add(
+						new TestDeviceTelemetry()
+						{
+							CreateTimeUtc = referenceTime,
+							DeviceId = id,
+							ReservationFinishUtc = referenceTime + duration
+						}
+					);
+					referenceTime = referenceTime.AddSeconds(5); // shift timing slightly every time
+				}
+				referenceTime = referenceTime.Add(duration + TimeSpan.FromMinutes(5)); // shift timing 5 minutes between spikes
+			}
+		}
+
+		static void AddProblemSpikes(List<IDeviceTelemetry> telemetry, List<IDevice> devices, DevicePoolId pool, DevicePlatformId platform, ref DateTime referenceTime, int percSpikes, TimeSpan offset)
+		{
+			IEnumerable<DeviceId> selectedDevices = devices.Where(d => d.PoolId == pool && d.PlatformId == platform).Select(d => d.Id);
+			int devicesCount = selectedDevices.Count();
+			devicesCount = (int)((double)devicesCount * percSpikes / 100);
+			foreach (DeviceId id in selectedDevices.Take(devicesCount))
+			{
+				telemetry.Add(
+					new TestDeviceTelemetry()
+					{
+						CreateTimeUtc = referenceTime,
+						DeviceId = id,
+						ReservationFinishUtc = referenceTime + offset,
+						ProblemTimeUtc = referenceTime + offset
+					}
+				);
+				referenceTime = referenceTime.AddSeconds(5); // shift timing slightly every time
+			}
+		}
+
+		[TestMethod]
+		public async Task TestPopulateDeviceIssueReportAsync()
+		{
+			await SetupDevicesAsync();
+			List<IDevicePool> pools = DeviceService!.GetPools();
+			List<IDevicePlatform> platforms = DeviceService!.GetPlatforms();
+			List<IDevice> devices = (await DeviceService!.GetDevicesAsync()).OrderBy(d => d.PoolId.ToString()).ThenBy(d => d.PlatformId.ToString()).ToList();
+			List<IDeviceTelemetry>  telemetry = new List<IDeviceTelemetry>();
+
+			DateTime referencePoint = DateTime.Now - TimeSpan.FromDays(1);
+			AddLoadSpikes(telemetry, devices, pools[0].Id, platforms[0].Id, ref referencePoint, 3, TimeSpan.FromMinutes(3));
+			referencePoint = referencePoint.AddHours(1);
+			AddProblemSpikes(telemetry, devices, pools[0].Id, platforms[0].Id, ref referencePoint, 40, TimeSpan.FromMinutes(3));
+
+			DeviceReportData deviceData = new DeviceReportData()
+			{
+				Pools = pools,
+				Platforms = platforms,
+				Devices = devices,
+				DeviceTelemetry = telemetry
+			};
+			DeviceReportSettings settings = new DeviceReportSettings()
+			{
+				Channel = "a-channel",
+				ReportMinutes = 3 * 60,
+				DeviceProblemCoolDownMinutes = 5,
+				DeviceSaturationSpikeThresholdMinutes = 1,
+				ServerUrl = new Uri("https://testserver")
+			};
+
+			DeviceIssueReport issueReport = DeviceReportService.CreateDeviceReport(deviceData, settings);
+
+			Assert.AreEqual(issueReport.PoolReports.Count, 2);
+			Assert.AreEqual(issueReport.PoolReports[0].Metrics.Count, 3);
+			Assert.AreEqual(issueReport.PoolReports[0].Metrics[0].AverageLoadPercentage, 6);
+			Assert.AreEqual(issueReport.PoolReports[0].Metrics[0].SaturationSpikes, 3);
+			Assert.AreEqual(issueReport.PoolReports[0].Metrics[0].SpikeDurationAverage.ToString(@"hh\:mm"), "00:02");
+			Assert.AreEqual(issueReport.PoolReports[0].Metrics[0].SpikeDurationPercentage, 4);
+			Assert.AreEqual(issueReport.PoolReports[0].Metrics[0].MaxConcurrentProblemsPercentage, 37);
+			Assert.AreEqual(issueReport.PoolReports[0].Metrics[0].MaxConcurrentProblems, 3);
+		}
 	}
 }
diff --git a/Engine/Source/Programs/Horde/Plugins/Build/HordeServer.Build/Devices/DeviceReportService.cs b/Engine/Source/Programs/Horde/Plugins/Build/HordeServer.Build/Devices/DeviceReportService.cs
--- a/Engine/Source/Programs/Horde/Plugins/Build/HordeServer.Build/Devices/DeviceReportService.cs
+++ b/Engine/Source/Programs/Horde/Plugins/Build/HordeServer.Build/Devices/DeviceReportService.cs
@@ -70,14 +70,95 @@ namespace HordeServer.Devices
 		}
 	}
 
+	public class DeviceUsageEvent
+	{
+		/// <summary>
+		/// Start of the usage event.
+		/// </summary>
+		public DateTime Start { get; }
+		/// <summary>
+		/// Increment(positive value) or decrement(negative value) usage, usually by one. 
+		/// </summary>
+		public int Increment { get; }
+		/// <summary>
+		/// Whether a problem was reported during that event.
+		/// </summary>
+		public bool HasProblem { get; }
+
+		public DeviceUsageEvent(DateTime start, int increment, bool hasProblem = false)
+		{
+			Start = start;
+			Increment = increment;
+			HasProblem = hasProblem;
+		}
+	}
+
 	public class DevicePoolMetrics
 	{
+		/// <summary>
+		/// Platform Id
+		/// </summary>
 		public string PlatformId { get; }
+		/// <summary>
+		/// Platform Name
+		/// </summary>
 		public string PlatformName { get; }
+		/// <summary>
+		/// Number total of devices in the pool
+		/// </summary>
 		public int Total { get; set; } = 0;
+		/// <summary>
+		/// Number of disabled devices in the pool
+		/// </summary>
 		public int Disabled { get; set; } = 0;
+		/// <summary>
+		/// Number of devices marked as Maintenance in the pool
+		/// </summary>
 		public int Maintenance { get; set; } = 0;
+		/// <summary>
+		/// Total number of reported problems
+		/// </summary>
 		public int Problems { get; set; } = 0;
+		/// <summary>
+		/// Cumulated reservation duration in the pool
+		/// </summary>
+		public TimeSpan TotalReservationDuration { get; set; } = new TimeSpan();
+		/// <summary>
+		/// Percentage of the average load per device in the pool
+		/// </summary>
+		public int AverageLoadPercentage { get; set; } = 0;
+		/// <summary>
+		/// Total number of pool saturation spikes
+		/// </summary>
+		public int SaturationSpikes { get; set; } = 0;
+		/// <summary>
+		/// Cumulated duration of the saturation spikes
+		/// </summary>
+		public TimeSpan SpikesDuration { get; set; } = new TimeSpan();
+		/// <summary>
+		/// Average duration a saturation spike
+		/// </summary>
+		public TimeSpan SpikeDurationAverage { get; set; } = new TimeSpan();
+		/// <summary>
+		/// Percentage of the cumulated spikes duration against the total reservation duration
+		/// </summary>
+		public int SpikeDurationPercentage { get; set; } = 0;
+		/// <summary>
+		/// Maximum concurrent reported problems at any one time
+		/// </summary>
+		public int MaxConcurrentProblems { get; set; } = 0;
+		/// <summary>
+		/// Percentage of the maximum concurrent reported problems against the total number of available devices
+		/// </summary>
+		public int MaxConcurrentProblemsPercentage { get; set; } = 0;
+		/// <summary>
+		/// Usage time series
+		/// </summary>
+		public List<DeviceUsageEvent> UsageEvents { get; } = new List<DeviceUsageEvent>();
+		/// <summary>
+		/// Collection of streams from which device reservation where made
+		/// </summary>
+		public HashSet<string> Streams { get; } = new HashSet<string>();
 
 		public DevicePoolMetrics(string platformId, string platformName)
 		{
@@ -114,6 +195,24 @@ namespace HordeServer.Devices
 		}
 	}
 
+	public struct DeviceReportData
+	{
+		public List<IDevicePool> Pools { get; init; }
+		public List<IDevicePlatform> Platforms { get; init; }
+		public List<IDevice> Devices { get; init; }
+		public List<IDeviceTelemetry> DeviceTelemetry { get; init; }
+	}
+
+	public struct DeviceReportSettings
+	{
+		public string Channel { get; init; }
+		public int ReportMinutes { get; init; }
+		public int DeviceProblemCoolDownMinutes { get; init; }
+		public int DeviceSaturationSpikeThresholdMinutes { get; init; }
+		public Uri ServerUrl { get; init; }
+	}
+
+
 #pragma warning restore CS1591 // Missing XML comment for publicly visible type or member
 
 	[SingletonDocument("device-report-state", "6491bc0ce39f8b70b1ef2feb")]
@@ -137,7 +236,10 @@ namespace HordeServer.Devices
 		readonly IOptions<BuildServerConfig> _staticBuildConfig;
 		readonly ILogger<DeviceReportService> _logger;
 
-		readonly int _reportIntervalMinutes = 180;
+		readonly int _reportIntervalMinutes = 12 * 60; // every 12 hours
+
+		readonly int _deviceProblemCoolDownMinutes = 10;
+		readonly int _deviceSaturationSpikeThresholdMinutes = 10;
 
 		/// <summary>
 		/// Constructor
@@ -188,17 +290,51 @@ namespace HordeServer.Devices
 			List<IDevicePool> pools = _deviceService.GetPools();
 			List<IDevicePlatform> platforms = _deviceService.GetPlatforms();
 			List<IDevice> devices = await _deviceService.GetDevicesAsync();
+			List<IDeviceTelemetry> deviceTelemetry = await _deviceService.GetDeviceTelemetryAsync(null, lastReportTime);
+
+			DeviceReportData deviceData = new DeviceReportData()
+			{
+				Pools = pools,
+				Platforms = platforms,
+				Devices = devices,
+				DeviceTelemetry = deviceTelemetry
+			};
+			DeviceReportSettings settings = new DeviceReportSettings()
+			{
+				Channel = staticBuildConfig.DeviceReportChannel,
+				ReportMinutes = _reportIntervalMinutes,
+				DeviceProblemCoolDownMinutes = _deviceProblemCoolDownMinutes,
+				DeviceSaturationSpikeThresholdMinutes = _deviceSaturationSpikeThresholdMinutes,
+				ServerUrl = _serverInfo.DashboardUrl
+			};
 
-			devices = devices.OrderBy(d => d.PlatformId.ToString()).ThenBy(d => d.PoolId.ToString()).ToList();
+			DeviceIssueReport issueReport = CreateDeviceReport(deviceData, settings);
 
-			List<IDeviceTelemetry> deviceTelemetry = await _deviceService.GetDeviceTelemetryAsync(null, lastReportTime);
+			if (issueReport.PoolReports.Count > 0 || issueReport.PlatformReports.Count > 0)
+			{
+				try
+				{
+					_logger.LogInformation("Sending device report notification");
+
+					await _notificationService.SendDeviceIssueReportAsync(issueReport, cancellationToken);
+				}
+				catch (Exception ex)
+				{
+					_logger.LogError(ex, "Error sending device health notification");
+				}
+			}
 
-			DeviceIssueReport issueReport = new DeviceIssueReport(staticBuildConfig.DeviceReportChannel);
+			state = await _state.UpdateAsync(s => s.ReportTime = currentTime, cancellationToken);
+		}
 
-			devices.ForEach(device =>
+		static internal DeviceIssueReport CreateDeviceReport(DeviceReportData deviceData, DeviceReportSettings settings)
+		{
+			DeviceIssueReport issueReport = new DeviceIssueReport(settings.Channel);
+			deviceData.Devices.OrderBy(d => d.PlatformId.ToString()).ThenBy(d => d.PoolId.ToString()).ToList()
+				.ForEach(device =>
 			{
-				IDevicePool? pool = pools.Find(p => p.Id == device.PoolId);
-				IDevicePlatform? platform = platforms.Find(p => p.Id == device.PlatformId);
+				IDevicePool? pool = deviceData.Pools.Find(p => p.Id == device.PoolId);
+				IDevicePlatform? platform = deviceData.Platforms.Find(p => p.Id == device.PlatformId);
 
 				if (pool == null || platform == null)
 				{
@@ -239,10 +375,43 @@ namespace HordeServer.Devices
 					return;
 				}
 
-				// Report device issues, though only if not since disabled or put into maintenance
+				// Track usage events
+				List<DeviceUsageEvent> platformUsage = metrics.UsageEvents;
+				DateTime? endTimeUsage = null;
+				List<IDeviceTelemetry> telemetry = deviceData.DeviceTelemetry.Where(t => t.DeviceId == device.Id).OrderBy(t => t.CreateTimeUtc).ToList();
+				foreach (IDeviceTelemetry item in telemetry)
+				{
+					// Correct superposed reservations
+					if (endTimeUsage != null && endTimeUsage > item.CreateTimeUtc)
+					{
+						int lastIndex = platformUsage.Count - 1;
+						platformUsage[lastIndex] = new DeviceUsageEvent(item.CreateTimeUtc, -1, platformUsage[lastIndex].HasProblem);
+						metrics.TotalReservationDuration -= (DateTime)endTimeUsage - platformUsage[lastIndex].Start;
+					}
+
+					// Add incremental usage
+					bool hasProblem = item.ProblemTimeUtc != null;
+					DeviceUsageEvent usage = new DeviceUsageEvent(item.CreateTimeUtc, 1, hasProblem);
+					platformUsage.Add(usage);
+
+					// Add decremental usage
+					endTimeUsage = !hasProblem? item.ReservationFinishUtc : item.ProblemTimeUtc?.AddMinutes(settings.DeviceProblemCoolDownMinutes);
+					if (endTimeUsage != null)
+					{
+						DeviceUsageEvent endUsage = new DeviceUsageEvent((DateTime)endTimeUsage, -1, hasProblem);
+						platformUsage.Add(endUsage);
+						metrics.TotalReservationDuration += endUsage.Start - usage.Start;
+					}
+
+					// Track stream usage
+					if (item.StreamId != null)
+					{
+						metrics.Streams.Add(item.StreamId);
+					}
+				}
 
-				List<IDeviceTelemetry> telemetry = deviceTelemetry.Where(t => t.DeviceId == device.Id).ToList();
-				List<IDeviceTelemetry> problems = telemetry.Where(t => t.ProblemTimeUtc != null && t.ProblemTimeUtc! > lastReportTime).OrderByDescending(t => t.ProblemTimeUtc!).ToList();
+				// Report device issues, though only if not since disabled or put into maintenance
+				List<IDeviceTelemetry> problems = telemetry.Where(t => t.ProblemTimeUtc != null).OrderByDescending(t => t.ProblemTimeUtc!).ToList();
 
 				if (problems.Count == 0)
 				{
@@ -261,7 +430,7 @@ namespace HordeServer.Devices
 				if (lastProblem.JobName != null && lastProblem.StepName != null && lastProblem.JobId != null && lastProblem.StepId != null)
 				{
 					deviceReport.LastProblemDesc = $"{lastProblem.JobName} - {lastProblem.StepName}";
-					deviceReport.LastProblemURL = new Uri($"{_serverInfo.DashboardUrl}job/{lastProblem.JobId}?step={lastProblem.StepId}").ToString();
+					deviceReport.LastProblemURL = new Uri($"{settings.ServerUrl}job/{lastProblem.JobId}?step={lastProblem.StepId}").ToString();
 				}
 
 				DevicePlatformReport? platformReport = issueReport.PlatformReports.Find(p => p.PlatformId == device.PlatformId.ToString());
@@ -274,27 +443,63 @@ namespace HordeServer.Devices
 				platformReport.DeviceReports.Add(deviceReport);
 			});
 
+			// Calculate usage and spikes
+			issueReport.PoolReports.ForEach(report =>
+			{
+				report.Metrics.ForEach(metrics =>
+				{
+					int maxCount = metrics.Total - metrics.Disabled - metrics.Maintenance;
+					if (maxCount > 0)
+					{
+						metrics.UsageEvents.SortBy(usage => usage.Start);
+						int usageCount = 0;
+						int concurrentProblemsCount = 0;
+						int index = 0;
+						int lastIndex = metrics.UsageEvents.Count - 1;
+						foreach (DeviceUsageEvent usage in metrics.UsageEvents)
+						{
+							usageCount += usage.Increment;
+							if (usageCount == maxCount)
+							{
+								if (index < lastIndex)
+								{
+									TimeSpan duration = metrics.UsageEvents[index + 1].Start - usage.Start;
+									if (duration.TotalMinutes > settings.DeviceSaturationSpikeThresholdMinutes)
+									{
+										metrics.SaturationSpikes += 1;
+										metrics.SpikesDuration += duration;
+									}
+								}
+								else
+								{
+									// Can't get the duration, but we still want to count the spike.
+									metrics.SaturationSpikes += 1;
+								}
+							}
+							if (usage.HasProblem)
+							{
+								concurrentProblemsCount += usage.Increment;
+								metrics.MaxConcurrentProblems = metrics.MaxConcurrentProblems > concurrentProblemsCount ? metrics.MaxConcurrentProblems : concurrentProblemsCount;
+							}
+							index++;
+						}
+						metrics.AverageLoadPercentage = (int)((((double)metrics.TotalReservationDuration.TotalMinutes / maxCount) / settings.ReportMinutes) * 100);
+						metrics.MaxConcurrentProblemsPercentage = (int)(((double)metrics.MaxConcurrentProblems / maxCount) * 100);
+						if (metrics.SaturationSpikes > 0)
+						{
+							metrics.SpikeDurationAverage = metrics.SpikesDuration / metrics.SaturationSpikes;
+							metrics.SpikeDurationPercentage = (int)((metrics.SpikesDuration.TotalMinutes / settings.ReportMinutes) * 100);
+						}
+					}
+				});
+			});
+
 			// order everything nicely
 			issueReport.PoolReports = issueReport.PoolReports.OrderBy(r => r.PoolName).ToList();
 			issueReport.PlatformReports = issueReport.PlatformReports.OrderBy(r => r.PlatformName).ToList();
 			issueReport.PlatformReports.ForEach(r => r.DeviceReports = r.DeviceReports.OrderBy(r => r.PoolName).ThenBy(r => r.DeviceName).ToList());
 
-			if (issueReport.PoolReports.Count > 0 || issueReport.PlatformReports.Count > 0)
-			{
-				try
-				{
-					_logger.LogInformation("Sending device report notification");
-
-					await _notificationService.SendDeviceIssueReportAsync(issueReport, cancellationToken);
-				}
-				catch (Exception ex)
-				{
-					_logger.LogError(ex, "Error sending device health notification");
-				}
-			}
-
-			state = await _state.UpdateAsync(s => s.ReportTime = currentTime, cancellationToken);
-
+			return issueReport;
 		}
 	}
 }
diff --git a/Engine/Source/Programs/Horde/Plugins/Build/HordeServer.Build/Notifications/Sinks/SlackNotificationSink.cs b/Engine/Source/Programs/Horde/Plugins/Build/HordeServer.Build/Notifications/Sinks/SlackNotificationSink.cs
--- a/Engine/Source/Programs/Horde/Plugins/Build/HordeServer.Build/Notifications/Sinks/SlackNotificationSink.cs
+++ b/Engine/Source/Programs/Horde/Plugins/Build/HordeServer.Build/Notifications/Sinks/SlackNotificationSink.cs
@@ -2410,40 +2410,56 @@ namespace HordeServer.Notifications.Sinks
 				SlackMessage headerMessage = new SlackMessage();
 				headerMessage.AddHeader($"Device Pool Health Summary");
 				await SendMessageAsync(report.Channel, headerMessage, cancellationToken);
+				int count = 0;
 
 				foreach (DevicePoolReport pool in report.PoolReports)
 				{
 					List<SlackAttachment> attachments = new List<SlackAttachment>();
 					foreach (DevicePoolMetrics metrics in pool.Metrics)
 					{
-
-						double totalPct = (((double)metrics.Disabled + (double)metrics.Maintenance + (double)metrics.Problems) / (double)metrics.Total) * 100.0;
-
 						StringBuilder builder = new StringBuilder("   ");
 
-						string color = "#eeeeee";
+						string color = "#eff700"; // yellow
 
-						if (totalPct >= 40)
+						if (metrics.AverageLoadPercentage >= 40 || metrics.MaxConcurrentProblemsPercentage >= 50)
 						{
-							color = "#ff0000";
+							color = "#d60404"; // red
 						}
-						else if (totalPct >= 20)
+						else if (metrics.AverageLoadPercentage >= 20 || metrics.MaxConcurrentProblemsPercentage >= 30)
 						{
-							color = "#ffff00";
+							color = "#ffaa00"; // orange
 						}
-						else
+						else if (metrics.Problems == 0 && metrics.SaturationSpikes == 0 && metrics.Total > (metrics.Disabled + metrics.Maintenance))
 						{
 							continue;
 						}
 
+						count++;
+
 						builder.Append($"*{metrics.PlatformName}* - ");
+						builder.Append($"Average Load: {metrics.AverageLoadPercentage}%| ");
+
+						if (metrics.SaturationSpikes > 0)
+						{
+							builder.Append($"Saturation Spikes: {metrics.SaturationSpikes}| ");
+							builder.Append($"Average Spike Duration: {metrics.SpikeDurationAverage.ToString(@"hh\:mm")}({metrics.SpikeDurationPercentage}%)| ");
+						}
 
 						if (metrics.Problems > 0)
 						{
-							builder.Append($"Problems: {metrics.Problems}, ");
+							builder.Append($"Max Concurrent Problems: {metrics.MaxConcurrentProblems}({metrics.MaxConcurrentProblemsPercentage}%)| ");
 						}
 
-						builder.Append($"Devices: {metrics.Total}, Disabled: {metrics.Disabled}, Maintenance: {metrics.Maintenance}");
+						builder.Append($"Devices: {metrics.Total}| ");
+						if (metrics.Disabled > 0)
+						{
+							builder.Append($"Disabled: {metrics.Disabled}| ");
+						}
+						if (metrics.Maintenance > 0)
+						{
+							builder.Append($"Maintenance: {metrics.Maintenance}| ");
+						}
+						builder.Append($"Streams: {metrics.Streams.Count}");
 
 						SlackAttachment attachment = new SlackAttachment();
 						attachment.FallbackText = builder.ToString();
@@ -2461,10 +2477,16 @@ namespace HordeServer.Notifications.Sinks
 					}
 				}
 
+				if (count == 0)
+				{
+					SlackMessage message = new SlackMessage();
+					message.Text = "no outstanding usage to report";
+					await SendMessageAsync(report.Channel, message, cancellationToken);
+				}
+
 				SlackMessage divider = new SlackMessage();
 				divider.AddDivider();
 				await SendMessageAsync(report.Channel, divider, cancellationToken);
-
 			}
 
 			if (report.PlatformReports.Count > 0)
@@ -2475,29 +2497,27 @@ namespace HordeServer.Notifications.Sinks
 
 				foreach (DevicePlatformReport platform in report.PlatformReports)
 				{
-					SlackMessage message = new SlackMessage();
-					message.Markdown = true;
-					message.Text = $"*{platform.PlatformName}*";
-
-					SlackAttachment attachment = new SlackAttachment();
-
 					foreach (DeviceReport device in platform.DeviceReports.OrderByDescending(r => r.ProblemDelta))
 					{
-						StringBuilder builder = new StringBuilder($"   {device.DeviceName} / {device.DeviceAddress}");
-
-						builder.Append($" / Problems: {device.ProblemDelta}");
+						StringBuilder builder = new StringBuilder($"   {device.ProblemDelta} Problem(s)");
 
 						if (!String.IsNullOrEmpty(device.LastProblemURL))
 						{
 							builder.Append($" - <{device.LastProblemURL}|{device.LastProblemDesc ?? "??"}>");
 						}
 
+						SlackMessage message = new SlackMessage();
+						message.Markdown = true;
+						message.Text = $"*{device.DeviceName} | {device.DeviceAddress} - {platform.PlatformName} - {device.PoolName}*";
+
+						SlackAttachment attachment = new SlackAttachment();
 						attachment.FallbackText = builder.ToString();
+						attachment.Color = "#d60404";
 						attachment.AddSection(builder.ToString());
-					}
 
-					message.Attachments.Add(attachment);
-					await SendMessageAsync(report.Channel, message, cancellationToken);
+						message.Attachments.Add(attachment);
+						await SendMessageAsync(report.Channel, message, cancellationToken);
+					}
 				}
 
 				SlackMessage divider = new SlackMessage();
