Add pool links to device service report

#rnx

diff --git a/Engine/Source/Programs/Horde/Plugins/Build/HordeServer.Build/Devices/DeviceReportService.cs b/Engine/Source/Programs/Horde/Plugins/Build/HordeServer.Build/Devices/DeviceReportService.cs
--- a/Engine/Source/Programs/Horde/Plugins/Build/HordeServer.Build/Devices/DeviceReportService.cs
+++ b/Engine/Source/Programs/Horde/Plugins/Build/HordeServer.Build/Devices/DeviceReportService.cs
@@ -42,6 +42,11 @@ namespace HordeServer.Devices
 		/// </summary>
 		public string? LastProblemURL { get; set; }
 
+		/// <summary>
+		/// The URL to the device on the pool view
+		/// </summary>
+		public string? DevicePoolURL { get; set; }
+
 		public List<IDeviceTelemetry> Telemetry { get; }
 
 		public DeviceReport(string platformId, string platformName, string deviceId, string deviceName, string deviceAddress, string poolId, string poolName, List<IDeviceTelemetry> telemetry)
@@ -104,6 +109,10 @@ namespace HordeServer.Devices
 		/// </summary>
 		public string PlatformName { get; }
 		/// <summary>
+		/// The URL to the platform on the pool view
+		/// </summary>
+		public string? PlatformURL { get; set; }
+		/// <summary>
 		/// Number total of devices in the pool
 		/// </summary>
 		public int Total { get; set; } = 0;
@@ -174,6 +183,11 @@ namespace HordeServer.Devices
 
 		public List<DevicePoolMetrics> Metrics { get; } = new List<DevicePoolMetrics>();
 
+		/// <summary>
+		/// The URL to the whole pool view
+		/// </summary>
+		public string? PoolURL { get; set; }
+
 		public DevicePoolReport(string poolId, string poolName)
 		{
 			PoolId = poolId;
@@ -350,6 +364,7 @@ namespace HordeServer.Devices
 				if (poolReport == null)
 				{
 					poolReport = new DevicePoolReport(pool.Id.ToString(), pool.Name.ToString());
+					poolReport.PoolURL = new Uri($"{settings.ServerUrl}devices?pivotKey=pivot-key-automation&pool={pool.Id}").ToString();
 					issueReport.PoolReports.Add(poolReport);
 				}
 
@@ -357,6 +372,7 @@ namespace HordeServer.Devices
 				if (metrics == null)
 				{
 					metrics = new DevicePoolMetrics(platform.Id.ToString(), platform.Name.ToString());
+					metrics.PlatformURL = new Uri($"{settings.ServerUrl}devices?pivotKey=pivot-key-automation&pool={pool.Id}&platform={platform.Id}").ToString();
 					poolReport.Metrics.Add(metrics);
 				}
 
@@ -433,6 +449,8 @@ namespace HordeServer.Devices
 					deviceReport.LastProblemURL = new Uri($"{settings.ServerUrl}job/{lastProblem.JobId}?step={lastProblem.StepId}").ToString();
 				}
 
+				deviceReport.DevicePoolURL = new Uri($"{settings.ServerUrl}devices?pivotKey=pivot-key-automation&pool={pool.Id}&platform={platform.Id}&filter={device.Name}").ToString();
+
 				DevicePlatformReport? platformReport = issueReport.PlatformReports.Find(p => p.PlatformId == device.PlatformId.ToString());
 				if (platformReport == null)
 				{
diff --git a/Engine/Source/Programs/Horde/Plugins/Build/HordeServer.Build/Notifications/Sinks/SlackNotificationSink.cs b/Engine/Source/Programs/Horde/Plugins/Build/HordeServer.Build/Notifications/Sinks/SlackNotificationSink.cs
--- a/Engine/Source/Programs/Horde/Plugins/Build/HordeServer.Build/Notifications/Sinks/SlackNotificationSink.cs
+++ b/Engine/Source/Programs/Horde/Plugins/Build/HordeServer.Build/Notifications/Sinks/SlackNotificationSink.cs
@@ -2436,7 +2436,7 @@ namespace HordeServer.Notifications.Sinks
 
 						count++;
 
-						builder.Append($"*{metrics.PlatformName}* - ");
+						builder.Append($"<{metrics.PlatformURL}|*{metrics.PlatformName}*> - ");
 						builder.Append($"Average Load: {metrics.AverageLoadPercentage}%| ");
 
 						if (metrics.SaturationSpikes > 0)
@@ -2472,7 +2472,7 @@ namespace HordeServer.Notifications.Sinks
 						SlackMessage message = new SlackMessage();
 						message.Attachments.AddRange(attachments);
 						message.Markdown = true;
-						message.Text = $"*{pool.PoolName}*";
+						message.Text = $"<{pool.PoolURL}|*{pool.PoolName}*>";
 						await SendMessageAsync(report.Channel, message, cancellationToken);
 					}
 				}
@@ -2508,7 +2508,7 @@ namespace HordeServer.Notifications.Sinks
 
 						SlackMessage message = new SlackMessage();
 						message.Markdown = true;
-						message.Text = $"*{device.DeviceName} | {device.DeviceAddress} - {platform.PlatformName} - {device.PoolName}*";
+						message.Text = $"<{device.DevicePoolURL}|*{device.DeviceName}*>* | {device.DeviceAddress} - {platform.PlatformName} - {device.PoolName}*";
 
 						SlackAttachment attachment = new SlackAttachment();
 						attachment.FallbackText = builder.ToString();
