Horde: Fallback to normal change number for BuildGraph invocation if code change is missing

diff --git a/Engine/Source/Programs/Horde/Drivers/JobDriver.Tests/Execution/WorkspaceExecutorTest.cs b/Engine/Source/Programs/Horde/Drivers/JobDriver.Tests/Execution/WorkspaceExecutorTest.cs
--- a/Engine/Source/Programs/Horde/Drivers/JobDriver.Tests/Execution/WorkspaceExecutorTest.cs
+++ b/Engine/Source/Programs/Horde/Drivers/JobDriver.Tests/Execution/WorkspaceExecutorTest.cs
@@ -136,6 +136,7 @@ namespace JobDriver.Tests.Execution
 		private readonly JobId _jobId = JobId.Parse("65bd0655591b5d5d7d047b58");
 		private readonly JobStepBatchId _batchId = new JobStepBatchId(0x1234);
 		private const string AgentType = "BogusAgentType";
+		private const int ChangeNum = 100;
 
 		private readonly ILogger<WorkspaceExecutorTest> _logger;
 		private readonly FakeHordeClient _hordeClient;
@@ -157,11 +158,11 @@ namespace JobDriver.Tests.Execution
 			_hordeClient = new FakeHordeClient();
 			_hordeClient.AddStream(_streamId, "//Foo/Main");
 			_hordeClient.AddAgentType(_streamId, AgentType);
-			_hordeClient.AddJob(_jobId, _streamId, 1, 0);
+			_hordeClient.AddJob(_jobId, _streamId, ChangeNum, 0);
 
-			_autoSdkWorkspace.SetFile(1, "HostWin64/Android/base.h", "base");
-			_workspace.SetFile(1, "main.cpp", "main");
-			_workspace.SetFile(1, "foo/bar/baz.h", "baz");
+			_autoSdkWorkspace.SetFile(ChangeNum, "HostWin64/Android/base.h", "base");
+			_workspace.SetFile(ChangeNum, "main.cpp", "main");
+			_workspace.SetFile(ChangeNum, "foo/bar/baz.h", "baz");
 
 			JobExecutorOptions executorOptions = new JobExecutorOptions(_hordeClient, _workingDir, null, _jobId, _batchId, default, new RpcJobOptions());
 			_executor = new(executorOptions, _workspace, null, SimpleTestExecutor.NoOpTracer, NullLogger.Instance);
@@ -177,7 +178,7 @@ namespace JobDriver.Tests.Execution
 		[TestMethod]
 		public async Task RegularWorkspaceAsync()
 		{
-			RpcBeginBatchResponse batch = new RpcBeginBatchResponse { Change = 1 };
+			RpcBeginBatchResponse batch = new RpcBeginBatchResponse { Change = ChangeNum };
 			await _executor.InitializeAsync(batch, _logger, CancellationToken.None);
 			AssertWorkspaceFile(_workspace, "main.cpp", "main");
 			AssertWorkspaceFile(_workspace, "foo/bar/baz.h", "baz");
@@ -190,7 +191,7 @@ namespace JobDriver.Tests.Execution
 			JobExecutorOptions executorOptions = new JobExecutorOptions(_hordeClient, _workingDir, Array.Empty<ProcessToTerminate>(), _jobId, _batchId, default, new RpcJobOptions());
 			using WorkspaceExecutor executor = new(executorOptions, _workspace, _autoSdkWorkspace, SimpleTestExecutor.NoOpTracer, NullLogger.Instance);
 
-			RpcBeginBatchResponse batch = new RpcBeginBatchResponse { Change = 1 };
+			RpcBeginBatchResponse batch = new RpcBeginBatchResponse { Change = ChangeNum };
 			await executor.InitializeAsync(batch, _logger, CancellationToken.None);
 			AssertWorkspaceFile(_autoSdkWorkspace, "HostWin64/Android/base.h", "base");
 			AssertWorkspaceFile(_workspace, "main.cpp", "main");
@@ -201,10 +202,8 @@ namespace JobDriver.Tests.Execution
 		[TestMethod]
 		public async Task EnvVarsAsync()
 		{
-			JobExecutorOptions executorOptions = new JobExecutorOptions(_hordeClient, _workingDir, Array.Empty<ProcessToTerminate>(), _jobId, _batchId, default, new RpcJobOptions());
-			using WorkspaceExecutor executor = new(executorOptions, _workspace, _autoSdkWorkspace, SimpleTestExecutor.NoOpTracer, NullLogger.Instance);
-
-			RpcBeginBatchResponse batch = new RpcBeginBatchResponse { Change = 1, StreamName = "//UE5/Main" };
+			using WorkspaceExecutor executor = Create();
+			RpcBeginBatchResponse batch = new () { Change = ChangeNum, CodeChange = 200, StreamName = "//UE5/Main" };
 			await executor.InitializeAsync(batch, _logger, CancellationToken.None);
 
 			IReadOnlyDictionary<string, string> envVars = executor.GetEnvVars();
@@ -212,25 +211,39 @@ namespace JobDriver.Tests.Execution
 			Assert.AreEqual(_workspace.DirectoryPath.FullName, envVars["uebp_LOCAL_ROOT"]);
 			Assert.AreEqual(batch.StreamName, envVars["uebp_BuildRoot_P4"]);
 			Assert.AreEqual("++UE5+Main", envVars["uebp_BuildRoot_Escaped"]);
-			Assert.AreEqual("1", envVars["uebp_CL"]);
-			Assert.AreEqual("0", envVars["uebp_CodeCL"]);
+			Assert.AreEqual("100", envVars["uebp_CL"]);
+			Assert.AreEqual("200", envVars["uebp_CodeCL"]);
 			Assert.AreEqual(_autoSdkWorkspace.DirectoryPath.FullName, envVars["UE_SDKS_ROOT"]);
 
 			await executor.FinalizeAsync(_logger, CancellationToken.None);
 		}
+		
+		[TestMethod]
+		public async Task EnvVars_CodeClDefaultsToCl_Async()
+		{
+			using WorkspaceExecutor executor = Create();
+			RpcBeginBatchResponse batch = new () { Change = ChangeNum, StreamName = "//UE5/Main" };
+			await executor.InitializeAsync(batch, _logger, CancellationToken.None);
+			
+			IReadOnlyDictionary<string, string> envVars = executor.GetEnvVars();
+			Assert.AreEqual("100", envVars["uebp_CL"]);
+			Assert.AreEqual("100", envVars["uebp_CodeCL"]);
+			
+			await executor.FinalizeAsync(_logger, CancellationToken.None);
+		}
 
 		[TestMethod]
 		public async Task JobWithPreflightAsync()
 		{
 			JobId preflightJobId = JobId.Parse("65bd0655591b5d5d7d047b59");
 
-			_hordeClient.AddJob(preflightJobId, _streamId, 1, 1000);
+			_hordeClient.AddJob(preflightJobId, _streamId, ChangeNum, 1000);
 			_workspace.SetFile(1000, "New/Feature/Foo.cs", "foo");
 
 			JobExecutorOptions executorOptions = new JobExecutorOptions(_hordeClient, _workingDir, Array.Empty<ProcessToTerminate>(), preflightJobId, _batchId, default, new RpcJobOptions());
 			using WorkspaceExecutor executor = new(executorOptions, _workspace, null, SimpleTestExecutor.NoOpTracer, NullLogger.Instance);
 
-			RpcBeginBatchResponse batch = new RpcBeginBatchResponse { Change = 1, PreflightChange = 1000 };
+			RpcBeginBatchResponse batch = new RpcBeginBatchResponse { Change = ChangeNum, PreflightChange = 1000 };
 			await executor.InitializeAsync(batch, _logger, CancellationToken.None);
 			AssertWorkspaceFile(_workspace, "main.cpp", "main");
 			AssertWorkspaceFile(_workspace, "foo/bar/baz.h", "baz");
@@ -257,5 +270,11 @@ namespace JobDriver.Tests.Execution
 			DirectoryReference workspaceDir = workspace.DirectoryPath;
 			Assert.AreEqual(expectedContent, File.ReadAllText(Path.Join(workspaceDir.FullName, relativePath)));
 		}
+		
+		private WorkspaceExecutor Create()
+		{
+			JobExecutorOptions executorOptions = new (_hordeClient, _workingDir, Array.Empty<ProcessToTerminate>(), _jobId, _batchId, default, new RpcJobOptions());
+			return new(executorOptions, _workspace, _autoSdkWorkspace, SimpleTestExecutor.NoOpTracer, NullLogger.Instance);
+		}
 	}
 }
diff --git a/Engine/Source/Programs/Horde/Drivers/JobDriver/Execution/WorkspaceExecutor.cs b/Engine/Source/Programs/Horde/Drivers/JobDriver/Execution/WorkspaceExecutor.cs
--- a/Engine/Source/Programs/Horde/Drivers/JobDriver/Execution/WorkspaceExecutor.cs
+++ b/Engine/Source/Programs/Horde/Drivers/JobDriver/Execution/WorkspaceExecutor.cs
@@ -78,7 +78,7 @@ namespace JobDriver.Execution
 			_envVars["uebp_BuildRoot_P4"] = Batch.StreamName;
 			_envVars["uebp_BuildRoot_Escaped"] = Batch.StreamName.Replace('/', '+');
 			_envVars["uebp_CL"] = Batch.Change.ToString();
-			_envVars["uebp_CodeCL"] = Batch.CodeChange.ToString();
+			_envVars["uebp_CodeCL"] = Batch.CodeChange > 0 ? Batch.CodeChange.ToString() : Batch.Change.ToString();
 
 			foreach ((string key, string value) in _workspace.EnvironmentVariables)
 			{
