Horde: Add missing file for CL 37388671

diff --git a/Engine/Source/Programs/Horde/Drivers/JobDriver/Execution/PerforceExecutor.cs b/Engine/Source/Programs/Horde/Drivers/JobDriver/Execution/PerforceExecutor.cs
--- a/Engine/Source/Programs/Horde/Drivers/JobDriver/Execution/PerforceExecutor.cs
+++ b/Engine/Source/Programs/Horde/Drivers/JobDriver/Execution/PerforceExecutor.cs
@@ -158,7 +158,7 @@ namespace JobDriver.Execution
 			_envVars["uebp_BuildRoot_Escaped"] = Batch.StreamName.Replace('/', '+');
 			_envVars["uebp_CLIENT_ROOT"] = $"//{_workspace.ClientName}";
 			_envVars["uebp_CL"] = Batch.Change.ToString();
-			_envVars["uebp_CodeCL"] = Batch.CodeChange.ToString();
+			_envVars["uebp_CodeCL"] = Batch.CodeChange > 0 ? Batch.CodeChange.ToString() : Batch.Change.ToString();
 			_envVars["P4USER"] = _workspace.UserName;
 			_envVars["P4CLIENT"] = _workspace.ClientName;
 
