Horde: Fix compiler warnings

diff --git a/Engine/Source/Programs/Horde/Plugins/Build/HordeServer.Build.Tests/Artifacts/ArtifactAclTests.cs b/Engine/Source/Programs/Horde/Plugins/Build/HordeServer.Build.Tests/Artifacts/ArtifactAclTests.cs
--- a/Engine/Source/Programs/Horde/Plugins/Build/HordeServer.Build.Tests/Artifacts/ArtifactAclTests.cs
+++ b/Engine/Source/Programs/Horde/Plugins/Build/HordeServer.Build.Tests/Artifacts/ArtifactAclTests.cs
@@ -76,7 +76,7 @@ namespace HordeServer.Tests.Artifacts
 			return entry;
 		}
 
-		ClaimsPrincipal CreateUser(string role)
+		static ClaimsPrincipal CreateUser(string role)
 		{
 			ClaimsIdentity identity = new ClaimsIdentity([new Claim(ClaimType, role)]);
 			return new ClaimsPrincipal(identity);
diff --git a/Engine/Source/Programs/Horde/Plugins/Build/HordeServer.Build/Artifacts/ArtifactTypeConfig.cs b/Engine/Source/Programs/Horde/Plugins/Build/HordeServer.Build/Artifacts/ArtifactTypeConfig.cs
--- a/Engine/Source/Programs/Horde/Plugins/Build/HordeServer.Build/Artifacts/ArtifactTypeConfig.cs
+++ b/Engine/Source/Programs/Horde/Plugins/Build/HordeServer.Build/Artifacts/ArtifactTypeConfig.cs
@@ -1,6 +1,5 @@
 // Copyright Epic Games, Inc. All Rights Reserved.
 
-using EpicGames.Horde.Acls;
 using EpicGames.Horde.Artifacts;
 using EpicGames.Horde.Storage;
 using HordeServer.Acls;
diff --git a/Engine/Source/Programs/Horde/Plugins/Build/HordeServer.Build/Devices/DeviceReportService.cs b/Engine/Source/Programs/Horde/Plugins/Build/HordeServer.Build/Devices/DeviceReportService.cs
--- a/Engine/Source/Programs/Horde/Plugins/Build/HordeServer.Build/Devices/DeviceReportService.cs
+++ b/Engine/Source/Programs/Horde/Plugins/Build/HordeServer.Build/Devices/DeviceReportService.cs
@@ -226,7 +226,6 @@ namespace HordeServer.Devices
 		public Uri ServerUrl { get; init; }
 	}
 
-
 #pragma warning restore CS1591 // Missing XML comment for publicly visible type or member
 
 	[SingletonDocument("device-report-state", "6491bc0ce39f8b70b1ef2feb")]
@@ -240,7 +239,6 @@ namespace HordeServer.Devices
 	/// </summary>
 	public class DeviceReportService : IHostedService
 	{
-
 		readonly SingletonDocument<DeviceReportState> _state;
 		readonly DeviceService _deviceService;
 		readonly INotificationService _notificationService;
@@ -341,7 +339,7 @@ namespace HordeServer.Devices
 			state = await _state.UpdateAsync(s => s.ReportTime = currentTime, cancellationToken);
 		}
 
-		static internal DeviceIssueReport CreateDeviceReport(DeviceReportData deviceData, DeviceReportSettings settings)
+		internal static DeviceIssueReport CreateDeviceReport(DeviceReportData deviceData, DeviceReportSettings settings)
 		{
 			DeviceIssueReport issueReport = new DeviceIssueReport(settings.Channel);
 			deviceData.Devices.OrderBy(d => d.PlatformId.ToString()).ThenBy(d => d.PoolId.ToString()).ToList()
diff --git a/Engine/Source/Programs/Horde/Plugins/Build/HordeServer.Build/Notifications/Sinks/SlackNotificationSink.cs b/Engine/Source/Programs/Horde/Plugins/Build/HordeServer.Build/Notifications/Sinks/SlackNotificationSink.cs
--- a/Engine/Source/Programs/Horde/Plugins/Build/HordeServer.Build/Notifications/Sinks/SlackNotificationSink.cs
+++ b/Engine/Source/Programs/Horde/Plugins/Build/HordeServer.Build/Notifications/Sinks/SlackNotificationSink.cs
@@ -1,6 +1,5 @@
 // Copyright Epic Games, Inc. All Rights Reserved.
 
-using System;
 using System.Diagnostics;
 using System.Diagnostics.CodeAnalysis;
 using System.Globalization;
@@ -1367,93 +1366,95 @@ namespace HordeServer.Notifications.Sinks
 				}
 			}
 
-			// Reactions
+			await UpdateReactionsAsync(triageChannel, issue, fixFailedSpan, eventId, threadId, cancellationToken);
+			
+			if (issue.ExternalIssueKey != null)
+			{
+				string extIssueEventId = $"issue_{issue.Id}_ext_{issue.ExternalIssueKey}";
+				string extIssueMessage = $"Linked to issue {FormatExternalIssue(issue.ExternalIssueKey)}";
+				await PostSingleMessageToThreadAsync(triageChannel, extIssueEventId, threadId, extIssueMessage, cancellationToken);
+			}
+		}
+		
+		private async Task UpdateReactionsAsync(string triageChannel, IIssue issue, IIssueSpan? fixFailedSpan, string eventId, SlackMessageId threadId, CancellationToken cancellationToken)
+		{
+			ReactionFlags reactions = ReactionFlags.None;
+			if (issue.AcknowledgedAt != null)
+			{
+				reactions |= ReactionFlags.Acknowledged;
+			}
+			if (issue.QuarantinedByUserId != null)
+			{
+				reactions |= ReactionFlags.Quarantined;
+			}
+			if (fixFailedSpan != null)
+			{
+				reactions |= ReactionFlags.FixFailed;
+			}
+			if (issue.ResolvedAt != null && fixFailedSpan == null)
+			{
+				reactions |= ReactionFlags.Resolved;
+			}
+			if (issue.AcknowledgedAt == null && issue.OwnerId != null)
+			{
+				reactions |= ReactionFlags.AssignedButNotAcknowledged;
+			}
+			
+			string reactionEventId = $"{eventId}_reactions";
+			string reactionStateDigest = ((int)reactions).ToString();
+			
+			MessageStateDocument? reactionState = await GetMessageStateAsync(triageChannel, reactionEventId, cancellationToken);
+			if ((reactionState == null && reactions != ReactionFlags.None) || (reactionState != null && reactionState.Digest != reactionStateDigest))
 			{
-				ReactionFlags reactions = ReactionFlags.None;
-				if (issue.AcknowledgedAt != null)
+				if ((reactions & ReactionFlags.Acknowledged) != 0)
 				{
-					reactions |= ReactionFlags.Acknowledged;
+					await _slackClient.AddReactionAsync(threadId, "eyes", cancellationToken);
 				}
-				if (issue.QuarantinedByUserId != null)
+				else
 				{
-					reactions |= ReactionFlags.Quarantined;
+					await _slackClient.RemoveReactionAsync(threadId, "eyes", cancellationToken);
 				}
-				if (fixFailedSpan != null)
+				
+				if ((reactions & ReactionFlags.Quarantined) != 0)
 				{
-					reactions |= ReactionFlags.FixFailed;
+					await _slackClient.AddReactionAsync(threadId, "mask", cancellationToken);
 				}
-				if (issue.ResolvedAt != null && fixFailedSpan == null)
+				else
 				{
-					reactions |= ReactionFlags.Resolved;
+					await _slackClient.RemoveReactionAsync(threadId, "mask", cancellationToken);
 				}
-				if (issue.AcknowledgedAt == null && issue.OwnerId != null)
+				
+				if ((reactions & ReactionFlags.FixFailed) != 0)
 				{
-					reactions |= ReactionFlags.AssignedButNotAcknowledged;
+					await _slackClient.AddReactionAsync(threadId, "x", cancellationToken);
 				}
-
-				string reactionEventId = $"{eventId}_reactions";
-				string reactionStateDigest = ((int)reactions).ToString();
-
-				MessageStateDocument? reactionState = await GetMessageStateAsync(triageChannel, reactionEventId, cancellationToken);
-				if ((reactionState == null && reactions != ReactionFlags.None) || (reactionState != null && reactionState.Digest != reactionStateDigest))
+				else
 				{
-					if ((reactions & ReactionFlags.Acknowledged) != 0)
-					{
-						await _slackClient.AddReactionAsync(threadId, "eyes", cancellationToken);
-					}
-					else
-					{
-						await _slackClient.RemoveReactionAsync(threadId, "eyes", cancellationToken);
-					}
-
-					if ((reactions & ReactionFlags.Quarantined) != 0)
-					{
-						await _slackClient.AddReactionAsync(threadId, "mask", cancellationToken);
-					}
-					else
-					{
-						await _slackClient.RemoveReactionAsync(threadId, "mask", cancellationToken);
-					}
-
-					if ((reactions & ReactionFlags.FixFailed) != 0)
-					{
-						await _slackClient.AddReactionAsync(threadId, "x", cancellationToken);
-					}
-					else
-					{
-						await _slackClient.RemoveReactionAsync(threadId, "x", cancellationToken);
-					}
-
-					if ((reactions & ReactionFlags.Resolved) != 0)
-					{
-						await _slackClient.AddReactionAsync(threadId, "tick", cancellationToken);
-					}
-					else
-					{
-						await _slackClient.RemoveReactionAsync(threadId, "tick", cancellationToken);
-					}
-
-					if ((reactions & ReactionFlags.AssignedButNotAcknowledged) != 0)
-					{
-						await _slackClient.AddReactionAsync(threadId, "mailbox", cancellationToken);
-					}
-					else
-					{
-						await _slackClient.RemoveReactionAsync(threadId, "mailbox", cancellationToken);
-					}
-
-					await AddOrUpdateMessageStateAsync(triageChannel, reactionEventId, null, reactionStateDigest, threadId, cancellationToken);
+					await _slackClient.RemoveReactionAsync(threadId, "x", cancellationToken);
 				}
-			}
-
-			if (issue.ExternalIssueKey != null)
-			{
-				string extIssueEventId = $"issue_{issue.Id}_ext_{issue.ExternalIssueKey}";
-				string extIssueMessage = $"Linked to issue {FormatExternalIssue(issue.ExternalIssueKey)}";
-				await PostSingleMessageToThreadAsync(triageChannel, extIssueEventId, threadId, extIssueMessage, cancellationToken);
+				
+				if ((reactions & ReactionFlags.Resolved) != 0)
+				{
+					await _slackClient.AddReactionAsync(threadId, "tick", cancellationToken);
+				}
+				else
+				{
+					await _slackClient.RemoveReactionAsync(threadId, "tick", cancellationToken);
+				}
+				
+				if ((reactions & ReactionFlags.AssignedButNotAcknowledged) != 0)
+				{
+					await _slackClient.AddReactionAsync(threadId, "mailbox", cancellationToken);
+				}
+				else
+				{
+					await _slackClient.RemoveReactionAsync(threadId, "mailbox", cancellationToken);
+				}
+				
+				await AddOrUpdateMessageStateAsync(triageChannel, reactionEventId, null, reactionStateDigest, threadId, cancellationToken);
 			}
 		}
-
+		
 		Uri? GetSwarmLinkForSpan(IIssueSpan span)
 		{
 			if (_settings.P4SwarmUrl == null || span.LastSuccess == null)
