Horde: Update release notes

#rnx

diff --git a/Engine/Source/Programs/Horde/Docs/ReleaseNotes.md b/Engine/Source/Programs/Horde/Docs/ReleaseNotes.md
--- a/Engine/Source/Programs/Horde/Docs/ReleaseNotes.md
+++ b/Engine/Source/Programs/Horde/Docs/ReleaseNotes.md
@@ -1,5 +1,24 @@
 # Release Notes
 
+## 2024-10-24
+
+* Refactor compiler warnings controlled via module rules or target rules (37409668)
+* Fallback to normal change number for BuildGraph invocation if code change is missing (37388671)
+* Prevent send and possible exceptions in background send for ComputeSocket (37387020)
+* Fix undefined access in job vew (37350938)
+* Update dashboard install docs for Vite (37350502)
+* Device view fixes (37350192)
+* Tweaks to compact views to support hiding/showing names (37350158)
+* Slowed the read rate to avoid out of memory issues. (37323730)
+* Fix hanging chrome when missing job template, also surface job error (37322037)
+* UBA: Fixed bug where bool was used instead of byte when transferring from native to managed code. (37320327)
+* UBA: Hooked up cross architecture logic from c# to c++ (37316216)
+* Set agent's default setting for EC2 support to disabled
+* Two second connect timeout and no printed message when UBT telemetry fails to connect or send. (37297792)
+* Add initial documentation for compute plugin
+* Return Android devices with Turnkey, and Editor `Platforms` menu update
+* Adding robomerge track change messages to slack messages. (37266636)
+
 ## 2024-10-17
 
 * Delete ephemeral agents that have been offline for more than X hours. This ignores checking for any outdated sessions and simply uses the last online timestamp instead. (37172088)
