Horde: Reduce size of self-contained agent

Removes unused libs from published directory, such as Linux and macOS versions accidentally making their way in.

diff --git a/Engine/Source/Programs/Horde/HordeAgent/HordeAgent.csproj b/Engine/Source/Programs/Horde/HordeAgent/HordeAgent.csproj
--- a/Engine/Source/Programs/Horde/HordeAgent/HordeAgent.csproj
+++ b/Engine/Source/Programs/Horde/HordeAgent/HordeAgent.csproj
@@ -194,6 +194,19 @@
     <Content Update="appsettings.Local.json" CopyToPublishDirectory="Never" />
   </ItemGroup>
 
+  <Target Name="DeleteUnusedFiles-Win64" AfterTargets="Publish" Condition="'$(Configuration)' == 'Release' AND '$(RuntimeIdentifier)' == 'win-x64'">
+    <ItemGroup>
+      <FilesToDelete Include="$(PublishDir)**/*.dylib" />
+      <FilesToDelete Include="$(PublishDir)**/*.so.9" />
+      <FilesToDelete Include="$(PublishDir)**/*.so" />
+      <FilesToDelete Include="$(PublishDir)**/*.xml" />
+    </ItemGroup>
+    <Delete Files="@(FilesToDelete)">
+      <Output TaskParameter="DeletedFiles" ItemName="DeletedList"/>
+    </Delete>
+    <Message Text="Deleted files:%0A@(DeletedList->'  %(Identity)','%0A')" Importance="high" />
+  </Target>
+
   <Target Name="PublishSymbols" AfterTargets="AfterBuild" Condition="'$(FullSymbolsDir)' != ''">
     <ItemGroup>
       <OutputFiles Include="$(OutDir)**" />
