Horde: Add cluster to perforce config check
#skipci

diff --git a/Engine/Source/Programs/Horde/HordeDashboard/src/backend/Backend.ts b/Engine/Source/Programs/Horde/HordeDashboard/src/backend/Backend.ts
--- a/Engine/Source/Programs/Horde/HordeDashboard/src/backend/Backend.ts
+++ b/Engine/Source/Programs/Horde/HordeDashboard/src/backend/Backend.ts
@@ -1662,10 +1662,10 @@ export class Backend {
         });
     }
 
-    checkPreflightConfig(shelvedChange: number): Promise<PreflightConfigResponse> {
+    checkPreflightConfig(shelvedChange: number, cluster?: string ): Promise<PreflightConfigResponse> {
 
         return new Promise<PreflightConfigResponse>((resolve, reject) => {
-            this._fetch.post(`/api/v1/server/preflightconfig`, { shelvedChange: shelvedChange }).then((value) => {
+            this._fetch.post(`/api/v1/server/preflightconfig`, { shelvedChange: shelvedChange, cluster: cluster }).then((value) => {
                 resolve(value.data as PreflightConfigResponse);
             }).catch(reason => {
                 reject(reason);
diff --git a/Engine/Source/Programs/Horde/HordeDashboard/src/components/preflights/PreflightConfigCheck.tsx b/Engine/Source/Programs/Horde/HordeDashboard/src/components/preflights/PreflightConfigCheck.tsx
--- a/Engine/Source/Programs/Horde/HordeDashboard/src/components/preflights/PreflightConfigCheck.tsx
+++ b/Engine/Source/Programs/Horde/HordeDashboard/src/components/preflights/PreflightConfigCheck.tsx
@@ -1,17 +1,62 @@
-import { FontIcon, IconButton, MaskedTextField, Modal, PrimaryButton, Spinner, SpinnerSize, Stack, Text } from "@fluentui/react";
+import { ComboBox, FontIcon, IComboBoxOption, IconButton, MaskedTextField, Modal, PrimaryButton, Spinner, SpinnerSize, Stack, Text } from "@fluentui/react";
 import { useState } from "react";
 import { useNavigate } from "react-router-dom";
 import backend from "../../backend";
 import dashboard, { StatusColor } from "../../backend/Dashboard";
 import { getHordeStyling } from "../../styles/Styles";
 
+class PerforceServerHandler {
+
+   async load(): Promise<void> {
+
+      try {
+
+         this.loading = true;
+         const servers = await backend.getPerforceServerStatus();
+         const unique = new Set(servers.map(s => s.cluster));
+         this.options = [];
+         unique.forEach(c => {
+            this.options?.push({
+               key: `${c}`,
+               text: c
+            })
+         })           
+         
+         this.options = this.options?.sort((a, b) => {
+            return a.text.localeCompare(b.text);
+         })
+
+      } catch (err) {
+         console.error(err);
+         this.options = [];
+      } finally {
+         this.loading = false;
+      }
+   }
+
+   loading = false;
+   options?: IComboBoxOption[];
+}
+
+const handler = new PerforceServerHandler();
+
+
 const PreflightConfigPanel: React.FC = () => {
-   
+
    const search = new URL(window.location.toString()).searchParams;
    const shelvedChange = search.get("preflightconfig") ? search.get("preflightconfig")! : undefined;
+   const queryCluster = search.get("cluster") ? search.get("cluster")! : undefined;
 
    const navigate = useNavigate();
-   const [state, setState] = useState<{ initialCL?: string, submitting?: boolean, success?: boolean, message?: string }>({ initialCL: shelvedChange });
+   const [state, setState] = useState<{ initialCL?: string, submitting?: boolean, success?: boolean, message?: string, cluster?: string }>({ initialCL: shelvedChange, cluster: queryCluster === undefined ? "Default" : queryCluster });
+   const [servers, setServers] = useState(handler.options);
+
+   if (servers === undefined) {
+      if (!handler.loading) {
+         handler.load().then(() => setServers(handler.options))
+      }
+      return null;
+   }
 
    const maskFormat: { [key: string]: RegExp } = {
       '*': /[0-9]/,
@@ -28,10 +73,10 @@ const PreflightConfigPanel: React.FC = () => {
          preflightCL = shelvedChange;
       }
 
-      setState({ submitting: true });
+      setState({ submitting: true, cluster: state.cluster });
 
       try {
-         const response = await backend.checkPreflightConfig(parseInt(preflightCL));
+         const response = await backend.checkPreflightConfig(parseInt(preflightCL), state.cluster === "Default" ? undefined : state.cluster);
          setState({ submitting: false, success: response.result, message: response.message });
       } catch (error) {
          setState({ submitting: false, success: false, message: error as string });
@@ -39,18 +84,17 @@ const PreflightConfigPanel: React.FC = () => {
 
    }
 
-   if (state.initialCL) {
+   if (state.initialCL && (state.cluster === "Default" || !!queryCluster)) {
       checkPreflight(state.initialCL);
       return null;
    }
 
-
    return (<Stack>
       <Stack styles={{ root: { paddingTop: 18, paddingLeft: 0, paddingRight: 0, width: "100%" } }} >
          <Stack tokens={{ childrenGap: 12 }} >
             <Stack style={{ width: 800, paddingLeft: 12 }}>
                <Stack horizontal verticalAlign="center" tokens={{ childrenGap: 24 }}>
-                  <Stack horizontal verticalAlign="center" tokens={{ childrenGap: 18 }}>
+                  <Stack horizontal verticalAlign="center" tokens={{ childrenGap: 12 }}>
                      <Text styles={{ root: { fontFamily: "Horde Open Sans SemiBold" } }}>Shelved Change</Text>
                      <MaskedTextField placeholder="Shelved Change" mask="***********" maskFormat={maskFormat} maskChar="" value={shelvedChange} onChange={(ev, newValue) => {
                         ev.preventDefault();
@@ -65,6 +109,14 @@ const PreflightConfigPanel: React.FC = () => {
 
                      }} />
                   </Stack>
+
+                  {(handler.options!.length > 1) && <Stack horizontal verticalAlign="center" tokens={{ childrenGap: 12 }}>
+                     <Text styles={{ root: { fontFamily: "Horde Open Sans SemiBold" } }}>Cluster</Text>
+                     <ComboBox styles={{ root: { width: 180 } }} disabled={state.submitting} selectedKey={state.cluster} options={handler.options!} onChange={(ev, option, index, value) => {
+                        setState({ ...state, cluster: option?.text })
+                     }} />
+                  </Stack>}
+
                   <Stack grow />
                   {!!state.submitting && <Stack>
                      <Spinner size={SpinnerSize.large} />
@@ -76,6 +128,7 @@ const PreflightConfigPanel: React.FC = () => {
 
                   {state.success === false && <Stack>
                      <FontIcon style={{ color: dashboard.getStatusColors().get(StatusColor.Failure)!, fontSize: 24 }} iconName="Cross" />
+
                   </Stack>}
 
                   <Stack>
