Horde: Expose CPU and OS architecture as agent properties

Needed to distinguish between x86_64 and aarch64 for upcoming support in UBA.

diff --git a/Engine/Source/Programs/Horde/HordeAgent/Services/CapabilitiesService.cs b/Engine/Source/Programs/Horde/HordeAgent/Services/CapabilitiesService.cs
--- a/Engine/Source/Programs/Horde/HordeAgent/Services/CapabilitiesService.cs
+++ b/Engine/Source/Programs/Horde/HordeAgent/Services/CapabilitiesService.cs
@@ -469,6 +469,8 @@ namespace HordeAgent.Services
 			// Used during the transition period over from multi-platform, non-self-contained agent packages.
 			capabilities.Properties.Add($"SelfContained={AgentApp.IsSelfContained}");
 			capabilities.Properties.Add($"DotNetFramework={RuntimeInformation.FrameworkDescription}");
+			capabilities.Properties.Add($"ProcessArch={RuntimeInformation.ProcessArchitecture}");
+			capabilities.Properties.Add($"OsArch={RuntimeInformation.OSArchitecture}");
 
 			// Add any additional properties from the config file
 			capabilities.Properties.AddRange(_settings.CurrentValue.Properties.Select(kvp => $"{kvp.Key}={kvp.Value}"));
