Horde: Fix some step retry logic
#submitool bypass
#skipci

diff --git a/Engine/Source/Programs/Horde/HordeDashboard/src/components/jobDetailsV2/JobOperationsBar.tsx b/Engine/Source/Programs/Horde/HordeDashboard/src/components/jobDetailsV2/JobOperationsBar.tsx
--- a/Engine/Source/Programs/Horde/HordeDashboard/src/components/jobDetailsV2/JobOperationsBar.tsx
+++ b/Engine/Source/Programs/Horde/HordeDashboard/src/components/jobDetailsV2/JobOperationsBar.tsx
@@ -57,21 +57,46 @@ export const JobOperations: React.FC<{ jobDetails: JobDetailsV2 }> = observer(({
 
    const abortDisabled = jobData.state === JobState.Complete;
 
-   const failedSteps = jobDetails.getSteps().filter(s => {
+   const seenSteps = new Set<string>();
+   const runningSteps = new Set<string>();
 
-      if (s.state !== JobStepState.Skipped && (!s.finishTime || (s.state !== JobStepState.Aborted && s.outcome !== JobStepOutcome.Failure))) {
+   let failedSteps = jobDetails.getSteps().reverse().filter(s => {
+
+      if (!s.allowRetry) {
          return false;
       }
 
-      const retries = jobDetails.getStepRetries(s.id);
-      const retryNumber = jobDetails.getStepRetryNumber(s.id);
-      if (retries.length && retryNumber < (retries.length - 1)) {
+      if (seenSteps.has(s.name) || runningSteps.has(s.name)) {
+         return false;
+      }
+
+      seenSteps.add(s.name);
+
+      if (s.state === JobStepState.Running || s.state === JobStepState.Waiting || s.state === JobStepState.Ready) { 
+         runningSteps.add(s.name);
+         return false;
+      }
+
+      if (s.state === JobStepState.Completed && s.outcome == JobStepOutcome.Success) 
+      {
          return false;
       }
 
       return true;
 
-   });
+   }).reverse();
+
+   const allSteps = jobDetails.getSteps();
+   failedSteps = failedSteps.sort((a, b) => {
+
+      const idxA = allSteps.findIndex((c) => a.name === c.name);
+      const idxB = allSteps.findIndex((c) => b.name === c.name);
+
+      return idxA - idxB;
+   })
+
+   
+
    const retryFailedStepsDisabled = !failedSteps.length;
 
    const pinned = dashboard.jobPinned(jobId);
