Horde: And maintain order when selecting deselecting steps
#skipci

diff --git a/Engine/Source/Programs/Horde/HordeDashboard/src/components/jobDetailsV2/StepRetryModal.tsx b/Engine/Source/Programs/Horde/HordeDashboard/src/components/jobDetailsV2/StepRetryModal.tsx
--- a/Engine/Source/Programs/Horde/HordeDashboard/src/components/jobDetailsV2/StepRetryModal.tsx
+++ b/Engine/Source/Programs/Horde/HordeDashboard/src/components/jobDetailsV2/StepRetryModal.tsx
@@ -91,12 +91,19 @@ export const RetryStepsModal: React.FC<{ stepIds: string[]; jobDetails: JobDetai
 
       const results: RunAgainResult[] = [];
 
-      const retryStepIds = Array.from(retrySteps);
+      const retryStepIds = Array.from(retrySteps).sort((a, b) => {
+         const idxA = stepIds.indexOf(a);
+         const idxB = stepIds.indexOf(b);
+
+         return idxA - idxB;
+      });
 
       for (let i = 0; i < retryStepIds.length; i++) {
 
          const stepId = retryStepIds[i];
 
+         console.log(`Retrying: ${jobDetails.stepById(stepId)?.name}`);
+
          const batch = jobDetails.batchByStepId(stepId);
          if (!batch) {
             results.push({
@@ -107,6 +114,7 @@ export const RetryStepsModal: React.FC<{ stepIds: string[]; jobDetails: JobDetai
          }
 
          try {
+            
             const response = await backend.updateJobStep(jobData.id, batch.id, stepId, { retry: true })
 
             if (!response.stepId) {
