Horde: Add better error message when agent IPC server cannot start

diff --git a/Engine/Source/Programs/Horde/HordeAgent/Services/StatusService.cs b/Engine/Source/Programs/Horde/HordeAgent/Services/StatusService.cs
--- a/Engine/Source/Programs/Horde/HordeAgent/Services/StatusService.cs
+++ b/Engine/Source/Programs/Horde/HordeAgent/Services/StatusService.cs
@@ -160,6 +160,10 @@ namespace HordeAgent.Services
 				{
 					break;
 				}
+				catch (UnauthorizedAccessException e)
+				{
+					_logger.LogError("Unable to start IPC server. Ensure no other Horde agent processes are running! Reason: {Reason}", e.Message);
+				}
 				catch (Exception ex)
 				{
 					_logger.LogError(ex, "Exception while running pipe server: {Message}", ex.Message);
