Horde: Add install/uninstall commands for agent

diff --git a/Engine/Source/Programs/Horde/HordeAgent/Commands/Service/ServiceInstall.cs b/Engine/Source/Programs/Horde/HordeAgent/Commands/Service/ServiceInstall.cs
--- a/Engine/Source/Programs/Horde/HordeAgent/Commands/Service/ServiceInstall.cs
+++ b/Engine/Source/Programs/Horde/HordeAgent/Commands/Service/ServiceInstall.cs
@@ -9,11 +9,19 @@ using Microsoft.Extensions.Logging;
 
 namespace HordeAgent.Commands.Service
 {
+	internal record InstallServiceOptions(
+		string? UserName = null,
+		string? Password = null,
+		string? ServerProfile = null,
+		string DotNetExecutable = "dotnet",
+		bool StartService = true
+	);
+	
 	/// <summary>
 	/// Installs the agent as a Windows service
 	/// </summary>
 	[Command("service", "install", "Installs the agent as a Windows service")]
-	class InstallCommand : Command
+	class InstallServiceCommand : Command
 	{
 		/// <summary>
 		/// Name of the service
@@ -53,72 +61,91 @@ namespace HordeAgent.Commands.Service
 				return Task.FromResult(1);
 			}
 
-			using (WindowsServiceManager serviceManager = new WindowsServiceManager())
+			InstallServiceOptions options = new ()
 			{
-				using (WindowsService service = serviceManager.Open(ServiceName))
+				UserName = UserName,
+				Password = Password,
+				ServerProfile = Server,
+				DotNetExecutable = DotNetExecutable,
+				StartService = Start != null && Start.Equals("true", StringComparison.OrdinalIgnoreCase),
+			};
+			return Task.FromResult(InstallService(options, logger));
+		}
+		
+		internal static bool StopAndDeleteWindowsService(ILogger logger)
+		{
+			using WindowsServiceManager serviceManager = new ();
+			using WindowsService service = serviceManager.Open(ServiceName);
+			if (service.IsValid)
+			{
+				logger.LogInformation("Stopping existing service...");
+				service.Stop();
+
+				WindowsServiceStatus status = service.WaitForStatusChange(WindowsServiceStatus.Stopping, TimeSpan.FromSeconds(30.0));
+				if (status != WindowsServiceStatus.Stopped)
 				{
-					if (service.IsValid)
-					{
-						logger.LogInformation("Stopping existing service...");
-						service.Stop();
-
-						WindowsServiceStatus status = service.WaitForStatusChange(WindowsServiceStatus.Stopping, TimeSpan.FromSeconds(30.0));
-						if (status != WindowsServiceStatus.Stopped)
-						{
-							logger.LogError("Unable to stop service (status = {Status})", status);
-							return Task.FromResult(1);
-						}
-
-						logger.LogInformation("Deleting service");
-						service.Delete();
-					}
+					logger.LogError("Unable to stop service (status = {Status})", status);
+					return false;
 				}
 
-				logger.LogInformation("Registering {ServiceName} service", ServiceName);
+				logger.LogInformation("Deleting service...");
+				service.Delete();
+			}
+			
+			return true;
+		}
+		
+		internal static int InstallService(InstallServiceOptions options, ILogger logger)
+		{
+			using WindowsServiceManager serviceManager = new ();
+			if (!StopAndDeleteWindowsService(logger))
+			{
+				return 1;
+			}
+			
+			logger.LogInformation("Registering {ServiceName} service", ServiceName);
 
-				StringBuilder commandLine = new();
-				if (AgentApp.IsSelfContained)
+			StringBuilder commandLine = new();
+			if (AgentApp.IsSelfContained)
+			{
+				if (Environment.ProcessPath == null)
 				{
-					if (Environment.ProcessPath == null)
-					{
-						logger.LogError("Unable to detect current process path");
-						return Task.FromResult(1);
-					}
-					commandLine.Append($"\"{Environment.ProcessPath}\" service run");
+					logger.LogError("Unable to detect current process path");
+					return 1;
 				}
-				else
-				{
+				commandLine.Append($"\"{Environment.ProcessPath}\" service run");
+			}
+			else
+			{
 #pragma warning disable IL3000 // Avoid accessing Assembly file path when publishing as a single file
-					commandLine.AppendFormat("{0} \"{1}\" service run", DotNetExecutable, Assembly.GetEntryAssembly()!.Location);
+				commandLine.AppendFormat("{0} \"{1}\" service run", options.DotNetExecutable, Assembly.GetEntryAssembly()!.Location);
 #pragma warning restore IL3000 // Avoid accessing Assembly file path when publishing as a single file					
-				}
+			}
 
-				if (Server != null)
-				{
-					commandLine.Append($" -server={Server}");
-				}
+			if (options.ServerProfile != null)
+			{
+				commandLine.Append($" -Server={options.ServerProfile}");
+			}
 
-				using (WindowsService service = serviceManager.Create(ServiceName, "Horde Agent", commandLine.ToString(), UserName, Password))
+			using (WindowsService service = serviceManager.Create(ServiceName, "Horde Agent", commandLine.ToString(), options.UserName, options.Password))
+			{
+				service.SetDescription("Allows this machine to participate in a Horde farm.");
+				if (options.StartService)
 				{
-					service.SetDescription("Allows this machine to participate in a Horde farm.");
+					logger.LogInformation("Starting...");
+					service.Start();
 
-					if (Start != null && Start.Equals("true", StringComparison.OrdinalIgnoreCase))
+					WindowsServiceStatus status = service.WaitForStatusChange(WindowsServiceStatus.Starting, TimeSpan.FromSeconds(30.0));
+					if (status != WindowsServiceStatus.Running)
 					{
-						logger.LogInformation("Starting...");
-						service.Start();
-
-						WindowsServiceStatus status = service.WaitForStatusChange(WindowsServiceStatus.Starting, TimeSpan.FromSeconds(30.0));
-						if (status != WindowsServiceStatus.Running)
-						{
-							logger.LogError("Unable to start service (status = {Status})", status);
-							return Task.FromResult(1);
-						}
+						logger.LogError("Unable to start service (status = {Status})", status);
+						return 1;
 					}
-
-					logger.LogInformation("Done.");
 				}
+				logger.LogInformation("Done");
 			}
-			return Task.FromResult(0);
+
+			return 0;
 		}
 	}
 }
diff --git a/Engine/Source/Programs/Horde/HordeAgent/Commands/Service/ServiceUninstall.cs b/Engine/Source/Programs/Horde/HordeAgent/Commands/Service/ServiceUninstall.cs
--- a/Engine/Source/Programs/Horde/HordeAgent/Commands/Service/ServiceUninstall.cs
+++ b/Engine/Source/Programs/Horde/HordeAgent/Commands/Service/ServiceUninstall.cs
@@ -10,12 +10,12 @@ namespace HordeAgent.Commands.Service
 	/// Uninstalls the Windows service
 	/// </summary>
 	[Command("service", "uninstall", "Uninstalls the Windows service")]
-	class UninstallCommand : Command
+	class UninstallServiceCommand : Command
 	{
 		/// <summary>
 		/// Name of the service
 		/// </summary>
-		public const string ServiceName = InstallCommand.ServiceName;
+		public const string ServiceName = InstallServiceCommand.ServiceName;
 
 		/// <summary>
 		/// Uninstalls the service
@@ -40,7 +40,7 @@ namespace HordeAgent.Commands.Service
 							return Task.FromResult(1);
 						}
 
-						logger.LogInformation("Deleting service");
+						logger.LogInformation("Deleting service...");
 						service.Delete();
 					}
 					else
diff --git a/Engine/Source/Programs/Horde/HordeAgent/Commands/Setup/SetupInstall.cs b/Engine/Source/Programs/Horde/HordeAgent/Commands/Setup/SetupInstall.cs
new file mode 100644
--- /dev/null
+++ b/Engine/Source/Programs/Horde/HordeAgent/Commands/Setup/SetupInstall.cs
@@ -0,0 +1,356 @@
+// Copyright Epic Games, Inc. All Rights Reserved.
+
+using System.ComponentModel;
+using System.Diagnostics;
+using System.Runtime.Versioning;
+using System.Text.Json;
+using System.Text.Json.Nodes;
+using EpicGames.Core;
+using Microsoft.Extensions.Logging;
+using Microsoft.Win32;
+using JsonObject = System.Text.Json.Nodes.JsonObject;
+
+namespace HordeAgent.Commands.Setup;
+
+/// <summary>
+/// Exception for any installer related errors
+/// </summary>
+/// <param name="message">Error message</param>
+/// <param name="exitCode">Optional error code to return</param>
+/// <param name="innerException">Optional inner exception</param>
+public class InstallException(string message, int exitCode = 1, Exception? innerException = null) : Exception(message, innerException)
+{
+	/// <summary>
+	/// Exit code
+	/// </summary>
+	public int ExitCode { get; } = exitCode;
+}
+
+/// <summary>
+/// Install the agent as a Windows application and service
+/// </summary>
+[Command("setup", "install", "Installs the agent (on Windows)")]
+class SetupInstallCommand : Command
+{
+	const string HordeAgentDll = "HordeAgent.dll";
+	const string HordeAgentExe = "HordeAgent.exe";
+	const string HordeAgentIco = "HordeAgent.ico";
+	const string UninstallKeyPath = @"SOFTWARE\Microsoft\Windows\CurrentVersion\Uninstall";
+	const string UninstallAppName = "Unreal Horde Agent";
+	const string PortRange = "7000-7010";
+	internal const string FirewallRuleName = "Horde Agent";
+	
+	[CommandLine("-ShowStackTrace=")]
+	[Description("Whether to show full stack trace on errors (true/false)")]
+	public string? ShowStackTrace { get; set; } = "false";
+	
+	/// <summary>
+	/// Install the 
+	/// </summary>
+	/// <param name="logger">Logger to use</param>
+	/// <returns>Exit code</returns>
+	public override Task<int> ExecuteAsync(ILogger logger)
+	{
+		if (!OperatingSystem.IsWindows())
+		{
+			logger.LogError("This command requires Windows");
+			return Task.FromResult(1);
+		}
+		
+		try
+		{
+			Install(logger, new InstallParams());
+		}
+		catch (InstallException e)
+		{
+			logger.LogError("Install failed: {Reason}", e.Message);
+			if (ShowStackTrace != null && ShowStackTrace.Equals("true", StringComparison.OrdinalIgnoreCase))
+			{
+				throw;
+			}
+			
+			return Task.FromResult(e.ExitCode);
+		}
+		
+		return Task.FromResult(0);
+	}
+	
+	/// <summary>
+	/// Parameters for installing agent on Windows 
+	/// </summary>
+	/// <param name="SourceDir">Directory containing files to install (HordeAgent.dll etc)</param>
+	/// <param name="DestDir">Installation directory (unless specified, defaults to a directory under Program Files)</param>
+	/// <param name="ConfigDir"></param>
+	/// <param name="WorkingDir"></param>
+	public record InstallParams(
+		string? SourceDir = null,
+		string? DestDir = null,
+		string? ConfigDir = null,
+		string? WorkingDir = null
+	);
+	
+	/// <summary>
+	/// Install agent on Windows 
+	/// </summary>
+	/// <param name="logger">Logger</param>
+	/// <param name="installParams">Parameters</param>
+	[SupportedOSPlatform("windows")]
+	public static int Install(ILogger logger, InstallParams installParams)
+	{
+		if (!UserSecurityUtils.IsAdministrator())
+		{
+			throw new InstallException("Administrator access is required for installing");
+		}
+		
+		DirectoryReference sourceDir = new(installParams.SourceDir ?? AppContext.BaseDirectory);
+		DirectoryReference destinationDir = new(installParams.DestDir ?? Path.Join(Environment.GetFolderPath(Environment.SpecialFolder.ProgramFiles), "Epic Games", "Horde", "Agent"));
+		DirectoryReference configDir = new(installParams.ConfigDir ?? Path.Join(Environment.GetFolderPath(Environment.SpecialFolder.CommonApplicationData), "Epic", "Horde", "Agent"));
+		DirectoryReference workingDir = new(installParams.WorkingDir ?? Path.Join("c:", "HordeAgent", "Sandbox"));
+		logger.LogInformation(" Source dir {SourceDir}", sourceDir);
+		logger.LogInformation("   Dest dir {DestinationDir}", destinationDir);
+		logger.LogInformation(" Config dir {ConfigDir}", configDir);
+		logger.LogInformation("Working dir {WorkingDir}", workingDir);
+		
+		if (DirectoryReference.Exists(destinationDir) && Directory.GetFiles(destinationDir.FullName).Length > 0)
+		{
+			throw new InstallException($"Installation destination directory {destinationDir} already exists. Please uninstall any current agent!");
+		}
+		
+		DirectoryReference.CreateDirectory(destinationDir);
+		if (!FileReference.Exists(FileReference.Combine(sourceDir, HordeAgentDll)))
+		{
+			throw new InstallException($"Source directory does not contain {HordeAgentDll}");
+		}
+		
+		logger.LogInformation("Copying files...");
+		try
+		{
+			DirectoryReference.Copy(sourceDir, destinationDir);
+		}
+		catch (Exception e)
+		{
+			throw new InstallException($"Failed copying files to destination directory {destinationDir}", 1, e);
+		}
+		
+		// Remove dir containing agent.json config template as it's installed under ProgramData instead
+		DirectoryReference.Delete(DirectoryReference.Combine(destinationDir, "Defaults"), true);
+		
+		FileReference agentJson = FileReference.Combine(configDir, "agent.json");
+		if (!FileReference.Exists(agentJson))
+		{
+			logger.LogInformation("Copied default template agent.json to {AgentJsonPath}", agentJson);
+			FileReference.Copy(FileReference.Combine(sourceDir, "Defaults", "agent.json"), agentJson);
+		}
+		
+		PatchAgentJsonFile(agentJson, "WorkingDir", workingDir.FullName);
+		DirectoryReference.CreateDirectory(workingDir);
+		
+		logger.LogInformation("Configuring Windows firewall...");
+		AddFirewallPortRule(FirewallRuleName, PortRange);
+		
+		logger.LogInformation("Installing Horde Agent as a Windows service...");
+		(string executable, string args) = ResolveCommandLine(destinationDir.FullName, "service install");
+		if (RunProcessSilently(executable, args) != 0)
+		{
+			throw new InstallException("Failed installing Horde Agent as a Windows service");
+		}
+		
+		(executable, args) = ResolveCommandLine(destinationDir.FullName, "setup uninstall");
+		RegisterAppForUninstall(new AppInfo(
+			DisplayName: "Unreal Horde Agent",
+			InstallLocation: destinationDir,
+			Version: AgentApp.Version,
+			Publisher: "Epic Games",
+			UninstallCommand: $"\"{executable}\" {args}",
+			IconPath: FileReference.Combine(destinationDir, HordeAgentIco),
+			EstimatedSizeInKb: DirectoryReference.GetTotalSize(destinationDir) / 1024
+		));
+		
+		return 0;
+	}
+	
+	/// <summary>
+	/// Resolves path to executable and arguments to send for invoking the Horde agent executable
+	/// If agent is self-contained, the binary will be a standalone executable.
+	/// Otherwise, 'dotnet' executable will be used with the .dll and additional arguments appended
+	/// </summary>
+	/// <param name="hordeAgentDir">Path to directory where Horde Agent is installed</param>
+	/// <param name="args">Additional arguments</param>
+	/// <returns>Executable and arguments</returns>
+	private static (string executable, string args) ResolveCommandLine(string hordeAgentDir, string args)
+	{
+		return AgentApp.IsSelfContained
+			? (Path.Join(hordeAgentDir, HordeAgentExe), args)
+			: ("dotnet", $"\"{Path.Join(hordeAgentDir, HordeAgentDll)}\" {args}");
+	}
+	
+	private static bool PatchAgentJsonFile(FileReference jsonFilePath, string key, string value)
+	{
+		// Patching like this will change formatting and comments of the JSON file. Handling that in a good way is tricky.
+		string rawJson = File.ReadAllText(jsonFilePath.FullName);
+		JsonNode? jsonNode = JsonNode.Parse(rawJson, new JsonNodeOptions { PropertyNameCaseInsensitive = true },
+			new JsonDocumentOptions() { AllowTrailingCommas = true, CommentHandling = JsonCommentHandling.Skip });
+		if (jsonNode == null)
+		{
+			return false;
+		}
+		
+		JsonNode hordeNode = jsonNode["Horde"] ?? new JsonObject();
+		hordeNode[key] = value;
+		File.WriteAllText(jsonFilePath.FullName, jsonNode.ToJsonString(new JsonSerializerOptions() { WriteIndented = true }));
+		return true;
+	}
+	
+	internal static void AddFirewallPortRule(string name, string portRange)
+	{
+		RemoveFirewallPortRule(name);
+		string arguments = $"advfirewall firewall add rule name=\"{name}\" dir=in action=allow protocol=TCP localport={portRange}";
+		if (RunProcessSilently("netsh", arguments) != 0)
+		{
+			throw new InstallException($"Failed to add firewall rule '{name}'");
+		}
+	}
+	
+	internal static void RemoveFirewallPortRule(string name)
+	{
+		RunProcessSilently("netsh", $"advfirewall firewall delete rule name=\"{name}\"");
+	}
+	
+	internal static int RunProcessSilently(string fileName, string arguments)
+	{
+		ProcessStartInfo startInfo = new() { FileName = fileName, Arguments = arguments, UseShellExecute = false, CreateNoWindow = true };
+		try
+		{
+			using Process? process = Process.Start(startInfo);
+			process?.WaitForExit(10000);
+			return process?.ExitCode ?? 100;
+		}
+		catch (Exception)
+		{
+			return 100;
+		}
+	}
+	
+	internal record AppInfo(
+		string DisplayName,
+		DirectoryReference InstallLocation,
+		string Version,
+		string Publisher,
+		string UninstallCommand,
+		FileReference? IconPath,
+		long EstimatedSizeInKb);
+	
+	[SupportedOSPlatform("windows")]
+	public static void RegisterAppForUninstall(AppInfo appInfo)
+	{
+		string regKeyPath = $"{UninstallKeyPath}\\{UninstallAppName}";
+		using RegistryKey? appKey = Registry.LocalMachine.CreateSubKey(regKeyPath, true);
+		if (appKey == null)
+		{
+			throw new InstallException($"Unable to create registry key {regKeyPath}");
+		}
+		
+		appKey.SetValue("DisplayName", appInfo.DisplayName);
+		appKey.SetValue("DisplayVersion", appInfo.Version);
+		appKey.SetValue("InstallLocation", appInfo.InstallLocation.FullName);
+		appKey.SetValue("InstallDate", DateTime.Now.ToString("yyyyMMdd"));
+		appKey.SetValue("Publisher", appInfo.Publisher);
+		appKey.SetValue("UninstallString", appInfo.UninstallCommand);
+		appKey.SetValue("QuietUninstallString", $"{appInfo.UninstallCommand}");
+		appKey.SetValue("EstimatedSize", appInfo.EstimatedSizeInKb, RegistryValueKind.DWord);
+		appKey.SetValue("NoModify", 1, RegistryValueKind.DWord);
+		appKey.SetValue("NoRepair", 1, RegistryValueKind.DWord);
+		
+		if (appInfo.IconPath != null)
+		{
+			appKey.SetValue("DisplayIcon", appInfo.IconPath.FullName);
+		}
+	}
+	
+	[SupportedOSPlatform("windows")]
+	public static bool UnregisterAppForUninstall(string appName = UninstallAppName)
+	{
+		using RegistryKey? key = Registry.LocalMachine.OpenSubKey(UninstallKeyPath, true);
+		if (key == null)
+		{
+			return false;
+		}
+		
+		key.DeleteSubKeyTree(appName, false);
+		return true;
+	}
+	
+	/// <summary>
+	/// Runs a Windows batch script that will delete the given directory
+	/// The script will be executed in the background and wait for a process to exit before attempting to delete.
+	/// This allows the uninstaller (HordeAgent.dll or similar) to remove itself once the main duties are completed.
+	/// </summary>
+	/// <param name="deleteDir">Directory to delete</param>
+	/// <param name="waitPid">Process ID to wait for before attempting a delete</param>
+	/// <exception cref="InstallException">Thrown when script can't be created or executed</exception>
+	public static void RunSelfDeleteScript(DirectoryReference deleteDir, int waitPid)
+	{
+		string script = $"""
+		                 @echo off
+		                 setlocal EnableDelayedExpansion
+		                 
+		                 :: Configuration
+		                 set "dirPath={deleteDir.FullName}"
+		                 set "maxWaitSeconds=60"
+		                 set "maxRetries=3"
+		                 set "retryDelaySeconds=2"
+		                 set "currentPid={waitPid}"
+		                 
+		                 :: Wait for the process to exit
+		                 set /A secsPassed=0
+		                 :waitloop
+		                 tasklist /FI "PID eq %currentPid%" 2>nul | find /I "%currentPid%" >nul
+		                 if !ERRORLEVEL! EQU 0 (
+		                     echo Process still running, checking again in one second...
+		                     timeout /t 1 /nobreak >nul
+		                     set /A secsPassed+=1
+		                     if !secsPassed! LSS %maxWaitSeconds% goto :waitloop
+		                     echo Warning: Process %currentPid% did not exit within %maxWaitSeconds% seconds.
+		                 )
+		                 
+		                 :: Attempt to remove directory
+		                 set /A retry=0
+		                 :retryloop
+		                 rd /s /q "%dirPath%" 2>nul
+		                 if exist "%dirPath%" (
+		                     set /A retry+=1
+		                     if !retry! LSS %maxRetries% (
+		                         echo Attempt !retry! failed. Retrying in %retryDelaySeconds% seconds...
+		                         timeout /t %retryDelaySeconds% /nobreak >nul
+		                         goto :retryloop
+		                     ) else (
+		                         echo Error: Failed to remove directory after %maxRetries% attempts.
+		                         exit /b 1
+		                     )
+		                 ) else (
+		                     echo Successfully removed directory.
+		                 )
+		                 
+		                 :: Clean up the script itself
+		                 rem del "%~f0"
+		                 exit /b 0
+		                 """;
+		
+		try
+		{
+			string batFilePath = Path.Combine(Path.GetTempPath(), "horde-agent-uninstall.bat");
+			File.WriteAllText(batFilePath, script);
+			Process.Start(new ProcessStartInfo
+			{
+				FileName = batFilePath,
+				CreateNoWindow = true,
+				UseShellExecute = false,
+				WindowStyle = ProcessWindowStyle.Hidden,
+			});
+		}
+		catch (System.Exception ex)
+		{
+			throw new InstallException("Unable to create and run self-delete script", 1, ex);
+		}      
+	}
+}
\ No newline at end of file
diff --git a/Engine/Source/Programs/Horde/HordeAgent/Commands/Setup/SetupUninstall.cs b/Engine/Source/Programs/Horde/HordeAgent/Commands/Setup/SetupUninstall.cs
new file mode 100644
--- /dev/null
+++ b/Engine/Source/Programs/Horde/HordeAgent/Commands/Setup/SetupUninstall.cs
@@ -0,0 +1,64 @@
+// Copyright Epic Games, Inc. All Rights Reserved.
+
+using System.ComponentModel;
+using System.Runtime.Versioning;
+using EpicGames.Core;
+using HordeAgent.Commands.Service;
+using Microsoft.Extensions.Logging;
+
+namespace HordeAgent.Commands.Setup;
+
+/// <summary>
+/// Uninstall the agent from Windows
+/// </summary>
+[Command("setup", "uninstall", "Uninstalls the agent (on Windows)")]
+class SetupUninstallCommand : Command
+{
+	[CommandLine("-ShowStackTrace=")]
+	[Description("Whether to show full stack trace on errors (true/false)")]
+	public string? ShowStackTrace { get; set; } = "false";
+	
+	/// <summary>
+	/// Uninstall the agent from Windows 
+	/// </summary>
+	/// <param name="logger">Logger to use</param>
+	/// <returns>Exit code</returns>
+	public override Task<int> ExecuteAsync(ILogger logger)
+	{
+		if (!OperatingSystem.IsWindows())
+		{
+			logger.LogError("This command requires Windows");
+			return Task.FromResult(1);
+		}
+		
+		try
+		{
+			Uninstall(logger);
+		}
+		catch (InstallException e)
+		{
+			logger.LogError("Uninstall failed: {Reason}", e.Message);
+			if (ShowStackTrace != null && ShowStackTrace.Equals("true", StringComparison.OrdinalIgnoreCase))
+			{
+				throw;
+			}
+			
+			return Task.FromResult(e.ExitCode);
+		}
+		
+		return Task.FromResult(0);
+	}
+	
+	[SupportedOSPlatform("windows")]
+	private static void Uninstall(ILogger logger)
+	{
+		logger.LogInformation("Uninstalling Horde agent...");
+		InstallServiceCommand.StopAndDeleteWindowsService(logger);
+		SetupInstallCommand.RemoveFirewallPortRule(SetupInstallCommand.FirewallRuleName);
+		SetupInstallCommand.UnregisterAppForUninstall();
+		
+		// Remove directory containing agent binaries
+		logger.LogInformation("Deleting dir {AgentDir}", AppContext.BaseDirectory);
+		SetupInstallCommand.RunSelfDeleteScript(new DirectoryReference(AppContext.BaseDirectory), Environment.ProcessId);
+	}
+}
\ No newline at end of file
diff --git a/Engine/Source/Programs/Horde/HordeAgent/HordeAgent.csproj b/Engine/Source/Programs/Horde/HordeAgent/HordeAgent.csproj
--- a/Engine/Source/Programs/Horde/HordeAgent/HordeAgent.csproj
+++ b/Engine/Source/Programs/Horde/HordeAgent/HordeAgent.csproj
@@ -112,6 +112,9 @@
       <CopyToOutputDirectory>Always</CopyToOutputDirectory>
       <CopyToPublishDirectory>Always</CopyToPublishDirectory>
     </Content>
+    <Content Include="HordeAgent.ico">
+      <CopyToOutputDirectory>PreserveNewest</CopyToOutputDirectory>
+    </Content>
   </ItemGroup>
 
   <ItemGroup>
