Horde: Fix issue with password not being sent when creating a new account if it is provided

diff --git a/Engine/Source/Programs/Horde/HordeDashboard/src/components/accounts/AccountsView.tsx b/Engine/Source/Programs/Horde/HordeDashboard/src/components/accounts/AccountsView.tsx
--- a/Engine/Source/Programs/Horde/HordeDashboard/src/components/accounts/AccountsView.tsx
+++ b/Engine/Source/Programs/Horde/HordeDashboard/src/components/accounts/AccountsView.tsx
@@ -231,11 +231,13 @@ const AccountEditor: React.FC<{ accountIn?: GetAccountResponse, onClose: () => v
       try {
 
          setSubmitting(true);
+
+         const password = secrets.password ? secrets.password : undefined;
          if (accountIn?.id) {
-            const uaccount: UpdateAccountRequest = { ...account, password: secrets.password };
+            const uaccount: UpdateAccountRequest = { ...account, password: password };
             await backend.updateAccount(accountIn.id, uaccount);
          } else {
-            const naccount: CreateAccountRequest = { ...account };
+            const naccount: CreateAccountRequest = { ...account, password: password };
             await backend.createAccount(naccount)
          }
          setSubmitting(false);
