Report device failures with one message per platform

#rnx

diff --git a/Engine/Source/Programs/Horde/Plugins/Build/HordeServer.Build/Notifications/Sinks/SlackNotificationSink.cs b/Engine/Source/Programs/Horde/Plugins/Build/HordeServer.Build/Notifications/Sinks/SlackNotificationSink.cs
--- a/Engine/Source/Programs/Horde/Plugins/Build/HordeServer.Build/Notifications/Sinks/SlackNotificationSink.cs
+++ b/Engine/Source/Programs/Horde/Plugins/Build/HordeServer.Build/Notifications/Sinks/SlackNotificationSink.cs
@@ -2498,27 +2498,28 @@ namespace HordeServer.Notifications.Sinks
 
 				foreach (DevicePlatformReport platform in report.PlatformReports)
 				{
+					SlackMessage message = new SlackMessage();
+					message.Markdown = true;
+					message.Text = $"*{platform.PlatformName}*";
+
 					foreach (DeviceReport device in platform.DeviceReports.OrderByDescending(r => r.ProblemDelta))
 					{
-						StringBuilder builder = new StringBuilder($"   {device.ProblemDelta} Problem(s)");
+						StringBuilder builder = new StringBuilder($"   <{device.DevicePoolURL}|*{device.DeviceName}*>* | {device.DeviceAddress} - {device.PoolName} Pool*");
+						builder.Append($"\n{device.ProblemDelta} Problem(s)");
 
 						if (!String.IsNullOrEmpty(device.LastProblemURL))
 						{
 							builder.Append($" - <{device.LastProblemURL}|{device.LastProblemDesc ?? "??"}>");
 						}
 
-						SlackMessage message = new SlackMessage();
-						message.Markdown = true;
-						message.Text = $"<{device.DevicePoolURL}|*{device.DeviceName}*>* | {device.DeviceAddress} - {platform.PlatformName} - {device.PoolName}*";
-
 						SlackAttachment attachment = new SlackAttachment();
 						attachment.FallbackText = builder.ToString();
 						attachment.Color = "#d60404";
 						attachment.AddSection(builder.ToString());
-
 						message.Attachments.Add(attachment);
-						await SendMessageAsync(report.Channel, message, cancellationToken);
 					}
+
+					await SendMessageAsync(report.Channel, message, cancellationToken);
 				}
 
 				SlackMessage divider = new SlackMessage();
