Horde: Allow users without a configured password to log in correctly.

#rnx

diff --git a/Engine/Source/Programs/Horde/HordeServer/Accounts/AccountCollection.cs b/Engine/Source/Programs/Horde/HordeServer/Accounts/AccountCollection.cs
--- a/Engine/Source/Programs/Horde/HordeServer/Accounts/AccountCollection.cs
+++ b/Engine/Source/Programs/Horde/HordeServer/Accounts/AccountCollection.cs
@@ -176,7 +176,16 @@ namespace HordeServer.Accounts
 			}
 
 			public bool ValidatePassword(string password)
-				=> PasswordSalt != null && PasswordHash != null && PasswordHasher.ValidatePassword(password, PasswordHasher.SaltFromString(PasswordSalt), PasswordHasher.HashFromString(PasswordHash));
+			{
+				if (String.IsNullOrEmpty(PasswordHash))
+				{
+					return String.IsNullOrEmpty(password);
+				}
+				else
+				{
+					return PasswordSalt != null && PasswordHash != null && PasswordHasher.ValidatePassword(password, PasswordHasher.SaltFromString(PasswordSalt), PasswordHasher.HashFromString(PasswordHash));
+				}
+			}
 
 			public async Task<IAccount?> RefreshAsync(CancellationToken cancellationToken)
 				=> await _accountCollection!.GetAsync(Id, cancellationToken);
