Horde: Reference remote shader tutorial from home and landing doc pages

Also mentioned compute plugin in config overview

diff --git a/Engine/Source/Programs/Horde/Docs/Config.md b/Engine/Source/Programs/Horde/Docs/Config.md
--- a/Engine/Source/Programs/Horde/Docs/Config.md
+++ b/Engine/Source/Programs/Horde/Docs/Config.md
@@ -53,7 +53,8 @@ Manages worker machines that can be leased out to perform
 workloads, including build automation and remote execution use
 cases.
 
-* [Agents](Config/Agents.md): Configure worker machines to connect
+* [Compute](Config/Compute.md): Configuring agent pools and compute clusters.
+* [Agents](Config/Agents.md): Configuring worker machines to connect
   to the Horde server for CI and remote execution workspaces.
 
 ### Secrets
diff --git a/Engine/Source/Programs/Horde/Docs/Home.md b/Engine/Source/Programs/Horde/Docs/Home.md
--- a/Engine/Source/Programs/Horde/Docs/Home.md
+++ b/Engine/Source/Programs/Horde/Docs/Home.md
@@ -3,7 +3,9 @@
 ## Getting Started
 
 * **[Install the Horde agent](Tutorials/InstallHorde.md)**
-* **[Enable remote C++ compilation with Unreal Build Accelerator](Tutorials/RemoteCompilation.md)**
+* Remote compilation with Unreal Build Accelerator
+  * **[Enable remote C++ compilation](Tutorials/RemoteCompilation.md)**
+  * **[Enable remote shader compilation](Tutorials/RemoteShaderCompilation.md)**
 * **[Set up build automation](Tutorials/BuildAutomation.md)**
 * **[Enable test automation with Gauntlet](Tutorials/TestAutomation.md)**
 * **[Get telemetry and analytics for your team](Tutorials/Analytics.md)**
diff --git a/Engine/Source/Programs/Horde/Docs/Landing.md b/Engine/Source/Programs/Horde/Docs/Landing.md
--- a/Engine/Source/Programs/Horde/Docs/Landing.md
+++ b/Engine/Source/Programs/Horde/Docs/Landing.md
@@ -5,7 +5,9 @@
 **Horde** has been installed successfully. To continue setting up your deployment, check out the following guides:
 
 * **[Install the Horde agent](Tutorials/InstallHorde.md)**
-* **[Enable remote C++ compilation with Unreal Build Accelerator](Tutorials/RemoteCompilation.md)**
+* Remote compilation with Unreal Build Accelerator
+  * **[Enable remote C++ compilation](Tutorials/RemoteCompilation.md)**
+  * **[Enable remote shader compilation](Tutorials/RemoteShaderCompilation.md)**
 * **[Set up build automation](Tutorials/BuildAutomation.md)**
 * **[Enable test automation with Gauntlet](Tutorials/TestAutomation.md)**
 * **[Get telemetry and analytics for your team](Tutorials/Analytics.md)**
