Horde: Merging 5.5 doc updates and dashboard changes.

#rnx

diff --git a/Engine/Source/Programs/Horde/HordeServer/Accounts/AccountCollection.cs b/Engine/Source/Programs/Horde/HordeServer/Accounts/AccountCollection.cs
--- a/Engine/Source/Programs/Horde/HordeServer/Accounts/AccountCollection.cs
+++ b/Engine/Source/Programs/Horde/HordeServer/Accounts/AccountCollection.cs
@@ -183,7 +183,7 @@ namespace HordeServer.Accounts
 				}
 				else
 				{
-					return PasswordSalt != null && PasswordHash != null && PasswordHasher.ValidatePassword(password, PasswordHasher.SaltFromString(PasswordSalt), PasswordHasher.HashFromString(PasswordHash));
+					return PasswordSalt != null && PasswordHasher.ValidatePassword(password, PasswordHasher.SaltFromString(PasswordSalt), PasswordHasher.HashFromString(PasswordHash));
 				}
 			}
 
