Horde: Additional docs for 5.5 release.

#rnx

diff --git a/Engine/Source/Programs/Horde/README.md b/Engine/Source/Programs/Horde/README.md
--- a/Engine/Source/Programs/Horde/README.md
+++ b/Engine/Source/Programs/Horde/README.md
@@ -32,7 +32,9 @@ See also: [Feature Status](Docs/Features.md)
 ## Getting Started
 
 * **[Install the Horde agent](Docs/Tutorials/InstallHorde.md)**
-* **[Enable remote C++ compilation with Unreal Build Accelerator](Docs/Tutorials/RemoteCompilation.md)**
+* Remote compilation with Unreal Build Accelerator
+  * **[Enable remote C++ compilation](Docs/Tutorials/RemoteCompilation.md)**
+  * **[Enable remote shader compilation](Docs/Tutorials/RemoteShaderCompilation.md)**
 * **[Set up build automation](Docs/Tutorials/BuildAutomation.md)**
 * **[Enable test automation with Gauntlet](Docs/Tutorials/TestAutomation.md)**
 * **[Get telemetry and analytics for your team](Docs/Tutorials/Analytics.md)**
