Horde: Flatten the tutorial list in documentation.

#rnx

diff --git a/Engine/Source/Programs/Horde/Docs/Home.md b/Engine/Source/Programs/Horde/Docs/Home.md
--- a/Engine/Source/Programs/Horde/Docs/Home.md
+++ b/Engine/Source/Programs/Horde/Docs/Home.md
@@ -3,9 +3,8 @@
 ## Getting Started
 
 * **[Install the Horde agent](Tutorials/InstallHorde.md)**
-* Remote compilation with Unreal Build Accelerator
-  * **[Enable remote C++ compilation](Tutorials/RemoteCompilation.md)**
-  * **[Enable remote shader compilation](Tutorials/RemoteShaderCompilation.md)**
+* **[Enable remote C++ compilation](Tutorials/RemoteCompilation.md)**
+* **[Enable remote shader compilation](Tutorials/RemoteShaderCompilation.md)**
 * **[Set up build automation](Tutorials/BuildAutomation.md)**
 * **[Enable test automation with Gauntlet](Tutorials/TestAutomation.md)**
 * **[Get telemetry and analytics for your team](Tutorials/Analytics.md)**
diff --git a/Engine/Source/Programs/Horde/Docs/Landing.md b/Engine/Source/Programs/Horde/Docs/Landing.md
--- a/Engine/Source/Programs/Horde/Docs/Landing.md
+++ b/Engine/Source/Programs/Horde/Docs/Landing.md
@@ -5,9 +5,8 @@
 **Horde** has been installed successfully. To continue setting up your deployment, check out the following guides:
 
 * **[Install the Horde agent](Tutorials/InstallHorde.md)**
-* Remote compilation with Unreal Build Accelerator
-  * **[Enable remote C++ compilation](Tutorials/RemoteCompilation.md)**
-  * **[Enable remote shader compilation](Tutorials/RemoteShaderCompilation.md)**
+* **[Enable remote C++ compilation](Tutorials/RemoteCompilation.md)**
+* **[Enable remote shader compilation](Tutorials/RemoteShaderCompilation.md)**
 * **[Set up build automation](Tutorials/BuildAutomation.md)**
 * **[Enable test automation with Gauntlet](Tutorials/TestAutomation.md)**
 * **[Get telemetry and analytics for your team](Tutorials/Analytics.md)**
@@ -15,12 +14,4 @@
 * **[Install UnrealGameSync and distribute Unreal Editor to your team](Tutorials/UnrealGameSync.md)**
 * **[Enable authentication](Tutorials/Authentication.md)**
 
-Alternatively, check out the full documentation through the `Help` menu at the top right of the page.
-
-## Status
-
-Horde is under heavy development, and large parts of it are still in flux. While we use aspects of it
-(particularly the CI system) heavily at Epic, we consider it experimental for Unreal Engine licensees
-and offer limited support for it.
-
-See also: [Feature Status](Features.md)
+Alternatively, open the full documentation through the `Help` menu at the top right of the page.
diff --git a/Engine/Source/Programs/Horde/README.md b/Engine/Source/Programs/Horde/README.md
--- a/Engine/Source/Programs/Horde/README.md
+++ b/Engine/Source/Programs/Horde/README.md
@@ -32,9 +32,8 @@ See also: [Feature Status](Docs/Features.md)
 ## Getting Started
 
 * **[Install the Horde agent](Docs/Tutorials/InstallHorde.md)**
-* Remote compilation with Unreal Build Accelerator
-  * **[Enable remote C++ compilation](Docs/Tutorials/RemoteCompilation.md)**
-  * **[Enable remote shader compilation](Docs/Tutorials/RemoteShaderCompilation.md)**
+* **[Enable remote C++ compilation](Docs/Tutorials/RemoteCompilation.md)**
+* **[Enable remote shader compilation](Docs/Tutorials/RemoteShaderCompilation.md)**
 * **[Set up build automation](Docs/Tutorials/BuildAutomation.md)**
 * **[Enable test automation with Gauntlet](Docs/Tutorials/TestAutomation.md)**
 * **[Get telemetry and analytics for your team](Docs/Tutorials/Analytics.md)**
