Horde: Reload job pages to refresh and keep tab memory down
#rnx

diff --git a/Engine/Source/Programs/Horde/HordeDashboard/src/components/jobDetailsV2/JobDetailViewV2.tsx b/Engine/Source/Programs/Horde/HordeDashboard/src/components/jobDetailsV2/JobDetailViewV2.tsx
--- a/Engine/Source/Programs/Horde/HordeDashboard/src/components/jobDetailsV2/JobDetailViewV2.tsx
+++ b/Engine/Source/Programs/Horde/HordeDashboard/src/components/jobDetailsV2/JobDetailViewV2.tsx
@@ -53,12 +53,12 @@ const JobBreadCrumbs: React.FC<{ jobDetails: JobDetailsV2 }> = observer(({ jobDe
       if (!jobDetails.jobError) {
          return <Breadcrumbs items={[{ text: "Loading Job" }]} title={"Loading Job"} spinner={true} />
       } else {
-         return <Breadcrumbs items={[{ text: jobDetails.jobError }]} title={jobDetails.jobError} spinner={false} />               
+         return <Breadcrumbs items={[{ text: jobDetails.jobError }]} title={jobDetails.jobError} spinner={false} />
       }
-   } 
+   }
 
    if (jobDetails.jobError) {
-      return <Breadcrumbs items={[{ text: jobDetails.jobError }]} title={jobDetails.jobError} spinner={false} />               
+      return <Breadcrumbs items={[{ text: jobDetails.jobError }]} title={jobDetails.jobError} spinner={false} />
    }
 
    const crumbItems: BreadcrumbItem[] = [];
@@ -395,8 +395,14 @@ export const JobDetailViewV2: React.FC = () => {
    let { jobId } = useParams<{ jobId: string }>();
    const [state, setState] = useState<{ jobDetails?: JobDetailsV2 }>({});
 
+
    useEffect(() => {
+
+      // reload the page once an hour
+      const interval = setInterval(() => location.reload(), 3600000);
+
       return () => {
+         clearInterval(interval);
          state.jobDetails?.clear();
          state.jobDetails = undefined;
       };
