Horde: Tweak user inactivity reload
#rnx

diff --git a/Engine/Source/Programs/Horde/HordeDashboard/src/backend/UserInactivity.ts b/Engine/Source/Programs/Horde/HordeDashboard/src/backend/UserInactivity.ts
--- a/Engine/Source/Programs/Horde/HordeDashboard/src/backend/UserInactivity.ts
+++ b/Engine/Source/Programs/Horde/HordeDashboard/src/backend/UserInactivity.ts
@@ -11,7 +11,8 @@ class UserInactivity {
     }
 
     start() {
-        // Not listening for 'mousemove' as it's fired a lot and can affect performance
+        
+        window.addEventListener('mousemove', this.eventHandler);
         window.addEventListener('mouseenter', this.eventHandler);
         window.addEventListener('mouseleave', this.eventHandler);
         window.addEventListener('touchstart', this.eventHandler);
diff --git a/Engine/Source/Programs/Horde/HordeDashboard/src/components/jobDetailsV2/JobDetailViewV2.tsx b/Engine/Source/Programs/Horde/HordeDashboard/src/components/jobDetailsV2/JobDetailViewV2.tsx
--- a/Engine/Source/Programs/Horde/HordeDashboard/src/components/jobDetailsV2/JobDetailViewV2.tsx
+++ b/Engine/Source/Programs/Horde/HordeDashboard/src/components/jobDetailsV2/JobDetailViewV2.tsx
@@ -18,6 +18,7 @@ import { StepsPanelV2 } from "./JobDetailViewSteps";
 import { SummaryPanel } from "./JobDetailViewSummary";
 import { JobOperations } from "./JobOperationsBar";
 import { StepDetailView } from "./StepDetailView";
+import userInactivity from "horde/backend/UserInactivity";
 
 class BreadcrumbDataView extends JobDataView {
 
@@ -398,8 +399,17 @@ export const JobDetailViewV2: React.FC = () => {
 
    useEffect(() => {
 
-      // reload the page once an hour
-      const interval = setInterval(() => location.reload(), 3600000);
+      function reloadPage() {
+
+         // If haven't touched page in 10 minutes reload it
+         if (userInactivity.getSecondsSinceLastActivity() > 600) {
+            location.reload();
+         }
+
+      }
+
+      // reload the page once every 2 hours of inactivity
+      const interval = setInterval(() => reloadPage(), 3600000 * 2);
 
       return () => {
          clearInterval(interval);
