Horde: Push out the reload logic further
#rnx

diff --git a/Engine/Source/Programs/Horde/HordeDashboard/src/components/jobDetailsV2/JobDetailViewV2.tsx b/Engine/Source/Programs/Horde/HordeDashboard/src/components/jobDetailsV2/JobDetailViewV2.tsx
--- a/Engine/Source/Programs/Horde/HordeDashboard/src/components/jobDetailsV2/JobDetailViewV2.tsx
+++ b/Engine/Source/Programs/Horde/HordeDashboard/src/components/jobDetailsV2/JobDetailViewV2.tsx
@@ -408,8 +408,8 @@ export const JobDetailViewV2: React.FC = () => {
 
       }
 
-      // reload the page once every 2 hours of inactivity
-      const interval = setInterval(() => reloadPage(), 3600000 * 2);
+      // reload the page once every 8 hours of inactivity
+      const interval = setInterval(() => reloadPage(), 3600000 * 8);
 
       return () => {
          clearInterval(interval);
