[UE-141683] Right Click for extra Agent Info options isn't communicated in the UI, leads to support overhead

Improves UX communication by adding a Pencil Icon Button to column header which also displays the context menu when clicked.

diff --git a/Engine/Source/Programs/Horde/HordeDashboard/src/components/AgentView.tsx b/Engine/Source/Programs/Horde/HordeDashboard/src/components/AgentView.tsx
--- a/Engine/Source/Programs/Horde/HordeDashboard/src/components/AgentView.tsx
+++ b/Engine/Source/Programs/Horde/HordeDashboard/src/components/AgentView.tsx
@@ -1,5 +1,5 @@
 // Copyright Epic Games, Inc. All Rights Reserved.
-import { Checkbox, CommandButton, ConstrainMode, ContextualMenu, DefaultButton, DetailsHeader, DetailsList, DetailsListLayoutMode, Dialog, DialogType, DirectionalHint, Dropdown, FontIcon, FontSizes, FontWeights, IBasePickerProps, IColumn, IContextualMenuItem, IContextualMenuProps, IDetailsHeaderProps, IDetailsHeaderStyles, IDetailsListProps, ITag, ITagItemStyles, ITooltipHostStyles, Icon, IconButton, Pivot, PivotItem, PrimaryButton, ProgressIndicator, Link as ReactLink, ScrollablePane, ScrollbarVisibility, SearchBox, Selection, SelectionMode, Slider, Spinner, SpinnerSize, Stack, Sticky, StickyPositionType, TagItem, TagPicker, Text, TextField, mergeStyleSets } from '@fluentui/react';
+import { Checkbox, CommandButton, ConstrainMode, ContextualMenu, DefaultButton, DetailsHeader, DetailsList, DetailsListLayoutMode, Dialog, DialogType, DirectionalHint, Dropdown, FontIcon, FontSizes, FontWeights, IBasePickerProps, IColumn, IContextualMenuItem, IContextualMenuProps, IDetailsHeaderProps, IDetailsHeaderStyles, IDetailsListProps, ITag, ITagItemStyles, ITooltipHostStyles, Icon, IconButton, Pivot, PivotItem, PrimaryButton, ProgressIndicator, Link as ReactLink, ScrollablePane, ScrollbarVisibility, SearchBox, Selection, SelectionMode, Slider, Spinner, SpinnerSize, Stack, Sticky, StickyPositionType, TagItem, TagPicker, Text, TextField, TooltipHost, mergeStyleSets } from '@fluentui/react';
 import { action, makeObservable, observable } from 'mobx';
 import { observer } from 'mobx-react-lite';
 import moment from 'moment-timezone';
@@ -467,6 +467,15 @@ class LocalState {
       this.columnsState.forEach(colState => {
 
          let colSize = colState.colSize;
+         
+         // True when colState is for a column type that can be selected for the last column
+         const bIsSwapColumn: boolean = (
+            colState.key === "software"
+         || colState.key === "storage"
+         || colState.key.indexOf("systemInfo") !== -1
+         || colState.key === "taskTime"
+         || colState.key === "comment"
+         );
 
          if (!this.agentView && colState.displayText === "Pools") {
             colSize = 480;
@@ -481,6 +490,30 @@ class LocalState {
             isSorted: colState.key === "pools" ? undefined : colState.isSorted,
             isSortedDescending: colState.key === "pools" ? undefined : colState.isSortedDescending,
             onColumnClick: this._onColumnClick.bind(this),
+            styles: bIsSwapColumn ? {
+                  cellName: {
+                     flexGrow: '1' // Grow the cell name element so there can be space between it's elements
+                  }
+            } : undefined,
+            onRenderHeader: bIsSwapColumn ? () => {
+               return <Stack horizontal>
+                  <span style={{flexGrow: 1, textAlign: 'center'}}>{colState.displayText}</span>
+                  <TooltipHost
+                     content="Change column content"
+                  >
+                     <IconButton
+                        iconProps={{ iconName: 'Edit' }}
+                        onClick={(ev) => {
+                           ev.preventDefault();
+                           ev.stopPropagation(); // Consume the input before it causes a sort re-order
+                           localState.contextMenuTargetRef = ev?.nativeEvent!;
+                           localState.setHeaderContextMenuOpen(true);
+                        }
+                     }/>
+                  </TooltipHost>
+               </Stack>
+            }
+            : undefined,
             onRender: colState.key === "editAgent" ? () => {
                return <Stack style={{ cursor: "pointer" }} verticalFill verticalAlign='center' horizontalAlign='start' onClick={(ev) => {
                   localState.setRightClickDiv(ev?.clientX, ev?.clientY);
