Horde: Add custom ACL actions for updating aliases, and allow updating aliases with namespace/path claims.

#rnx

diff --git a/Engine/Source/Programs/Horde/Plugins/Storage/HordeServer.Storage/Storage/StorageAclAction.cs b/Engine/Source/Programs/Horde/Plugins/Storage/HordeServer.Storage/Storage/StorageAclAction.cs
--- a/Engine/Source/Programs/Horde/Plugins/Storage/HordeServer.Storage/Storage/StorageAclAction.cs
+++ b/Engine/Source/Programs/Horde/Plugins/Storage/HordeServer.Storage/Storage/StorageAclAction.cs
@@ -19,6 +19,21 @@ namespace HordeServer.Storage
 		/// </summary>
 		public static AclAction WriteBlobs { get; } = new AclAction("WriteBlobs");
 
+		/// <summary>
+		/// Ability to read aliases from the storage service
+		/// </summary>
+		public static AclAction ReadAliases { get; } = new AclAction("ReadAliases");
+
+		/// <summary>
+		/// Ability to write aliases to the storage service
+		/// </summary>
+		public static AclAction WriteAliases { get; } = new AclAction("WriteAliases");
+
+		/// <summary>
+		/// Ability to write aliases to the storage service
+		/// </summary>
+		public static AclAction DeleteAliases { get; } = new AclAction("DeleteAliases");
+
 		/// <summary>
 		/// Ability to read refs from the storage service
 		/// </summary>
diff --git a/Engine/Source/Programs/Horde/Plugins/Storage/HordeServer.Storage/Storage/StorageController.cs b/Engine/Source/Programs/Horde/Plugins/Storage/HordeServer.Storage/Storage/StorageController.cs
--- a/Engine/Source/Programs/Horde/Plugins/Storage/HordeServer.Storage/Storage/StorageController.cs
+++ b/Engine/Source/Programs/Horde/Plugins/Storage/HordeServer.Storage/Storage/StorageController.cs
@@ -45,6 +45,23 @@ namespace HordeServer.Storage
 			return _storageConfig.Value.TryGetNamespace(namespaceId, out NamespaceConfig? namespaceConfig) && namespaceConfig.Authorize(action, User);
 		}
 
+		bool AuthorizePaths(NamespaceId namespaceId, AclAction action, IEnumerable<string> paths)
+		{
+			if (!paths.Any())
+			{
+				return true;
+			}
+			if (Authorize(namespaceId, action))
+			{
+				return true;
+			}
+			if (paths.All(x => HasPathClaim(User, HordeClaimTypes.WriteNamespace, namespaceId, x)))
+			{
+				return true;
+			}
+			return false;
+		}
+
 		/// <summary>
 		/// Update aliases in the storage service
 		/// </summary>
@@ -60,11 +77,26 @@ namespace HordeServer.Storage
 			{
 				return NotFound(namespaceId);
 			}
-			if (!Authorize(namespaceId, StorageAclAction.WriteRefs))
+
+			// Authorize access to the target paths
+			if (!AuthorizePaths(namespaceId, StorageAclAction.WriteAliases, request.AddAliases.Select(x => x.Target.ToString())))
+			{
+				return Forbid(StorageAclAction.WriteAliases, namespaceId);
+			}
+			if (!AuthorizePaths(namespaceId, StorageAclAction.DeleteAliases, request.RemoveAliases.Select(x => x.Target.ToString())))
+			{
+				return Forbid(StorageAclAction.DeleteAliases, namespaceId);
+			}
+			if (!AuthorizePaths(namespaceId, StorageAclAction.WriteRefs, request.AddRefs.Select(x => x.RefName.ToString())))
 			{
 				return Forbid(StorageAclAction.WriteRefs, namespaceId);
 			}
+			if (!AuthorizePaths(namespaceId, StorageAclAction.DeleteRefs, request.RemoveRefs.Select(x => x.RefName.ToString())))
+			{
+				return Forbid(StorageAclAction.DeleteRefs, namespaceId);
+			}
 
+			// Execute the requests
 			foreach (AddAliasRequest addAlias in request.AddAliases)
 			{
 				await backend.AddAliasAsync(addAlias.Name, addAlias.Target, addAlias.Rank, addAlias.Data ?? Array.Empty<byte>(), cancellationToken);
