Horde: Remove UE 5.4 related note saying UBA only supports Windows

diff --git a/Engine/Source/Programs/Horde/Docs/Tutorials/RemoteCompilation.md b/Engine/Source/Programs/Horde/Docs/Tutorials/RemoteCompilation.md
--- a/Engine/Source/Programs/Horde/Docs/Tutorials/RemoteCompilation.md
+++ b/Engine/Source/Programs/Horde/Docs/Tutorials/RemoteCompilation.md
@@ -18,8 +18,6 @@ remote machine behind the scenes as necessary.
 **Unreal Build Tool** can use Unreal Build Accelerator with Horde to offload compilation tasks to connected agents,
 spreading the workload over multiple machines.
 
-> **Note:** Unreal Build Accelerator only supports Windows in Unreal Engine 5.4. Support for Mac and Linux are planned for 
-  a future release.
 
 ## Prerequisites
 
@@ -77,7 +75,7 @@ spreading the workload over multiple machines.
    This indicates work is being spread to multiple agents. If you enabled the UBA visualizer, you can also see
    a graphical overview of how the build progresses over multiple machines.
 
-   For debugging and tuning purposes, it can be useful to force remote execution all compile workloads. To do
+   For debugging and tuning purposes, it can be useful to force remote execution for all compile workloads. To do
    so, enable the following option in your `BuildConfiguration.xml` file or pass `-UBAForceRemote` on the
    UnrealBuildTool command line:
 
