Horde: Minor doc formatting updates.

#rnx

diff --git a/Engine/Source/Programs/Horde/Docs/Features.md b/Engine/Source/Programs/Horde/Docs/Features.md
--- a/Engine/Source/Programs/Horde/Docs/Features.md
+++ b/Engine/Source/Programs/Horde/Docs/Features.md
@@ -5,7 +5,7 @@
 | Feature | Status | Notes/Roadmap |
 | ------- | ------ | ----- |
 | Analytics | 🟧&nbsp;Experimental | - |
-| Artifact&nbsp;Distribution | 🟩&nbsp;Production Ready | - |
+| Artifact&nbsp;Distribution | 🟨&nbsp;Beta | Needs better integration with other tooling.<br>Improved chunking and upload performance. |
 | Automation Hub | 🟩&nbsp;Production Ready | - |
 | Build Automation | 🟩&nbsp;Production Ready | Migrate jobs to a more lightweight data model.<br>Programmable state machine for jobs. |
 | Build Health | 🟩&nbsp;Production Ready | - |
diff --git a/Engine/Source/Programs/Horde/Docs/Tutorials/BuildAutomation.md b/Engine/Source/Programs/Horde/Docs/Tutorials/BuildAutomation.md
--- a/Engine/Source/Programs/Horde/Docs/Tutorials/BuildAutomation.md
+++ b/Engine/Source/Programs/Horde/Docs/Tutorials/BuildAutomation.md
@@ -67,7 +67,7 @@ ensuring that the state of a project is continuously being monitored and that bu
 
 ## Default Jobs
 
-The example `ue5-release-5.4.stream.json` file configures the appearance of its page in the Horde Dashboard, as
+The example `ue5-release-5.5.stream.json` file configures the appearance of its page in the Horde Dashboard, as
 well as job templates and agent types.
 
 * A **job template** defines a set of parameters that are used to construct a
diff --git a/Engine/Source/Programs/Horde/Docs/Tutorials/RemoteCompilation.md b/Engine/Source/Programs/Horde/Docs/Tutorials/RemoteCompilation.md
--- a/Engine/Source/Programs/Horde/Docs/Tutorials/RemoteCompilation.md
+++ b/Engine/Source/Programs/Horde/Docs/Tutorials/RemoteCompilation.md
@@ -18,7 +18,6 @@ remote machine behind the scenes as necessary.
 **Unreal Build Tool** can use Unreal Build Accelerator with Horde to offload compilation tasks to connected agents,
 spreading the workload over multiple machines.
 
-
 ## Prerequisites
 
 * Horde Server and one or more Horde Agents (see [Getting Started: Install Horde](InstallHorde.md)).
@@ -103,28 +102,29 @@ an environment variable for each job step (`UE_HORDE_TOKEN`).
 
 Below is an updated global config with a `default` cluster and `AddComputeTasks` granted for letting UBT and UBA schedule remote compilation.
 To see what claims are available to your users, open `/api/v1/user/claims` as a logged in user.
-```json
-{
-  // ...
-  "plugins": {
-    // ...
-    "compute": {
+
+    ```json
+    {
       // ...
-      "clusters": [
-        {
-          "id": "default",
-          "namespaceid": "horde.compute",
-          "acl": {
-            "entries": [
-              {
-                "claim": { "type": "http://epicgames.com/ue/horde/user", "value": "jane.smith" },
-                "actions": ["AddComputeTasks"]
+      "plugins": {
+        // ...
+        "compute": {
+          // ...
+          "clusters": [
+            {
+              "id": "default",
+              "namespaceid": "horde.compute",
+              "acl": {
+                "entries": [
+                  {
+                    "claim": { "type": "http://epicgames.com/ue/horde/user", "value": "jane.smith" },
+                    "actions": ["AddComputeTasks"]
+                  }
+                ]
               }
-            ]
-          }
+            }
+          ]
         }
-      ]
+      }
     }
-  }
-}
-```
\ No newline at end of file
+    ```
