Horde: Make LeaseManager testable in agent

Previously used FakeHordeServer reinstantiated and session result now contain a reason leading to the outcome.

diff --git a/Engine/Source/Programs/Horde/HordeAgent.Tests/FakeHordeServer.cs b/Engine/Source/Programs/Horde/HordeAgent.Tests/FakeHordeServer.cs
new file mode 100644
--- /dev/null
+++ b/Engine/Source/Programs/Horde/HordeAgent.Tests/FakeHordeServer.cs
@@ -0,0 +1,275 @@
+﻿// Copyright Epic Games, Inc. All Rights Reserved.
+
+using System;
+using System.Collections.Generic;
+using System.Linq;
+using System.Threading;
+using System.Threading.Channels;
+using System.Threading.Tasks;
+using EpicGames.Horde;
+using EpicGames.Horde.Agents.Leases;
+using EpicGames.Horde.Artifacts;
+using EpicGames.Horde.Compute;
+using EpicGames.Horde.Logs;
+using EpicGames.Horde.Projects;
+using EpicGames.Horde.Secrets;
+using EpicGames.Horde.Storage;
+using EpicGames.Horde.Tools;
+using Google.Protobuf.WellKnownTypes;
+using Grpc.Core;
+using HordeCommon.Rpc;
+using HordeCommon.Rpc.Messages;
+using HordeCommon.Rpc.Tasks;
+using Microsoft.Extensions.Logging;
+
+namespace HordeAgent.Tests;
+
+internal class FakeHordeRpcClient(FakeHordeRpcServer outer) : HordeRpc.HordeRpcClient
+{
+	public override AsyncDuplexStreamingCall<RpcUpdateSessionRequest, RpcUpdateSessionResponse> UpdateSession(Metadata headers = null!,
+		DateTime? deadline = null, CancellationToken cancellationToken = default)
+	{
+		return outer.GetUpdateSessionCall(CancellationToken.None);
+	}
+}
+
+internal class FakeHordeClient(FakeHordeRpcServer server) : IHordeClient
+{
+	public Uri ServerUrl => new ("http://fake-horde-server");
+	public IArtifactCollection Artifacts => throw new NotImplementedException();
+	public IComputeClient Compute => throw new NotImplementedException();
+	public IProjectCollection Projects => throw new NotImplementedException();
+	public ISecretCollection Secrets => throw new NotImplementedException();
+	public IToolCollection Tools => throw new NotImplementedException();
+	public event Action? OnAccessTokenStateChanged
+	{
+		add { }
+		remove { }
+	}
+	public Task<bool> LoginAsync(bool allowLogin, CancellationToken cancellationToken) => throw new NotImplementedException();
+	public bool HasValidAccessToken() => true;
+	public Task<string?> GetAccessTokenAsync(bool interactive, CancellationToken cancellationToken = default) => Task.FromResult<string?>(null);
+	public Task<TClient> CreateGrpcClientAsync<TClient>(CancellationToken cancellationToken = default) where TClient : ClientBase<TClient>
+	{
+		if (typeof(TClient) == typeof(HordeRpc.HordeRpcClient))
+		{
+			return Task.FromResult((TClient)(object)new FakeHordeRpcClient(server));
+		}
+		throw new NotImplementedException($"No support for gRPC client {typeof(TClient)}");
+	}
+	
+	public HordeHttpClient CreateHttpClient() => throw new NotImplementedException();
+	public IStorageNamespace GetStorageNamespace(string relativePath, string? accessToken = null) => throw new NotImplementedException();
+	public IServerLogger CreateServerLogger(LogId logId, LogLevel minimumLevel = LogLevel.Information) => throw new NotImplementedException();
+}
+
+/// <summary>
+/// Fake implementation of a HordeRpc gRPC server.
+/// Provides a corresponding gRPC client class that can be used with to test client-server interactions inside a single process.
+/// </summary>
+internal class FakeHordeRpcServer
+{
+	public readonly TaskCompletionSource<bool> CreateSessionReceived = new();
+	public readonly TaskCompletionSource<bool> UpdateSessionReceived = new();
+
+	private readonly Dictionary<LeaseId, RpcLease> _leases = new();
+	private readonly FakeHordeClient _hordeClient;
+	private readonly ILogger<FakeHordeRpcServer> _logger;
+	
+	public FakeHordeRpcServer(ILogger<FakeHordeRpcServer> logger)
+	{
+		_logger = logger;
+		_hordeClient = new FakeHordeClient(this);
+	}
+	
+	public void ScheduleTestLease(LeaseId? leaseId = null)
+	{
+		if (leaseId == null)
+		{
+			Span<byte> randomBytes = stackalloc byte[12];
+			Random.Shared.NextBytes(randomBytes);
+			leaseId = new LeaseId(new BinaryId(randomBytes));
+		}
+		
+		if (_leases.ContainsKey(leaseId.Value))
+		{
+			throw new ArgumentException($"Lease ID {leaseId.Value} already exists");
+		}
+		
+		TestTask testTask = new();
+		_leases[leaseId.Value] = new RpcLease { Id = leaseId.Value, State = RpcLeaseState.Pending, Payload = Any.Pack(testTask) };
+	}
+	
+	public RpcLease GetLease(LeaseId leaseId)
+	{
+		return _leases[leaseId];
+	}
+	
+	public IHordeClient GetHordeClient()
+	{
+		return _hordeClient;
+	}
+	
+	public RpcCreateSessionResponse OnCreateSessionRequest(RpcCreateSessionRequest request)
+	{
+		CreateSessionReceived.TrySetResult(true);
+		_logger.LogInformation("OnCreateSessionRequest: {AgentId}", request.Id);
+		RpcCreateSessionResponse response = new()
+		{
+			AgentId = "bogusAgentId",
+			Token = "bogusToken",
+			SessionId = "bogusSessionId",
+			ExpiryTime = Timestamp.FromDateTime(DateTime.UtcNow.AddHours(3)),
+		};
+		
+		return response;
+	}
+	
+	public AsyncDuplexStreamingCall<RpcUpdateSessionRequest, RpcUpdateSessionResponse> GetUpdateSessionCall(CancellationToken cancellationToken)
+	{
+		FakeAsyncStreamReader<RpcUpdateSessionResponse> responseStream = new(cancellationToken);
+		
+		async Task OnRequest(RpcUpdateSessionRequest request)
+		{
+			UpdateSessionReceived.TrySetResult(true);
+			
+			foreach (RpcLease agentLease in request.Leases)
+			{
+				RpcLease serverLease = _leases[agentLease.Id];
+				serverLease.State = agentLease.State;
+				serverLease.Outcome = agentLease.Outcome;
+				serverLease.Output = agentLease.Output;
+			}
+			
+			_logger.LogInformation("OnUpdateSessionRequest: {AgentId} {SessionId} {Status}", request.AgentId, request.SessionId, request.Status);
+			await Task.Delay(100, cancellationToken);
+			RpcUpdateSessionResponse response = new() { ExpiryTime = Timestamp.FromDateTime(DateTime.UtcNow + TimeSpan.FromMinutes(120)) };
+			response.Leases.AddRange(_leases.Values.Where(x => x.State != RpcLeaseState.Completed));
+			await responseStream.Write(response);
+		}
+		
+		FakeClientStreamWriter<RpcUpdateSessionRequest> requestStream = new(OnRequest, () =>
+		{
+			responseStream.Complete();
+			return Task.CompletedTask;
+		});
+		
+		return new(
+			requestStream,
+			responseStream,
+			Task.FromResult(new Metadata()),
+			() => Status.DefaultSuccess,
+			() => new Metadata(),
+			() => { });
+	}
+}
+
+/// <summary>
+/// Fake stream reader used for testing gRPC clients
+/// </summary>
+/// <typeparam name="T">Message type reader will handle</typeparam>
+internal class FakeAsyncStreamReader<T> : IAsyncStreamReader<T> where T : class
+{
+	private readonly Channel<T> _channel = System.Threading.Channels.Channel.CreateUnbounded<T>();
+	private T? _current;
+	private readonly CancellationToken? _cancellationTokenOverride;
+	
+	public FakeAsyncStreamReader(CancellationToken? cancellationTokenOverride = null)
+	{
+		_cancellationTokenOverride = cancellationTokenOverride;
+	}
+	
+	public Task Write(T message)
+	{
+		if (!_channel.Writer.TryWrite(message))
+		{
+			throw new InvalidOperationException("Unable to write message.");
+		}
+		
+		return Task.CompletedTask;
+	}
+	
+	public void Complete()
+	{
+		_channel.Writer.Complete();
+	}
+	
+	/// <inheritdoc/>
+	public async Task<bool> MoveNext(CancellationToken cancellationToken)
+	{
+		if (_cancellationTokenOverride != null)
+		{
+			cancellationToken = _cancellationTokenOverride.Value;
+		}
+		
+		if (await _channel.Reader.WaitToReadAsync(cancellationToken))
+		{
+			if (_channel.Reader.TryRead(out T? message))
+			{
+				_current = message;
+				return true;
+			}
+		}
+		
+		_current = null!;
+		return false;
+	}
+	
+	/// <inheritdoc/>
+	public T Current
+	{
+		get
+		{
+			if (_current == null)
+			{
+				throw new InvalidOperationException("No current element is available.");
+			}
+			
+			return _current;
+		}
+	}
+}
+
+/// <summary>
+/// Fake stream writer used for testing gRPC clients
+/// </summary>
+/// <typeparam name="T">Message type writer will handle</typeparam>
+internal class FakeClientStreamWriter<T> : IClientStreamWriter<T> where T : class
+{
+	private readonly Func<T, Task>? _onWrite;
+	private readonly Func<Task>? _onComplete;
+	private bool _isCompleted;
+	
+	public FakeClientStreamWriter(Func<T, Task>? onWrite = null, Func<Task>? onComplete = null)
+	{
+		_onWrite = onWrite;
+		_onComplete = onComplete;
+	}
+	
+	/// <inheritdoc/>
+	public async Task WriteAsync(T message)
+	{
+		if (_isCompleted)
+		{
+			throw new InvalidOperationException("Stream is marked as complete");
+		}
+		
+		if (_onWrite != null)
+		{
+			await _onWrite(message);
+		}
+	}
+	
+	/// <inheritdoc/>
+	public WriteOptions? WriteOptions { get; set; }
+	
+	/// <inheritdoc/>
+	public async Task CompleteAsync()
+	{
+		_isCompleted = true;
+		if (_onComplete != null)
+		{
+			await _onComplete();
+		}
+	}
+}
\ No newline at end of file
diff --git a/Engine/Source/Programs/Horde/HordeAgent.Tests/Leases/LeaseManagerTests.cs b/Engine/Source/Programs/Horde/HordeAgent.Tests/Leases/LeaseManagerTests.cs
new file mode 100644
--- /dev/null
+++ b/Engine/Source/Programs/Horde/HordeAgent.Tests/Leases/LeaseManagerTests.cs
@@ -0,0 +1,134 @@
+// Copyright Epic Games, Inc. All Rights Reserved.
+
+using System;
+using System.Collections.Generic;
+using System.Diagnostics.CodeAnalysis;
+using System.IO;
+using System.Threading;
+using System.Threading.Tasks;
+using EpicGames.Core;
+using EpicGames.Horde;
+using EpicGames.Horde.Agents;
+using EpicGames.Horde.Agents.Leases;
+using EpicGames.Horde.Agents.Sessions;
+using HordeAgent.Leases;
+using HordeAgent.Services;
+using HordeAgent.Tests.Services;
+using HordeCommon.Rpc.Messages;
+using HordeCommon.Rpc.Tasks;
+using Microsoft.Extensions.Logging;
+using Microsoft.Extensions.Logging.Abstractions;
+using Microsoft.VisualStudio.TestTools.UnitTesting;
+using OpenTelemetry.Trace;
+
+namespace HordeAgent.Tests.Leases;
+
+public class CapabilitiesServiceStub(RpcAgentCapabilities? capabilities = null) : ICapabilitiesService
+{
+	public Task<RpcAgentCapabilities> GetCapabilitiesAsync(DirectoryReference? workingDir)
+	{
+		return Task.FromResult(capabilities ?? new RpcAgentCapabilities());
+	}
+}
+
+public class SessionStub(AgentId agentId, SessionId sessionId, DirectoryReference workingDir, IHordeClient hordeClient) : ISession
+{
+	public AgentId AgentId { get; } = agentId;
+	public SessionId SessionId { get; } = sessionId;
+	public DirectoryReference WorkingDir { get; } = workingDir;
+	public IHordeClient HordeClient { get; } = hordeClient;
+	
+	public async ValueTask DisposeAsync()
+	{
+		if (HordeClient is IAsyncDisposable disposableClient)
+		{
+			await disposableClient.DisposeAsync();
+		}
+		GC.SuppressFinalize(this);
+	}
+}
+
+internal class TestHandlerFactory : LeaseHandlerFactory<TestTask>
+{
+	public override LeaseHandler<TestTask> CreateHandler(RpcLease lease)
+	{
+		return new Handler(lease);
+	}
+	
+	private class Handler(RpcLease rpcLease) : LeaseHandler<TestTask>(rpcLease)
+	{
+		protected override Task<LeaseResult> ExecuteAsync(ISession session, LeaseId leaseId, TestTask message, Tracer tracer, ILogger logger, CancellationToken cancellationToken)
+		{
+			logger.LogInformation("Executed TestTask");
+			LeaseResult leaseResult = new((byte[]?)null);
+			return Task.FromResult(leaseResult);
+		}
+	}
+}
+
+[TestClass]
+public sealed class LeaseManagerTests
+{
+	private readonly FakeHordeRpcServer _hordeServer;
+	private readonly LeaseManager _leaseManager;
+	
+	public LeaseManagerTests()
+	{
+		_hordeServer= new FakeHordeRpcServer(CreateConsoleLogger<FakeHordeRpcServer>());
+		_leaseManager = CreateLeaseManager(_hordeServer.GetHordeClient());
+	}
+	
+	[TestMethod]
+	public async Task TerminateGracefully_LeaseIsActive_Async()
+	{
+		using CancellationTokenSource cts = new(5000);
+		_leaseManager.TerminateSessionAfterLease = true;
+		_hordeServer.ScheduleTestLease();
+		Assert.AreEqual(new SessionResult(SessionOutcome.Terminate, SessionReason.Completed), await _leaseManager.RunAsync(cts.Token));
+	}
+	
+	[TestMethod]
+	public async Task TerminateGracefully_NoLeaseActive_Async()
+	{
+		using CancellationTokenSource cts = new(5000);
+		_leaseManager.TerminateSessionAfterLease = true;
+		Assert.AreEqual(new SessionResult(SessionOutcome.Terminate, SessionReason.Completed), await _leaseManager.RunAsync(cts.Token));
+	}
+	
+	/// <summary>
+	/// Create a console logger for tests
+	/// </summary>
+	/// <typeparam name="T">Type to instantiate</typeparam>
+	/// <returns>A logger</returns>
+	private static ILogger<T> CreateConsoleLogger<T>()
+	{
+		using ILoggerFactory loggerFactory = LoggerFactory.Create(builder =>
+		{
+			builder.SetMinimumLevel(LogLevel.Debug);
+			builder.AddSimpleConsole(options => { options.SingleLine = true; });
+		});
+		
+		return loggerFactory.CreateLogger<T>();
+	}
+	
+	[SuppressMessage("Reliability", "CA2000:Dispose objects before losing scope")]
+	private static LeaseManager CreateLeaseManager(IHordeClient hordeClient)
+	{
+		DirectoryReference tempWorkingDir = new(Path.Join(Path.GetTempPath(), Path.GetRandomFileName()));
+		AgentSettings settings = new() { WriteStepOutputToLogger = true };
+		ISession session = new SessionStub(new AgentId("testAgent"), SessionId.Parse("aaaaaaaaaaaaaaaaaaaaaaaa"), tempWorkingDir, hordeClient);
+		TestOptionsMonitor<AgentSettings> settingsOptions = new (settings);
+		StatusService statusService = new(settingsOptions, NullLogger<StatusService>.Instance);
+
+		return new LeaseManager(
+			session,
+			new CapabilitiesServiceStub(),
+			statusService,
+			new DefaultSystemMetrics(),
+			new List<LeaseHandlerFactory>() { new TestHandlerFactory() },
+			new LeaseLoggerFactory(settingsOptions, CreateConsoleLogger<LeaseLoggerFactory>()),
+			settingsOptions,
+			TracerProvider.Default.GetTracer("TestTracer"),
+			CreateConsoleLogger<LeaseManager>());
+	}
+}
diff --git a/Engine/Source/Programs/Horde/HordeAgent/Leases/LeaseManager.cs b/Engine/Source/Programs/Horde/HordeAgent/Leases/LeaseManager.cs
--- a/Engine/Source/Programs/Horde/HordeAgent/Leases/LeaseManager.cs
+++ b/Engine/Source/Programs/Horde/HordeAgent/Leases/LeaseManager.cs
@@ -84,7 +84,7 @@ namespace HordeAgent.Leases
 		internal TimeSpan _rpcConnectionRetryDelay = TimeSpan.FromSeconds(5);
 
 		readonly ISession _session;
-		readonly CapabilitiesService _capabilitiesService;
+		readonly ICapabilitiesService _capabilitiesService;
 		readonly StatusService _statusService;
 		readonly ISystemMetrics _systemMetrics;
 		readonly Dictionary<string, LeaseHandlerFactory> _typeUrlToLeaseHandler;
@@ -97,7 +97,7 @@ namespace HordeAgent.Leases
 
 		public LeaseManager(
 			ISession session,
-			CapabilitiesService capabilitiesService,
+			ICapabilitiesService capabilitiesService,
 			StatusService statusService,
 			ISystemMetrics systemMetrics,
 			IEnumerable<LeaseHandlerFactory> leaseHandlerFactories,
@@ -167,12 +167,12 @@ namespace HordeAgent.Leases
 			catch (OperationCanceledException ex) when (stoppingToken.IsCancellationRequested)
 			{
 				_logger.LogInformation(ex, "Execution cancelled");
-				result = new SessionResult(SessionOutcome.Terminate);
+				result = new SessionResult(SessionOutcome.Terminate, SessionReason.Cancelled);
 			}
 			catch (Exception ex)
 			{
 				_logger.LogError(ex, "Exception while executing session. Restarting.");
-				result = new SessionResult(SessionOutcome.BackOff);
+				result = new SessionResult(SessionOutcome.BackOff, SessionReason.Failed);
 			}
 
 			while (_activeLeases.Count > 0)
@@ -281,7 +281,7 @@ namespace HordeAgent.Leases
 				if (stoppingToken.IsCancellationRequested && _sessionResult == null)
 				{
 					_logger.LogInformation("Cancellation from token requested; setting session result to terminate.");
-					_sessionResult = new SessionResult(SessionOutcome.Terminate);
+					_sessionResult = new SessionResult(SessionOutcome.Terminate, SessionReason.Cancelled);
 				}
 
 				bool stopping = false;
@@ -345,7 +345,7 @@ namespace HordeAgent.Leases
 							if (atLeastOneLeaseFinished && TerminateSessionAfterLease)
 							{
 								_logger.LogInformation("At least one lease executed. Terminating session");
-								_sessionResult = new SessionResult(SessionOutcome.Terminate);
+								_sessionResult = new SessionResult(SessionOutcome.Terminate, SessionReason.Completed);
 							}
 
 							PoolIds = updateSessionResponse.PoolIds;
@@ -393,7 +393,7 @@ namespace HordeAgent.Leases
 							// Update the session result if we've transitioned to stopped
 							if (updateSessionResponse.Status == RpcAgentStatus.Stopped)
 							{
-								SessionResult result = _sessionResult ?? new SessionResult(SessionOutcome.BackOff);
+								SessionResult result = _sessionResult ?? new SessionResult(SessionOutcome.BackOff, SessionReason.Completed);
 								_logger.LogInformation("Agent status is stopped; returning from session update loop with result {Result}", result.Outcome);
 								return result;
 							}
@@ -402,7 +402,7 @@ namespace HordeAgent.Leases
 						// If there's nothing still running and cancellation was requested, exit
 						if (_activeLeases.Count == 0 && _sessionResult != null)
 						{
-							_logger.LogInformation("No leases are active. Agent is stopping. Session result {Result}", _sessionResult.Outcome); // TODO: Should not really hit this any more; server should report stopping state in lease update above.
+							_logger.LogInformation("No leases are active. Agent is stopping. Session {SessionResult}", _sessionResult); // TODO: Should not really hit this any more; server should report stopping state in lease update above.
 							return _sessionResult;
 						}
 					}
diff --git a/Engine/Source/Programs/Horde/HordeAgent/Services/CapabilitiesService.cs b/Engine/Source/Programs/Horde/HordeAgent/Services/CapabilitiesService.cs
--- a/Engine/Source/Programs/Horde/HordeAgent/Services/CapabilitiesService.cs
+++ b/Engine/Source/Programs/Horde/HordeAgent/Services/CapabilitiesService.cs
@@ -22,10 +22,20 @@ using AddressFamily = System.Net.Sockets.AddressFamily;
 
 namespace HordeAgent.Services
 {
+	internal interface ICapabilitiesService
+	{
+		/// <summary>
+		/// Gets the hardware capabilities of this worker
+		/// </summary>
+		/// <param name="workingDir">Working directory for the agent</param>
+		/// <returns>Worker object for advertising to the server</returns>
+		Task<RpcAgentCapabilities> GetCapabilitiesAsync(DirectoryReference? workingDir);
+	}
+	
 	/// <summary>
 	/// Implements the message handling loop for an agent. Runs asynchronously until disposed.
 	/// </summary>
-	class CapabilitiesService
+	class CapabilitiesService : ICapabilitiesService
 	{
 		static readonly DateTimeOffset s_startTime = DateTimeOffset.Now;
 		static readonly DateTimeOffset s_bootTime = DateTimeOffset.Now - TimeSpan.FromTicks(Environment.TickCount64 * TimeSpan.TicksPerMillisecond);
@@ -469,8 +479,6 @@ namespace HordeAgent.Services
 			// Used during the transition period over from multi-platform, non-self-contained agent packages.
 			capabilities.Properties.Add($"SelfContained={AgentApp.IsSelfContained}");
 			capabilities.Properties.Add($"DotNetFramework={RuntimeInformation.FrameworkDescription}");
-			capabilities.Properties.Add($"ProcessArch={RuntimeInformation.ProcessArchitecture}");
-			capabilities.Properties.Add($"OsArch={RuntimeInformation.OSArchitecture}");
 
 			// Add any additional properties from the config file
 			capabilities.Properties.AddRange(_settings.CurrentValue.Properties.Select(kvp => $"{kvp.Key}={kvp.Value}"));
@@ -510,7 +518,7 @@ namespace HordeAgent.Services
 		/// <param name="hostname">A hostname to test against</param>
 		/// <param name="timeoutMs">Max time to wait for a connect, in milliseconds</param>
 		/// <returns>Local IP address of this machine</returns>
-		public async Task<IPAddress?> GetLocalIpAddressAsync(string hostname, int timeoutMs = 2000)
+		private async Task<IPAddress?> GetLocalIpAddressAsync(string hostname, int timeoutMs = 2000)
 		{
 			using TelemetrySpan _ = _tracer.StartActiveSpan(nameof(GetLocalIpAddressAsync));
 			try
diff --git a/Engine/Source/Programs/Horde/HordeAgent/Services/SessionFactory.cs b/Engine/Source/Programs/Horde/HordeAgent/Services/SessionFactory.cs
--- a/Engine/Source/Programs/Horde/HordeAgent/Services/SessionFactory.cs
+++ b/Engine/Source/Programs/Horde/HordeAgent/Services/SessionFactory.cs
@@ -41,7 +41,7 @@ namespace HordeAgent.Services
 		/// <summary>
 		/// Horde client instance
 		/// </summary>
-		HordeClient HordeClient { get; }
+		IHordeClient HordeClient { get; }
 	}
 
 	/// <summary>
@@ -72,12 +72,12 @@ namespace HordeAgent.Services
 		public DirectoryReference WorkingDir { get; }
 
 		/// <inheritdoc/>
-		public HordeClient HordeClient { get; }
+		public IHordeClient HordeClient { get; }
 
 		/// <summary>
 		/// Constructor
 		/// </summary>
-		public Session(AgentId agentId, SessionId sessionId, DirectoryReference workingDir, HordeClient hordeClient)
+		public Session(AgentId agentId, SessionId sessionId, DirectoryReference workingDir, IHordeClient hordeClient)
 		{
 			AgentId = agentId;
 			SessionId = sessionId;
@@ -213,7 +213,10 @@ namespace HordeAgent.Services
 		/// <returns></returns>
 		public async ValueTask DisposeAsync()
 		{
-			await HordeClient.DisposeAsync();
+			if (HordeClient is IAsyncDisposable disposableClient)
+			{
+				await disposableClient.DisposeAsync();
+			}
 		}
 
 		static FileReference GetRegistrationFile()
diff --git a/Engine/Source/Programs/Horde/HordeAgent/Services/SessionResult.cs b/Engine/Source/Programs/Horde/HordeAgent/Services/SessionResult.cs
--- a/Engine/Source/Programs/Horde/HordeAgent/Services/SessionResult.cs
+++ b/Engine/Source/Programs/Horde/HordeAgent/Services/SessionResult.cs
@@ -24,16 +24,42 @@ namespace HordeAgent.Services
 		/// </summary>
 		RunCallback,
 	}
+	
+	/// <summary>
+	/// Reason for outcome of a session
+	/// </summary>
+	enum SessionReason
+	{
+		/// <summary>
+		/// Session completed its operations normally
+		/// </summary>
+		Completed,
+		
+		/// <summary>
+		/// Session encountered an unrecoverable error
+		/// </summary>
+		Failed,
+		
+		/// <summary>
+		/// Session was explicitly cancelled by user or system
+		/// </summary>
+		Cancelled,
+	}
 
 	/// <summary>
 	/// Result from executing a session
 	/// </summary>
-	class SessionResult
+	class SessionResult : IEquatable<SessionResult>
 	{
 		/// <summary>
 		/// The outcome code
 		/// </summary>
 		public SessionOutcome Outcome { get; }
+		
+		/// <summary>
+		/// Reason for the outcome
+		/// </summary>
+		public SessionReason Reason { get; }
 
 		/// <summary>
 		/// Callback for running upgrades
@@ -43,9 +69,10 @@ namespace HordeAgent.Services
 		/// <summary>
 		/// Constructor
 		/// </summary>
-		public SessionResult(SessionOutcome outcome)
+		public SessionResult(SessionOutcome outcome, SessionReason reason)
 		{
 			Outcome = outcome;
+			Reason = reason;
 		}
 
 		/// <summary>
@@ -56,5 +83,25 @@ namespace HordeAgent.Services
 			Outcome = SessionOutcome.RunCallback;
 			CallbackAsync = callbackAsync;
 		}
+		
+		public bool Equals(SessionResult? other)
+		{
+			return Outcome == other?.Outcome && Reason == other.Reason;
+		}
+		
+		public override bool Equals(object? obj)
+		{
+			return obj is SessionResult other && Equals(other);
+		}
+		
+		public override int GetHashCode()
+		{
+			return HashCode.Combine((int)Outcome, (int)Reason);
+		}
+		
+		public override string ToString()
+		{
+			return $"Outcome={Outcome} Reason={Reason}";
+		}
 	}
 }
