Horde: Fix lease manager ignoring session termination requests when no lease is active

diff --git a/Engine/Source/Programs/Horde/HordeAgent/Leases/LeaseManager.cs b/Engine/Source/Programs/Horde/HordeAgent/Leases/LeaseManager.cs
--- a/Engine/Source/Programs/Horde/HordeAgent/Leases/LeaseManager.cs
+++ b/Engine/Source/Programs/Horde/HordeAgent/Leases/LeaseManager.cs
@@ -341,11 +341,14 @@ namespace HordeAgent.Leases
 						// Now reconcile the local state to match what the server reports
 						if (updateSessionResponse != null)
 						{
-							bool atLeastOneLeaseFinished = _activeLeases.Any(x => x.RpcLease.State is RpcLeaseState.Completed or RpcLeaseState.Cancelled);
-							if (atLeastOneLeaseFinished && TerminateSessionAfterLease)
+							if (TerminateSessionAfterLease)
 							{
-								_logger.LogInformation("At least one lease executed. Terminating session");
-								_sessionResult = new SessionResult(SessionOutcome.Terminate, SessionReason.Completed);
+								bool atLeastOneLeaseFinished = _activeLeases.Any(x => x.RpcLease.State is RpcLeaseState.Completed or RpcLeaseState.Cancelled);
+								if (atLeastOneLeaseFinished || _activeLeases.Count == 0)
+								{
+									_logger.LogInformation("At least one lease executed or no lease is active. Terminating session...");
+									_sessionResult = new SessionResult(SessionOutcome.Terminate, SessionReason.Completed);
+								}
 							}
 
 							PoolIds = updateSessionResponse.PoolIds;
