Horde: Use lowercase for processor architecture in agent property names

diff --git a/Engine/Source/Programs/Horde/HordeAgent/Services/CapabilitiesService.cs b/Engine/Source/Programs/Horde/HordeAgent/Services/CapabilitiesService.cs
--- a/Engine/Source/Programs/Horde/HordeAgent/Services/CapabilitiesService.cs
+++ b/Engine/Source/Programs/Horde/HordeAgent/Services/CapabilitiesService.cs
@@ -479,6 +479,10 @@ namespace HordeAgent.Services
 			// Used during the transition period over from multi-platform, non-self-contained agent packages.
 			capabilities.Properties.Add($"SelfContained={AgentApp.IsSelfContained}");
 			capabilities.Properties.Add($"DotNetFramework={RuntimeInformation.FrameworkDescription}");
+#pragma warning disable CA1308 // Normalize strings to uppercase
+			capabilities.Properties.Add($"ProcessArch={RuntimeInformation.ProcessArchitecture.ToString().ToLowerInvariant()}");
+			capabilities.Properties.Add($"OsArch={RuntimeInformation.OSArchitecture.ToString().ToLowerInvariant()}");
+#pragma warning restore CA1308
 
 			// Add any additional properties from the config file
 			capabilities.Properties.AddRange(_settings.CurrentValue.Properties.Select(kvp => $"{kvp.Key}={kvp.Value}"));
