Adjustments on device notifications looks
* Use blocks instead of attachments
* Use badges for color coding
* Introduce problem rate
* Reduce the number of messages sent by report cycle

#rnx

diff --git a/Engine/Source/Programs/Horde/Plugins/Build/HordeServer.Build/Devices/DeviceReportService.cs b/Engine/Source/Programs/Horde/Plugins/Build/HordeServer.Build/Devices/DeviceReportService.cs
--- a/Engine/Source/Programs/Horde/Plugins/Build/HordeServer.Build/Devices/DeviceReportService.cs
+++ b/Engine/Source/Programs/Horde/Plugins/Build/HordeServer.Build/Devices/DeviceReportService.cs
@@ -32,6 +32,11 @@ namespace HordeServer.Devices
 		/// </summary>
 		public int ProblemDelta { get; set; } = 0;
 
+		/// <summary>
+		/// Percent of reported problems against the number of reservation (rounded down)
+		/// </summary>
+		public int ProblemPercent { get; set; } = 0;
+
 		/// <summary>
 		/// The last problem Description
 		/// </summary>
@@ -425,9 +430,9 @@ namespace HordeServer.Devices
 				}
 
 				// Report device issues, though only if not since disabled or put into maintenance
-				List<IDeviceTelemetry> problems = telemetry.Where(t => t.ProblemTimeUtc != null).OrderByDescending(t => t.ProblemTimeUtc!).ToList();
+				IEnumerable<IDeviceTelemetry> problems = telemetry.Where(t => t.ProblemTimeUtc != null).OrderByDescending(t => t.ProblemTimeUtc!);
 
-				if (problems.Count == 0)
+				if (!problems.Any())
 				{
 					return;
 				}
@@ -435,11 +440,21 @@ namespace HordeServer.Devices
 				// mark that this device reported a problem
 				metrics.Problems += 1;
 
-				IDeviceTelemetry lastProblem = problems[0];
+				int problemsCount = problems.Count();
+				int problemsPercent = (int)((double)problemsCount / telemetry.Count * 100);
+
+				// don't report low problem rate
+				if (problemsPercent < 10)
+				{
+					return;
+				}
+
+				IDeviceTelemetry lastProblem = problems.First();
 
 				DeviceReport deviceReport = new DeviceReport(platform.Id.ToString(), platform.Name, device.Id.ToString(), device.Name, device.Address ?? "Unknown Address", pool.Id.ToString(), pool.Name, telemetry);
 
-				deviceReport.ProblemDelta = problems.Count;
+				deviceReport.ProblemDelta = problemsCount;
+				deviceReport.ProblemPercent = problemsPercent;
 
 				if (lastProblem.JobName != null && lastProblem.StepName != null && lastProblem.JobId != null && lastProblem.StepId != null)
 				{
diff --git a/Engine/Source/Programs/Horde/Plugins/Build/HordeServer.Build/Notifications/Sinks/SlackNotificationSink.cs b/Engine/Source/Programs/Horde/Plugins/Build/HordeServer.Build/Notifications/Sinks/SlackNotificationSink.cs
--- a/Engine/Source/Programs/Horde/Plugins/Build/HordeServer.Build/Notifications/Sinks/SlackNotificationSink.cs
+++ b/Engine/Source/Programs/Horde/Plugins/Build/HordeServer.Build/Notifications/Sinks/SlackNotificationSink.cs
@@ -67,6 +67,12 @@ namespace HordeServer.Notifications.Sinks
 		// The color to use for success messages.
 		const string SuccessColor = "#4db507";
 
+		// Badges to add color accents
+		const string YellowgBadge = ":large_yellow_square:";
+		const string OrangeBadge = ":large_orange_square:";
+		const string RedBadge = ":large_red_square:";
+		const string BlackBadge = ":black_square:";
+
 		class SocketResponse
 		{
 			[JsonPropertyName("ok")]
@@ -2408,123 +2414,107 @@ namespace HordeServer.Notifications.Sinks
 		{
 			if (report.PoolReports.Count > 0)
 			{
-				SlackMessage headerMessage = new SlackMessage();
-				headerMessage.AddHeader($"Device Pool Health Summary");
-				await SendMessageAsync(report.Channel, headerMessage, cancellationToken);
 				int count = 0;
 
 				foreach (DevicePoolReport pool in report.PoolReports)
 				{
-					List<SlackAttachment> attachments = new List<SlackAttachment>();
+					List<SectionBlock> blocks = new List<SectionBlock>();
 					foreach (DevicePoolMetrics metrics in pool.Metrics)
 					{
-						StringBuilder builder = new StringBuilder("   ");
+						StringBuilder builder = new StringBuilder();
 
-						string color = "#eff700"; // yellow
+						string badge = YellowgBadge;
 
 						if (metrics.AverageLoadPercentage >= 40 || metrics.MaxConcurrentProblemsPercentage >= 50)
 						{
-							color = "#d60404"; // red
+							badge = RedBadge;
 						}
 						else if (metrics.AverageLoadPercentage >= 20 || metrics.MaxConcurrentProblemsPercentage >= 30)
 						{
-							color = "#ffaa00"; // orange
+							badge = OrangeBadge;
+						}
+						else if (metrics.Total == (metrics.Disabled + metrics.Maintenance))
+						{
+							badge = BlackBadge;
 						}
-						else if (metrics.Problems == 0 && metrics.SaturationSpikes == 0 && metrics.Total > (metrics.Disabled + metrics.Maintenance))
+						else if (metrics.Problems == 0 && metrics.SaturationSpikes == 0)
 						{
 							continue;
 						}
+						builder.Append($"{badge} ");
 
 						count++;
 
-						builder.Append($"<{metrics.PlatformURL}|*{metrics.PlatformName}*> - ");
-						builder.Append($"Average Load: {metrics.AverageLoadPercentage}%| ");
+						builder.Append($"<{metrics.PlatformURL}|*{metrics.PlatformName}*>\n> ");
+						builder.Append($"Average Load: {metrics.AverageLoadPercentage}% | ");
 
 						if (metrics.SaturationSpikes > 0)
 						{
-							builder.Append($"Saturation Spikes: {metrics.SaturationSpikes}| ");
-							builder.Append($"Average Spike Duration: {metrics.SpikeDurationAverage.ToString(@"hh\:mm")}({metrics.SpikeDurationPercentage}%)| ");
+							builder.Append($"Saturation Spikes: {metrics.SaturationSpikes} | ");
+							builder.Append($"Average Spike Duration: {metrics.SpikeDurationAverage.ToString(@"hh\:mm")}({metrics.SpikeDurationPercentage}%) | ");
 						}
 
 						if (metrics.Problems > 0)
 						{
-							builder.Append($"Max Concurrent Problems: {metrics.MaxConcurrentProblems}({metrics.MaxConcurrentProblemsPercentage}%)| ");
+							builder.Append($"Max Concurrent Problems: {metrics.MaxConcurrentProblems}({metrics.MaxConcurrentProblemsPercentage}%) | ");
 						}
 
-						builder.Append($"Devices: {metrics.Total}| ");
+						builder.Append($"Devices: {metrics.Total} | ");
 						if (metrics.Disabled > 0)
 						{
-							builder.Append($"Disabled: {metrics.Disabled}| ");
+							builder.Append($"Disabled: {metrics.Disabled} | ");
 						}
 						if (metrics.Maintenance > 0)
 						{
-							builder.Append($"Maintenance: {metrics.Maintenance}| ");
+							builder.Append($"Maintenance: {metrics.Maintenance} | ");
 						}
 						builder.Append($"Streams: {metrics.Streams.Count}");
 
-						SlackAttachment attachment = new SlackAttachment();
-						attachment.FallbackText = builder.ToString();
-						attachment.Color = color;
-						attachment.AddSection(builder.ToString());
-						attachments.Add(attachment);
+						SectionBlock section = new SectionBlock(builder.ToString());
+						blocks.Add(section);
 					}
-					if (attachments.Count > 0)
+					if (blocks.Count > 0)
 					{
-						SlackMessage message = new SlackMessage();
-						message.Attachments.AddRange(attachments);
-						message.Markdown = true;
-						message.Text = $"<{pool.PoolURL}|*{pool.PoolName}*>";
-						await SendMessageAsync(report.Channel, message, cancellationToken);
+						SlackMessage poolReportMessage = new SlackMessage();
+						poolReportMessage.Blocks.Add(new SectionBlock($"<{pool.PoolURL}|*{pool.PoolName}*> *| Device Pool Health*"));
+						poolReportMessage.Blocks.AddRange(blocks);
+						poolReportMessage.AddDivider();
+						await SendMessageAsync(report.Channel, poolReportMessage, cancellationToken);
 					}
 				}
 
 				if (count == 0)
 				{
 					SlackMessage message = new SlackMessage();
-					message.Text = "no outstanding usage to report";
+					message.Blocks.Add(new SectionBlock("*Device Pool Health* | no outstanding usage to report"));
+					message.AddDivider();
 					await SendMessageAsync(report.Channel, message, cancellationToken);
 				}
-
-				SlackMessage divider = new SlackMessage();
-				divider.AddDivider();
-				await SendMessageAsync(report.Channel, divider, cancellationToken);
 			}
 
 			if (report.PlatformReports.Count > 0)
 			{
-				SlackMessage headerMessage = new SlackMessage();
-				headerMessage.AddHeader($"Device Problems Since Last Update");
-				await SendMessageAsync(report.Channel, headerMessage, cancellationToken);
-
 				foreach (DevicePlatformReport platform in report.PlatformReports)
 				{
 					SlackMessage message = new SlackMessage();
-					message.Markdown = true;
-					message.Text = $"*{platform.PlatformName}*";
+					message.Blocks.Add(new SectionBlock($"*{platform.PlatformName} | Device Problems Since Last Update"));
 
 					foreach (DeviceReport device in platform.DeviceReports.OrderByDescending(r => r.ProblemDelta))
 					{
-						StringBuilder builder = new StringBuilder($"   <{device.DevicePoolURL}|*{device.DeviceName}*>* | {device.DeviceAddress} - {device.PoolName} Pool*");
-						builder.Append($"\n{device.ProblemDelta} Problem(s)");
+						StringBuilder builder = new StringBuilder($"<{device.DevicePoolURL}|*{device.DeviceName}*> | {device.DeviceAddress} | {device.PoolName}");
+						builder.Append($"\n> {device.ProblemDelta} Problem(s) | {device.ProblemPercent}% Rate");
 
 						if (!String.IsNullOrEmpty(device.LastProblemURL))
 						{
-							builder.Append($" - <{device.LastProblemURL}|{device.LastProblemDesc ?? "??"}>");
+							builder.Append($"\n> <{device.LastProblemURL}|{device.LastProblemDesc ?? "??"}>");
 						}
 
-						SlackAttachment attachment = new SlackAttachment();
-						attachment.FallbackText = builder.ToString();
-						attachment.Color = "#d60404";
-						attachment.AddSection(builder.ToString());
-						message.Attachments.Add(attachment);
+						message.Blocks.Add(new SectionBlock(builder.ToString()));
 					}
 
+					message.AddDivider();
 					await SendMessageAsync(report.Channel, message, cancellationToken);
 				}
-
-				SlackMessage divider = new SlackMessage();
-				divider.AddDivider();
-				await SendMessageAsync(report.Channel, divider, cancellationToken);
 			}
 		}
 
