Fixing Device Slack notification formmatting

#rnx

diff --git a/Engine/Source/Programs/Horde/Plugins/Build/HordeServer.Build/Notifications/Sinks/SlackNotificationSink.cs b/Engine/Source/Programs/Horde/Plugins/Build/HordeServer.Build/Notifications/Sinks/SlackNotificationSink.cs
--- a/Engine/Source/Programs/Horde/Plugins/Build/HordeServer.Build/Notifications/Sinks/SlackNotificationSink.cs
+++ b/Engine/Source/Programs/Horde/Plugins/Build/HordeServer.Build/Notifications/Sinks/SlackNotificationSink.cs
@@ -2497,11 +2497,11 @@ namespace HordeServer.Notifications.Sinks
 				foreach (DevicePlatformReport platform in report.PlatformReports)
 				{
 					SlackMessage message = new SlackMessage();
-					message.Blocks.Add(new SectionBlock($"*{platform.PlatformName} | Device Problems Since Last Update"));
+					message.Blocks.Add(new SectionBlock($"*{platform.PlatformName} | Device Problems Since Last Update*"));
 
 					foreach (DeviceReport device in platform.DeviceReports.OrderByDescending(r => r.ProblemDelta))
 					{
-						StringBuilder builder = new StringBuilder($"<{device.DevicePoolURL}|*{device.DeviceName}*> | {device.DeviceAddress} | {device.PoolName}");
+						StringBuilder builder = new StringBuilder($"{RedBadge} <{device.DevicePoolURL}|*{device.DeviceName}*> | {device.DeviceAddress} | {device.PoolName}");
 						builder.Append($"\n> {device.ProblemDelta} Problem(s) | {device.ProblemPercent}% Rate");
 
 						if (!String.IsNullOrEmpty(device.LastProblemURL))
