Horde: Update release notes.

#rnx

diff --git a/Engine/Source/Programs/Horde/Docs/ReleaseNotes.md b/Engine/Source/Programs/Horde/Docs/ReleaseNotes.md
--- a/Engine/Source/Programs/Horde/Docs/ReleaseNotes.md
+++ b/Engine/Source/Programs/Horde/Docs/ReleaseNotes.md
@@ -1,8 +1,22 @@
 # Release Notes
 
+## 2024-10-31
+
+* Use lowercase for processor architecture in agent property names (37575294)
+* Fix lease manager ignoring session termination requests when no lease is active (37575102)
+* Make LeaseManager testable in agent. Previously used FakeHordeServer reinstantiated and session result now contain a reason leading to the outcome. (37574891)
+* Right Click for extra Agent Info options isn't communicated in the UI, leads to support overhead. Improves UX communication by adding a Pencil Icon Button to column header which also displays the context menu when clicked. (37566738)
+* Fix issue with password not being sent when creating a new account if it is provided (37514227)
+* Add install/uninstall commands for agent (37499477)
+* Add better error message when agent IPC server cannot start (37499449)
+* Maintain order when selecting deselecting steps (37476215)
+* Fix some step retry logic (37472690)
+* Expose CPU and OS architecture as agent properties. Needed to distinguish between x86_64 and aarch64 for upcoming support in UBA. (37443665)
+* Add cluster to perforce config check (37429576)
+* Reduce size of self-contained agent. Removes unused libs from published directory, such as Linux and macOS versions accidentally making their way in. (37428045)
+
 ## 2024-10-24
 
-* Refactor compiler warnings controlled via module rules or target rules (37409668)
 * Fallback to normal change number for BuildGraph invocation if code change is missing (37388671)
 * Prevent send and possible exceptions in background send for ComputeSocket (37387020)
 * Fix undefined access in job vew (37350938)
@@ -11,12 +25,7 @@
 * Tweaks to compact views to support hiding/showing names (37350158)
 * Slowed the read rate to avoid out of memory issues. (37323730)
 * Fix hanging chrome when missing job template, also surface job error (37322037)
-* UBA: Fixed bug where bool was used instead of byte when transferring from native to managed code. (37320327)
-* UBA: Hooked up cross architecture logic from c# to c++ (37316216)
 * Set agent's default setting for EC2 support to disabled
-* Two second connect timeout and no printed message when UBT telemetry fails to connect or send. (37297792)
-* Add initial documentation for compute plugin
-* Return Android devices with Turnkey, and Editor `Platforms` menu update
 * Adding robomerge track change messages to slack messages. (37266636)
 
 ## 2024-10-17
