Horde: Fix DI error with Horde client tool commands

diff --git a/Engine/Source/Programs/Horde/HordeCmd/Commands/Tools/ToolDownload.cs b/Engine/Source/Programs/Horde/HordeCmd/Commands/Tools/ToolDownload.cs
--- a/Engine/Source/Programs/Horde/HordeCmd/Commands/Tools/ToolDownload.cs
+++ b/Engine/Source/Programs/Horde/HordeCmd/Commands/Tools/ToolDownload.cs
@@ -7,44 +7,37 @@ using EpicGames.Horde;
 using EpicGames.Horde.Tools;
 using Microsoft.Extensions.Logging;
 
-namespace Horde.Commands
+namespace HordeCmd.Commands.Tools;
+
+[Command("tool", "download", "Downloads a tool from the server to a local directory")]
+class ToolDownload(IHordeClient hordeClient) : Command
 {
-	[Command("tool", "download", "Downloads a tool from the server to a local directory")]
-	class ToolDownload : Command
+	[CommandLine("-Id=", Required = true)]
+	[Description("The tool to download")]
+	public ToolId Id { get; set; }
+
+	[CommandLine("-Deployment=")]
+	[Description("Deployment to fetch")]
+	public ToolDeploymentId? DeploymentId { get; set; }
+
+	[CommandLine("-OutputDir=", Required = true)]
+	[Description("Output directory for the downloaded tool data")]
+	public DirectoryReference OutputDir { get; set; } = null!;
+	
+	/// <inheritdoc/>
+	public override async Task<int> ExecuteAsync(ILogger logger)
 	{
-		[CommandLine("-Id=", Required = true)]
-		[Description("The tool to download")]
-		public ToolId Id { get; set; }
-
-		[CommandLine("-Deployment=")]
-		[Description("Deployment to fetch")]
-		public ToolDeploymentId? DeploymentId { get; set; }
-
-		[CommandLine("-OutputDir=", Required = true)]
-		[Description("Output directory for the downloaded tool data")]
-		public DirectoryReference OutputDir { get; set; } = null!;
-
-		readonly HordeHttpClient _hordeHttpClient;
+		DirectoryReference.CreateDirectory(OutputDir);
 
-		public ToolDownload(HordeHttpClient httpClient)
+		FileReference zipFile = FileReference.Combine(OutputDir, "temp.zip");
+		await using (FileStream zipOutputStream = FileReference.Open(zipFile, FileMode.Create, FileAccess.Write))
 		{
-			_hordeHttpClient = httpClient;
+			using HordeHttpClient httpClient = hordeClient.CreateHttpClient();
+			await using Stream stream = await httpClient.GetToolDeploymentZipAsync(Id, DeploymentId, CancellationToken.None);
+			await stream.CopyToAsync(zipOutputStream);
 		}
 
-		/// <inheritdoc/>
-		public override async Task<int> ExecuteAsync(ILogger logger)
-		{
-			DirectoryReference.CreateDirectory(OutputDir);
-
-			FileReference zipFile = FileReference.Combine(OutputDir, "temp.zip");
-			using (FileStream zipOutputStream = FileReference.Open(zipFile, FileMode.Create, FileAccess.Write))
-			{
-				await using Stream stream = await _hordeHttpClient.GetToolDeploymentZipAsync(Id, DeploymentId, CancellationToken.None);
-				await stream.CopyToAsync(zipOutputStream);
-			}
-
-			ZipFile.ExtractToDirectory(zipFile.FullName, OutputDir.FullName, true);
-			return 0;
-		}
+		ZipFile.ExtractToDirectory(zipFile.FullName, OutputDir.FullName, true);
+		return 0;
 	}
 }
diff --git a/Engine/Source/Programs/Horde/HordeCmd/Commands/Tools/ToolList.cs b/Engine/Source/Programs/Horde/HordeCmd/Commands/Tools/ToolList.cs
--- a/Engine/Source/Programs/Horde/HordeCmd/Commands/Tools/ToolList.cs
+++ b/Engine/Source/Programs/Horde/HordeCmd/Commands/Tools/ToolList.cs
@@ -5,27 +5,25 @@ using EpicGames.Horde;
 using EpicGames.Horde.Tools;
 using Microsoft.Extensions.Logging;
 
-namespace Horde.Commands
+namespace HordeCmd.Commands.Tools;
+
+[Command("tool", "list", "List all tools available for download")]
+class ToolList(IHordeClient hordeClient) : Command
 {
-	[Command("tool", "list", "List all tools available for download")]
-	class ToolList : Command
+	/// <inheritdoc/>
+	public override async Task<int> ExecuteAsync(ILogger logger)
 	{
-		readonly HordeHttpClient _hordeHttpClient;
-
-		public ToolList(HordeHttpClient httpClient)
-		{
-			_hordeHttpClient = httpClient;
-		}
-
-		/// <inheritdoc/>
-		public override async Task<int> ExecuteAsync(ILogger logger)
+		using HordeHttpClient httpClient = hordeClient.CreateHttpClient();
+		logger.LogInformation("Horde server URL {ServerUrl}", httpClient.BaseUrl);
+		GetToolsSummaryResponse response = await httpClient.GetToolsAsync();
+		
+		logger.LogInformation("  ---------------------------------------------------------------------------------------------");
+		logger.LogInformation("  {ToolId,-30} {Deployment,-30} {Version,-15}", "ID", "DeploymentId", "Version");
+		logger.LogInformation("  ---------------------------------------------------------------------------------------------");
+		foreach (GetToolSummaryResponse tool in response.Tools)
 		{
-			GetToolsSummaryResponse response = await _hordeHttpClient.GetToolsAsync();
-			foreach (GetToolSummaryResponse tool in response.Tools)
-			{
-				logger.LogInformation("  {ToolId,-30} {Deployment,-30} {Version,-15}", tool.Id, tool.DeploymentId, tool.Version);
-			}
-			return 0;
+			logger.LogInformation("  {ToolId,-30} {Deployment,-30} {Version,-15}", tool.Id, tool.DeploymentId, tool.Version);
 		}
+		return 0;
 	}
 }
diff --git a/Engine/Source/Programs/Horde/HordeCmd/Commands/Tools/ToolUpload.cs b/Engine/Source/Programs/Horde/HordeCmd/Commands/Tools/ToolUpload.cs
--- a/Engine/Source/Programs/Horde/HordeCmd/Commands/Tools/ToolUpload.cs
+++ b/Engine/Source/Programs/Horde/HordeCmd/Commands/Tools/ToolUpload.cs
@@ -8,47 +8,39 @@ using EpicGames.Horde.Storage.Nodes;
 using EpicGames.Horde.Tools;
 using Microsoft.Extensions.Logging;
 
-namespace Horde.Commands
+namespace HordeCmd.Commands.Tools;
+
+[Command("tool", "upload", "Uploads a tool from a local directory")]
+class ToolUpload(IHordeClient hordeClient, HttpStorageClient storageClient) : Command
 {
-	[Command("tool", "upload", "Uploads a tool from a local directory")]
-	class ToolUpload : Command
+	[CommandLine("-Id=", Required = true)]
+	[Description("Identifier for the tool to upload")]
+	public ToolId ToolId { get; set; }
+	
+	[CommandLine("-Version=")]
+	[Description("Optional version number for the new upload")]
+	public string? Version { get; set; }
+	
+	[CommandLine("-InputDir=", Required = true)]
+	[Description("Directory containing files to upload for the tool")]
+	public DirectoryReference InputDir { get; set; } = null!;
+	
+	/// <inheritdoc/>
+	public override async Task<int> ExecuteAsync(ILogger logger)
 	{
-		[CommandLine("-Id=")]
-		[Description("Identifier for the tool to upload")]
-		public ToolId ToolId { get; set; }
-
-		[CommandLine("-Version=")]
-		[Description("Optional version number for the new upload")]
-		public string? Version { get; set; }
-
-		[CommandLine("-InputDir=")]
-		[Description("Directory containing files to upload for the tool")]
-		public DirectoryReference InputDir { get; set; } = null!;
-
-		readonly HordeHttpClient _hordeHttpClient;
-		readonly HttpStorageClient _storageClient;
-
-		public ToolUpload(HordeHttpClient httpClient, HttpStorageClient storageClient)
+		IStorageNamespace storageNamespace = storageClient.GetNamespaceWithPath($"api/v1/tools/{ToolId}");
+		
+		IHashedBlobRef<DirectoryNode> target;
+		await using (IBlobWriter writer = storageNamespace.CreateBlobWriter())
 		{
-			_hordeHttpClient = httpClient;
-			_storageClient = storageClient;
-		}
-
-		/// <inheritdoc/>
-		public override async Task<int> ExecuteAsync(ILogger logger)
-		{
-			IStorageNamespace storageNamespace = _storageClient.GetNamespaceWithPath($"api/v1/tools/{ToolId}");
-
-			IHashedBlobRef<DirectoryNode> target;
-			await using (IBlobWriter writer = storageNamespace.CreateBlobWriter())
-			{
-				target = await writer.WriteFilesAsync(InputDir);
-			}
-
-			ToolDeploymentId deploymentId = await _hordeHttpClient.CreateToolDeploymentAsync(ToolId, Version, null, null, target.GetRefValue());
-			logger.LogInformation("Created deployment {DeploymentId}", deploymentId);
-
-			return 0;
+			target = await writer.WriteFilesAsync(InputDir);
 		}
+		
+		using HordeHttpClient httpClient = hordeClient.CreateHttpClient();
+		logger.LogInformation("Horde server URL {ServerUrl}", httpClient.BaseUrl);
+		ToolDeploymentId deploymentId = await httpClient.CreateToolDeploymentAsync(ToolId, Version, null, null, target.GetRefValue());
+		logger.LogInformation("Created deployment {DeploymentId}", deploymentId);
+		
+		return 0;
 	}
-}
+}
\ No newline at end of file
