Horde: Fix storage client initialization in tool upload command

Thanks Ben!

diff --git a/Engine/Source/Programs/Horde/HordeCmd/Commands/Tools/ToolUpload.cs b/Engine/Source/Programs/Horde/HordeCmd/Commands/Tools/ToolUpload.cs
--- a/Engine/Source/Programs/Horde/HordeCmd/Commands/Tools/ToolUpload.cs
+++ b/Engine/Source/Programs/Horde/HordeCmd/Commands/Tools/ToolUpload.cs
@@ -11,7 +11,7 @@ using Microsoft.Extensions.Logging;
 namespace HordeCmd.Commands.Tools;
 
 [Command("tool", "upload", "Uploads a tool from a local directory")]
-class ToolUpload(IHordeClient hordeClient, HttpStorageClient storageClient) : Command
+class ToolUpload(IHordeClient hordeClient) : Command
 {
 	[CommandLine("-Id=", Required = true)]
 	[Description("Identifier for the tool to upload")]
@@ -28,7 +28,7 @@ class ToolUpload(IHordeClient hordeClient, HttpStorageClient storageClient) : Co
 	/// <inheritdoc/>
 	public override async Task<int> ExecuteAsync(ILogger logger)
 	{
-		IStorageNamespace storageNamespace = storageClient.GetNamespaceWithPath($"api/v1/tools/{ToolId}");
+		IStorageNamespace storageNamespace = hordeClient.GetStorageNamespace(ToolId);
 		
 		IHashedBlobRef<DirectoryNode> target;
 		await using (IBlobWriter writer = storageNamespace.CreateBlobWriter())
