Horde: Add a better error for malformed URI formats when preflighting config changes.

#rnx

diff --git a/Engine/Source/Programs/Horde/Plugins/Build/HordeServer.Build/Perforce/ServerPerforceController.cs b/Engine/Source/Programs/Horde/Plugins/Build/HordeServer.Build/Perforce/ServerPerforceController.cs
--- a/Engine/Source/Programs/Horde/Plugins/Build/HordeServer.Build/Perforce/ServerPerforceController.cs
+++ b/Engine/Source/Programs/Horde/Plugins/Build/HordeServer.Build/Perforce/ServerPerforceController.cs
@@ -64,7 +64,16 @@ namespace HordeServer.Perforce
 
 					PrintRecord<byte[]> printRecord = printRecordResponse.Data;
 
-					Uri uri = new Uri($"perforce://{cluster}{printRecord.DepotFile}");
+					Uri uri;
+					try
+					{
+						uri = new Uri($"perforce://{cluster}{printRecord.DepotFile}");
+					}
+					catch (Exception ex)
+					{
+						return BadRequest($"Unable to create URI for cluster '{cluster}', file '{printRecord.DepotFile}': {ex.Message}");
+					}
+
 					files.Add(uri, printRecord.Contents);
 				}
 			}
