Horde: Include toolbox.json in agent with process elevation

diff --git a/Engine/Source/Programs/Horde/HordeAgent/HordeAgent.csproj b/Engine/Source/Programs/Horde/HordeAgent/HordeAgent.csproj
--- a/Engine/Source/Programs/Horde/HordeAgent/HordeAgent.csproj
+++ b/Engine/Source/Programs/Horde/HordeAgent/HordeAgent.csproj
@@ -113,7 +113,10 @@
       <CopyToPublishDirectory>Always</CopyToPublishDirectory>
     </Content>
     <Content Include="HordeAgent.ico">
-      <CopyToOutputDirectory>PreserveNewest</CopyToOutputDirectory>
+      <CopyToOutputDirectory>Always</CopyToOutputDirectory>
+    </Content>
+    <Content Include="Toolbox.json">
+      <CopyToOutputDirectory>Always</CopyToOutputDirectory>
     </Content>
   </ItemGroup>
 
diff --git a/Engine/Source/Programs/Horde/HordeAgent/Toolbox.json b/Engine/Source/Programs/Horde/HordeAgent/Toolbox.json
--- a/Engine/Source/Programs/Horde/HordeAgent/Toolbox.json
+++ b/Engine/Source/Programs/Horde/HordeAgent/Toolbox.json
@@ -1,17 +1,12 @@
 {
 	"manualInstall": true,
+	"requiresElevation": true,
 	"installCommand": {
-		"fileName": "msiexec.exe",
-		"arguments": [
-			"/i",
-			"UnrealHordeAgent.msi"
-		]
+		"fileName": "HordeAgent.exe",
+		"arguments": ["setup", "install"]
 	},
 	"uninstallCommand": {
-		"fileName": "msiexec.exe",
-		"arguments": [
-			"/x",
-			"UnrealHordeAgent.msi"
-		]
+		"fileName": "HordeAgent.exe",
+		"arguments": ["setup", "uninstall", "-UseRegistryPath=true"]
 	}
 }
\ No newline at end of file
