Horde: Remove EpicGames.Perforce dependency from agent

diff --git a/Engine/Source/Programs/Horde/HordeAgent/HordeAgent.csproj b/Engine/Source/Programs/Horde/HordeAgent/HordeAgent.csproj
--- a/Engine/Source/Programs/Horde/HordeAgent/HordeAgent.csproj
+++ b/Engine/Source/Programs/Horde/HordeAgent/HordeAgent.csproj
@@ -152,8 +152,6 @@
   <ItemGroup>
     <ProjectReference Include="..\..\Shared\EpicGames.Core\EpicGames.Core.csproj" />
     <ProjectReference Include="..\..\Shared\EpicGames.Horde\EpicGames.Horde.csproj" />
-    <ProjectReference Include="..\..\Shared\EpicGames.Perforce.Managed\EpicGames.Perforce.Managed.csproj" />
-    <ProjectReference Include="..\..\Shared\EpicGames.Perforce\EpicGames.Perforce.csproj" />
   </ItemGroup>
 
   <ItemGroup>
diff --git a/Engine/Source/Programs/Horde/HordeAgent/Leases/LeaseHandler.cs b/Engine/Source/Programs/Horde/HordeAgent/Leases/LeaseHandler.cs
--- a/Engine/Source/Programs/Horde/HordeAgent/Leases/LeaseHandler.cs
+++ b/Engine/Source/Programs/Horde/HordeAgent/Leases/LeaseHandler.cs
@@ -3,7 +3,6 @@
 using System.Diagnostics;
 using EpicGames.Core;
 using EpicGames.Horde.Agents.Leases;
-using EpicGames.Perforce.Managed;
 using Google.Protobuf;
 using Google.Protobuf.WellKnownTypes;
 using HordeAgent.Services;
@@ -121,12 +120,6 @@ namespace HordeAgent.Leases
 			{
 				logger.LogInformation("Lease {LeaseId} cancelled", Id);
 			}
-			catch (InsufficientSpaceException ex)
-			{
-				logger.LogError(ex, "{Message}", ex.Message);
-				span.SetStatus(Status.Error);
-				span.RecordException(ex);
-			}
 			catch (Exception ex)
 			{
 				logger.LogError(ex, "Unhandled exception while executing lease {LeaseId}: {Message}", Id, ex.Message);
