Horde: Add option for using uninstall location from registry

diff --git a/Engine/Source/Programs/Horde/HordeAgent/Commands/Setup/SetupInstall.cs b/Engine/Source/Programs/Horde/HordeAgent/Commands/Setup/SetupInstall.cs
--- a/Engine/Source/Programs/Horde/HordeAgent/Commands/Setup/SetupInstall.cs
+++ b/Engine/Source/Programs/Horde/HordeAgent/Commands/Setup/SetupInstall.cs
@@ -117,9 +117,9 @@ class SetupInstallCommand : Command
 		}
 		
 		DirectoryReference.CreateDirectory(destinationDir);
-		if (!FileReference.Exists(FileReference.Combine(sourceDir, HordeAgentDll)))
+		if (!ContainsHordeAgentExecutable(sourceDir))
 		{
-			throw new InstallException($"Source directory does not contain {HordeAgentDll}");
+			throw new InstallException($"Source directory is missing either {HordeAgentDll} or {HordeAgentExe}");
 		}
 		
 		logger.LogInformation("Copying files...");
@@ -211,6 +211,13 @@ class SetupInstallCommand : Command
 		}
 	}
 	
+	internal static bool ContainsHordeAgentExecutable(DirectoryReference dir)
+	{
+		bool containsDll = FileReference.Exists(FileReference.Combine(dir, HordeAgentDll));
+		bool containsExe = FileReference.Exists(FileReference.Combine(dir, HordeAgentExe));
+		return containsDll || containsExe;
+	}
+	
 	internal static void RemoveFirewallPortRule(string name)
 	{
 		RunProcessSilently("netsh", $"advfirewall firewall delete rule name=\"{name}\"");
@@ -280,6 +287,19 @@ class SetupInstallCommand : Command
 		return true;
 	}
 	
+	[SupportedOSPlatform("windows")]
+	public static DirectoryReference GetUninstallLocation()
+	{
+		const string RegKeyPath = $"{UninstallKeyPath}\\{UninstallAppName}";
+		using RegistryKey? key = Registry.LocalMachine.OpenSubKey(RegKeyPath, true);
+		if (key?.GetValue("InstallLocation") is string uninstallString)
+		{
+			return new DirectoryReference(uninstallString);
+		}
+		
+		throw new InstallException($"Unable to open registry key {RegKeyPath}");
+	}
+	
 	/// <summary>
 	/// Runs a Windows batch script that will delete the given directory
 	/// The script will be executed in the background and wait for a process to exit before attempting to delete.
diff --git a/Engine/Source/Programs/Horde/HordeAgent/Commands/Setup/SetupUninstall.cs b/Engine/Source/Programs/Horde/HordeAgent/Commands/Setup/SetupUninstall.cs
--- a/Engine/Source/Programs/Horde/HordeAgent/Commands/Setup/SetupUninstall.cs
+++ b/Engine/Source/Programs/Horde/HordeAgent/Commands/Setup/SetupUninstall.cs
@@ -18,6 +18,10 @@ class SetupUninstallCommand : Command
 	[Description("Whether to show full stack trace on errors (true/false)")]
 	public string? ShowStackTrace { get; set; } = "false";
 	
+	[CommandLine("-UseRegistryPath=")]
+	[Description("Use uninstall string as specified in registry")]
+	public string? UseRegistryPath { get; set; } = "false";
+	
 	/// <summary>
 	/// Uninstall the agent from Windows 
 	/// </summary>
@@ -31,14 +35,17 @@ class SetupUninstallCommand : Command
 			return Task.FromResult(1);
 		}
 		
+		bool useRegistryPath = UseRegistryPath != null && UseRegistryPath.Equals("true", StringComparison.OrdinalIgnoreCase);
+		bool showStackTrace = ShowStackTrace != null && ShowStackTrace.Equals("true", StringComparison.OrdinalIgnoreCase);
+		
 		try
 		{
-			Uninstall(logger);
+			Uninstall(useRegistryPath, logger);
 		}
 		catch (InstallException e)
 		{
 			logger.LogError("Uninstall failed: {Reason}", e.Message);
-			if (ShowStackTrace != null && ShowStackTrace.Equals("true", StringComparison.OrdinalIgnoreCase))
+			if (showStackTrace)
 			{
 				throw;
 			}
@@ -50,15 +57,26 @@ class SetupUninstallCommand : Command
 	}
 	
 	[SupportedOSPlatform("windows")]
-	private static void Uninstall(ILogger logger)
+	private static void Uninstall(bool useRegistryPath, ILogger logger)
 	{
+		if (!UserSecurityUtils.IsAdministrator())
+		{
+			throw new InstallException("Administrator access is required for uninstalling");
+		}
+		
+		DirectoryReference installDir = useRegistryPath ? SetupInstallCommand.GetUninstallLocation() : new DirectoryReference(AppContext.BaseDirectory);
+		if (!SetupInstallCommand.ContainsHordeAgentExecutable(installDir))
+		{
+			throw new InstallException($"Directory to uninstall does not contain Horde agent executable");
+		}
+		
 		logger.LogInformation("Uninstalling Horde agent...");
 		InstallServiceCommand.StopAndDeleteWindowsService(logger);
 		SetupInstallCommand.RemoveFirewallPortRule(SetupInstallCommand.FirewallRuleName);
 		SetupInstallCommand.UnregisterAppForUninstall();
 		
 		// Remove directory containing agent binaries
-		logger.LogInformation("Deleting dir {AgentDir}", AppContext.BaseDirectory);
-		SetupInstallCommand.RunSelfDeleteScript(new DirectoryReference(AppContext.BaseDirectory), Environment.ProcessId);
+		logger.LogInformation("Deleting dir {AgentDir}", installDir.FullName);
+		SetupInstallCommand.RunSelfDeleteScript(installDir, Environment.ProcessId);
 	}
 }
\ No newline at end of file
