Horde: Allow reading artifact blobs from other artifacts of the same type and within in the same stream. This allows accessing deduped blobs using the regular artifact endpoints.

#rnx

diff --git a/Engine/Source/Programs/Horde/Plugins/Build/HordeServer.Build/Artifacts/ArtifactsController.cs b/Engine/Source/Programs/Horde/Plugins/Build/HordeServer.Build/Artifacts/ArtifactsController.cs
--- a/Engine/Source/Programs/Horde/Plugins/Build/HordeServer.Build/Artifacts/ArtifactsController.cs
+++ b/Engine/Source/Programs/Horde/Plugins/Build/HordeServer.Build/Artifacts/ArtifactsController.cs
@@ -276,7 +276,7 @@ namespace HordeServer.Artifacts
 			{
 				return Forbid(ArtifactAclAction.ReadArtifact, artifact.StreamId);
 			}
-			if (!locator.WithinFolder(artifact.RefName.Text))
+			if (!locator.WithinFolder(new Utf8String(ArtifactCollection.GetArtifactPath(artifact.StreamId, artifact.Type))) && !locator.WithinFolder(artifact.RefName.Text))
 			{
 				return BadRequest("Invalid blob id for artifact");
 			}
