Horde: Include the artifact name, type, and description in the uartifact file.

#rnx

diff --git a/Engine/Source/Programs/Horde/Plugins/Build/HordeServer.Build/Artifacts/ArtifactsController.cs b/Engine/Source/Programs/Horde/Plugins/Build/HordeServer.Build/Artifacts/ArtifactsController.cs
--- a/Engine/Source/Programs/Horde/Plugins/Build/HordeServer.Build/Artifacts/ArtifactsController.cs
+++ b/Engine/Source/Programs/Horde/Plugins/Build/HordeServer.Build/Artifacts/ArtifactsController.cs
@@ -607,6 +607,9 @@ namespace HordeServer.Artifacts
 			Uri baseUri = new Uri(_serverInfo.ServerUrl, $"api/v2/artifacts/{artifact.Id}");
 
 			ArtifactDescriptor descriptor = new ArtifactDescriptor(baseUri, new RefName("default"), fileFilter);
+			descriptor.Name = artifact.Name.ToString();
+			descriptor.Type = artifact.Type.ToString();
+			descriptor.Description = artifact.Description;
 			descriptor.Keys = new List<string>(artifact.Keys);
 			descriptor.Metadata = new List<string>(artifact.Metadata);
 
diff --git a/Engine/Source/Programs/Shared/EpicGames.Horde/Artifacts/ArtifactDescriptor.cs b/Engine/Source/Programs/Shared/EpicGames.Horde/Artifacts/ArtifactDescriptor.cs
--- a/Engine/Source/Programs/Shared/EpicGames.Horde/Artifacts/ArtifactDescriptor.cs
+++ b/Engine/Source/Programs/Shared/EpicGames.Horde/Artifacts/ArtifactDescriptor.cs
@@ -16,6 +16,21 @@ namespace EpicGames.Horde.Storage
 	/// </summary>
 	public class ArtifactDescriptor
 	{
+		/// <summary>
+		/// Name of the artifact
+		/// </summary>
+		public string? Name { get; set; }
+
+		/// <summary>
+		/// Type of the artifact
+		/// </summary>
+		public string? Type { get; set; }
+
+		/// <summary>
+		/// Artifact description 
+		/// </summary>
+		public string? Description { get; set; }
+
 		/// <summary>
 		/// Base URL for downloading from
 		/// </summary>
