Horde: Add flag for registering agent for uninstall in Windows

If the agent is installed via Unreal Toolbox but later uninstalled via Windows add/remove programs, Toolbox won't be notified.
Until the uninstall command can be routed through Toolbox, this flag prevents Windows from handling uninstallation directly.

diff --git a/Engine/Source/Programs/Horde/HordeAgent/Commands/Setup/SetupInstall.cs b/Engine/Source/Programs/Horde/HordeAgent/Commands/Setup/SetupInstall.cs
--- a/Engine/Source/Programs/Horde/HordeAgent/Commands/Setup/SetupInstall.cs
+++ b/Engine/Source/Programs/Horde/HordeAgent/Commands/Setup/SetupInstall.cs
@@ -44,6 +44,10 @@ class SetupInstallCommand : Command
 	[Description("Whether to show full stack trace on errors (true/false)")]
 	public string? ShowStackTrace { get; set; } = "false";
 	
+	[CommandLine("-RegisterForUninstall=")]
+	[Description("Whether to register the app for uninstall in Windows (true/false)")]
+	public string? RegisterForUninstall { get; set; } = "true";
+	
 	/// <summary>
 	/// Install the 
 	/// </summary>
@@ -57,14 +61,17 @@ class SetupInstallCommand : Command
 			return Task.FromResult(1);
 		}
 		
+		bool showStackTrace = ShowStackTrace != null && ShowStackTrace.Equals("true", StringComparison.OrdinalIgnoreCase);
+		bool registerForUninstall = RegisterForUninstall != null && RegisterForUninstall.Equals("true", StringComparison.OrdinalIgnoreCase);
+		
 		try
 		{
-			Install(logger, new InstallParams());
+			Install(logger, new InstallParams(RegisterForUninstall: registerForUninstall));
 		}
 		catch (InstallException e)
 		{
 			logger.LogError("Install failed: {Reason}", e.Message);
-			if (ShowStackTrace != null && ShowStackTrace.Equals("true", StringComparison.OrdinalIgnoreCase))
+			if (showStackTrace)
 			{
 				throw;
 			}
@@ -80,13 +87,15 @@ class SetupInstallCommand : Command
 	/// </summary>
 	/// <param name="SourceDir">Directory containing files to install (HordeAgent.dll etc)</param>
 	/// <param name="DestDir">Installation directory (unless specified, defaults to a directory under Program Files)</param>
-	/// <param name="ConfigDir"></param>
-	/// <param name="WorkingDir"></param>
+	/// <param name="ConfigDir">Config directory</param>
+	/// <param name="WorkingDir">Sandbox directory for agent (where VCS-synced projects are stored)</param>
+	/// <param name="RegisterForUninstall">Whether to register the app for uninstall in Windows</param>
 	public record InstallParams(
 		string? SourceDir = null,
 		string? DestDir = null,
 		string? ConfigDir = null,
-		string? WorkingDir = null
+		string? WorkingDir = null,
+		bool RegisterForUninstall = true
 	);
 	
 	/// <summary>
@@ -155,16 +164,19 @@ class SetupInstallCommand : Command
 			throw new InstallException("Failed installing Horde Agent as a Windows service");
 		}
 		
-		(executable, args) = ResolveCommandLine(destinationDir.FullName, "setup uninstall");
-		RegisterAppForUninstall(new AppInfo(
-			DisplayName: "Unreal Horde Agent",
-			InstallLocation: destinationDir,
-			Version: AgentApp.Version,
-			Publisher: "Epic Games",
-			UninstallCommand: $"\"{executable}\" {args}",
-			IconPath: FileReference.Combine(destinationDir, HordeAgentIco),
-			EstimatedSizeInKb: DirectoryReference.GetTotalSize(destinationDir) / 1024
-		));
+		if (installParams.RegisterForUninstall)
+		{
+			(executable, args) = ResolveCommandLine(destinationDir.FullName, "setup uninstall");
+			RegisterAppForUninstall(new AppInfo(
+				DisplayName: "Unreal Horde Agent",
+				InstallLocation: destinationDir,
+				Version: AgentApp.Version,
+				Publisher: "Epic Games",
+				UninstallCommand: $"\"{executable}\" {args}",
+				IconPath: FileReference.Combine(destinationDir, HordeAgentIco),
+				EstimatedSizeInKb: DirectoryReference.GetTotalSize(destinationDir) / 1024
+			));
+		}
 		
 		return 0;
 	}
