Horde: Add an Experimental plugin for iterating on features that may be merged into the main codebase at some point in the future.

#rnx

diff --git a/Engine/Source/Programs/Horde/HordeServer.Shared/Plugins/PluginAttribute.cs b/Engine/Source/Programs/Horde/HordeServer.Shared/Plugins/PluginAttribute.cs
--- a/Engine/Source/Programs/Horde/HordeServer.Shared/Plugins/PluginAttribute.cs
+++ b/Engine/Source/Programs/Horde/HordeServer.Shared/Plugins/PluginAttribute.cs
@@ -14,6 +14,11 @@ namespace HordeServer.Plugins
 		/// </summary>
 		public string Name { get; }
 
+		/// <summary>
+		/// Whether the plugin should be enabled by default
+		/// </summary>
+		public bool EnabledByDefault { get; set; } = true;
+
 		/// <summary>
 		/// Type containing server settings for this plugin. Will be injected into a <see cref="Microsoft.Extensions.Options.IOptions{TOptions}"/>
 		/// object in the service container, as well as (optionally) the constructor for the the concrete plugin startup object.
diff --git a/Engine/Source/Programs/Horde/HordeServer.Shared/Plugins/PluginServerConfig.cs b/Engine/Source/Programs/Horde/HordeServer.Shared/Plugins/PluginServerConfig.cs
--- a/Engine/Source/Programs/Horde/HordeServer.Shared/Plugins/PluginServerConfig.cs
+++ b/Engine/Source/Programs/Horde/HordeServer.Shared/Plugins/PluginServerConfig.cs
@@ -10,6 +10,6 @@ namespace HordeServer.Plugins
 		/// <summary>
 		/// Whether the plugin should be enabled or not
 		/// </summary>
-		public bool Enabled { get; set; } = true;
+		public bool? Enabled { get; set; }
 	}
 }
diff --git a/Engine/Source/Programs/Horde/HordeServer/HordeServer.csproj b/Engine/Source/Programs/Horde/HordeServer/HordeServer.csproj
--- a/Engine/Source/Programs/Horde/HordeServer/HordeServer.csproj
+++ b/Engine/Source/Programs/Horde/HordeServer/HordeServer.csproj
@@ -114,6 +114,7 @@
     <ProjectReference Include="..\Plugins\Build\HordeServer.Build\HordeServer.Build.csproj" />
     <ProjectReference Include="..\Plugins\Compute\HordeServer.Compute\HordeServer.Compute.csproj" />
     <ProjectReference Include="..\Plugins\Ddc\HordeServer.Ddc\HordeServer.Ddc.csproj" />
+    <ProjectReference Include="..\Plugins\Experimental\HordeServer.Experimental\HordeServer.Experimental.csproj" />
     <ProjectReference Include="..\Plugins\Secrets\HordeServer.Secrets\HordeServer.Secrets.csproj" />
     <ProjectReference Include="..\Plugins\Storage\HordeServer.Storage\HordeServer.Storage.csproj" />
     <ProjectReference Include="..\Plugins\Symbols\HordeServer.Symbols\HordeServer.Symbols.csproj" />
diff --git a/Engine/Source/Programs/Horde/HordeServer/ServerApp.cs b/Engine/Source/Programs/Horde/HordeServer/ServerApp.cs
--- a/Engine/Source/Programs/Horde/HordeServer/ServerApp.cs
+++ b/Engine/Source/Programs/Horde/HordeServer/ServerApp.cs
@@ -255,7 +255,7 @@ namespace HordeServer
 			HashSet<string> missingPlugins = new HashSet<string>(StringComparer.OrdinalIgnoreCase);
 			foreach ((string name, PluginServerConfig pluginConfig) in pluginConfigs)
 			{
-				if (pluginConfig.Enabled)
+				if (pluginConfig.Enabled ?? true)
 				{
 					missingPlugins.Add(name);
 				}
@@ -273,8 +273,13 @@ namespace HordeServer
 						PluginAttribute? pluginAttribute = type.GetCustomAttribute<PluginAttribute>();
 						if (pluginAttribute != null)
 						{
-							PluginServerConfig? pluginConfig;
-							if (!pluginConfigs.TryGetValue(pluginAttribute.Name, out pluginConfig) || pluginConfig.Enabled)
+							bool? enabled = null;
+							if (pluginConfigs.TryGetValue(pluginAttribute.Name, out PluginServerConfig? pluginConfig))
+							{
+								enabled = pluginConfig.Enabled;
+							}
+							
+							if (enabled ?? pluginAttribute.EnabledByDefault)
 							{
 								logger.LogDebug("Added plugin '{Plugin}'", pluginAttribute.Name);
 								pluginCollection.Add(type);
diff --git a/Engine/Source/Programs/Horde/Plugins/Build/HordeServer.Build/BuildPlugin.cs b/Engine/Source/Programs/Horde/Plugins/Build/HordeServer.Build/BuildPlugin.cs
--- a/Engine/Source/Programs/Horde/Plugins/Build/HordeServer.Build/BuildPlugin.cs
+++ b/Engine/Source/Programs/Horde/Plugins/Build/HordeServer.Build/BuildPlugin.cs
@@ -38,7 +38,7 @@ using Microsoft.Extensions.DependencyInjection;
 namespace HordeServer
 {
 	/// <summary>
-	/// Entry point for the storage plugin
+	/// Entry point for the build plugin
 	/// </summary>
 	[Plugin("Build", GlobalConfigType = typeof(BuildConfig), ServerConfigType = typeof(BuildServerConfig))]
 	public class BuildPlugin : IPluginStartup
@@ -225,7 +225,7 @@ namespace HordeServer
 	}
 
 	/// <summary>
-	/// Helper methods for storage config
+	/// Helper methods for build config
 	/// </summary>
 	public static class BuildPluginExtensions
 	{
diff --git a/Engine/Source/Programs/Horde/Plugins/Experimental/HordeServer.Experimental.Tests/ExperimentalTestSetup.cs b/Engine/Source/Programs/Horde/Plugins/Experimental/HordeServer.Experimental.Tests/ExperimentalTestSetup.cs
new file mode 100644
--- /dev/null
+++ b/Engine/Source/Programs/Horde/Plugins/Experimental/HordeServer.Experimental.Tests/ExperimentalTestSetup.cs
@@ -0,0 +1,26 @@
+// Copyright Epic Games, Inc. All Rights Reserved.
+
+using HordeServer.Notifications;
+using HordeServer.Notifications.Sinks;
+using Microsoft.Extensions.DependencyInjection;
+
+namespace HordeServer.Tests
+{
+	/// <summary>
+	/// Handles set up of collections, services, fixtures etc during testing
+	/// </summary>
+	public class ExperimentalTestSetup : BuildTestSetup
+	{
+		public ExperimentalTestSetup()
+		{
+			AddPlugin<ExperimentalPlugin>();
+		}
+
+		protected override void ConfigureServices(IServiceCollection services)
+		{
+			base.ConfigureServices(services);
+
+			services.AddSingleton<INotificationSink, ExperimentalSlackNotificationSink>();
+		}
+	}
+}
\ No newline at end of file
diff --git a/Engine/Source/Programs/Horde/Plugins/Experimental/HordeServer.Experimental.Tests/HordeServer.Experimental.Tests.csproj b/Engine/Source/Programs/Horde/Plugins/Experimental/HordeServer.Experimental.Tests/HordeServer.Experimental.Tests.csproj
new file mode 100644
--- /dev/null
+++ b/Engine/Source/Programs/Horde/Plugins/Experimental/HordeServer.Experimental.Tests/HordeServer.Experimental.Tests.csproj
@@ -0,0 +1,46 @@
+<Project Sdk="Microsoft.NET.Sdk">
+
+  <PropertyGroup>
+    <TargetFramework>net8.0</TargetFramework>
+    <ImplicitUsings>enable</ImplicitUsings>
+    <Nullable>enable</Nullable>
+    <IsPackable>false</IsPackable>
+    <IsTestProject>true</IsTestProject>
+    <RootNamespace>HordeServer.Tests</RootNamespace>
+    <RunAnalyzersDuringBuild>False</RunAnalyzersDuringBuild>
+    <Configurations>Debug;Release;Analyze</Configurations>
+    <NoWarn>1701;1702;1591</NoWarn>
+  </PropertyGroup>
+
+  <PropertyGroup Condition=" '$(Configuration)|$(Platform)' == 'Analyze|AnyCPU' ">
+    <RunAnalyzersDuringBuild>True</RunAnalyzersDuringBuild>
+    <TreatWarningsAsErrors>false</TreatWarningsAsErrors>
+    <EnforceCodeStyleInBuild>True</EnforceCodeStyleInBuild>
+    <GenerateDocumentationFile>True</GenerateDocumentationFile>
+  </PropertyGroup>
+
+  <ItemGroup>
+    <PackageReference Include="coverlet.collector" Version="6.0.2">
+      <PrivateAssets>all</PrivateAssets>
+      <IncludeAssets>runtime; build; native; contentfiles; analyzers; buildtransitive</IncludeAssets>
+    </PackageReference>
+    <PackageReference Include="Microsoft.NET.Test.Sdk" Version="17.11.1" />
+    <PackageReference Include="MSTest.TestAdapter" Version="3.6.1" />
+    <PackageReference Include="MSTest.TestFramework" Version="3.6.1" />
+    <PackageReference Include="Microsoft.AspNetCore.Mvc.Testing" Version="8.0.10" />
+    <PackageReference Include="Moq" Version="4.20.70" />
+    <PackageReference Include="Microsoft.Extensions.Caching.Memory" Version="8.0.1" />
+  </ItemGroup>
+
+  <ItemGroup>
+    <ProjectReference Include="..\..\Build\HordeServer.Build\HordeServer.Build.csproj" />
+    <ProjectReference Include="..\..\Build\HordeServer.Build.Tests\HordeServer.Build.Tests.csproj" />
+    <ProjectReference Include="..\HordeServer.Experimental\HordeServer.Experimental.csproj" />
+    <ProjectReference Include="..\..\..\HordeServer.Tests.Shared\HordeServer.Tests.Shared.csproj" />
+  </ItemGroup>
+
+  <ItemGroup>
+    <Using Include="Microsoft.VisualStudio.TestTools.UnitTesting" />
+  </ItemGroup>
+
+</Project>
diff --git a/Engine/Source/Programs/Horde/Plugins/Experimental/HordeServer.Experimental/Notifications/Sinks/ExperimentalSlackNotificationSink.cs b/Engine/Source/Programs/Horde/Plugins/Experimental/HordeServer.Experimental/Notifications/Sinks/ExperimentalSlackNotificationSink.cs
new file mode 100644
--- /dev/null
+++ b/Engine/Source/Programs/Horde/Plugins/Experimental/HordeServer.Experimental/Notifications/Sinks/ExperimentalSlackNotificationSink.cs
@@ -0,0 +1,98 @@
+// Copyright Epic Games, Inc. All Rights Reserved.
+
+using EpicGames.Horde.Jobs;
+using EpicGames.Slack;
+using HordeServer.Agents;
+using HordeServer.Configuration;
+using HordeServer.Devices;
+using HordeServer.Issues;
+using HordeServer.Jobs;
+using HordeServer.Jobs.Graphs;
+using HordeServer.Logs;
+using HordeServer.Streams;
+using HordeServer.Users;
+using Microsoft.Extensions.Logging;
+using Microsoft.Extensions.Options;
+using ZstdSharp.Unsafe;
+
+namespace HordeServer.Notifications.Sinks
+{
+	/// <summary>
+	/// Maintains a connection to Slack, in order to receive socket-mode notifications of user interactions
+	/// </summary>
+	public sealed class ExperimentalSlackNotificationSink : INotificationSink, IDisposable
+	{
+		readonly BuildServerConfig _buildServerConfig;
+		readonly ILogger _logger;
+
+		readonly HttpClient _httpClient;
+		readonly SlackClient _slackClient;
+
+		/// <summary>
+		/// Constructor
+		/// </summary>
+		public ExperimentalSlackNotificationSink(IOptions<BuildServerConfig> buildServerConfig, ILogger<ExperimentalSlackNotificationSink> logger)
+		{
+			_buildServerConfig = buildServerConfig.Value;
+			_logger = logger;
+
+			_httpClient = new HttpClient();
+			_httpClient.DefaultRequestHeaders.Add("Authorization", $"Bearer {_buildServerConfig.SlackToken ?? ""}");
+			_slackClient = new SlackClient(_httpClient, _logger);
+
+			_ = _slackClient;
+		}
+
+		/// <inheritdoc/>
+		public void Dispose()
+			=> _httpClient.Dispose();
+
+		/// <inheritdoc/>
+		public Task NotifyConfigUpdateAsync(ConfigUpdateInfo info, CancellationToken cancellationToken)
+			=> Task.CompletedTask;
+
+		/// <inheritdoc/>
+		public Task NotifyConfigUpdateFailureAsync(string errorMessage, string fileName, int? change = null, IUser? author = null, string? description = null, CancellationToken cancellationToken = default)
+			=> Task.CompletedTask;
+
+		/// <inheritdoc/>
+		public Task NotifyDeviceServiceAsync(string message, IDevice? device = null, IDevicePool? pool = null, StreamConfig? streamConfig = null, IJob? job = null, IJobStep? step = null, INode? node = null, IUser? user = null, CancellationToken cancellationToken = default)
+			=> Task.CompletedTask;
+
+		/// <inheritdoc/>
+		public Task NotifyIssueUpdatedAsync(IIssue issue, CancellationToken cancellationToken)
+			=> Task.CompletedTask;
+
+		/// <inheritdoc/>
+		public Task NotifyJobCompleteAsync(IJob job, IGraph graph, LabelOutcome outcome, CancellationToken cancellationToken)
+			=> Task.CompletedTask;
+
+		/// <inheritdoc/>
+		public Task NotifyJobCompleteAsync(IUser user, IJob job, IGraph graph, LabelOutcome outcome, CancellationToken cancellationToken)
+			=> Task.CompletedTask;
+
+		/// <inheritdoc/>
+		public Task NotifyJobScheduledAsync(List<JobScheduledNotification> notifications, CancellationToken cancellationToken)
+			=> Task.CompletedTask;
+
+		/// <inheritdoc/>
+		public Task NotifyJobStepCompleteAsync(IEnumerable<IUser> usersToNotify, IJob job, IJobStepBatch batch, IJobStep step, INode node, List<ILogEventData> jobStepEventData, CancellationToken cancellationToken)
+			=> Task.CompletedTask;
+
+		/// <inheritdoc/>
+		public Task NotifyLabelCompleteAsync(IUser user, IJob job, ILabel label, int labelIdx, LabelOutcome outcome, List<(string, JobStepOutcome, Uri)> stepData, CancellationToken cancellationToken)
+			=> Task.CompletedTask;
+
+		/// <inheritdoc/>
+		public Task SendAgentReportAsync(AgentReport report, CancellationToken cancellationToken)
+			=> Task.CompletedTask;
+
+		/// <inheritdoc/>
+		public Task SendDeviceIssueReportAsync(DeviceIssueReport report, CancellationToken cancellationToken)
+			=> Task.CompletedTask;
+
+		/// <inheritdoc/>
+		public Task SendIssueReportAsync(IssueReportGroup report, CancellationToken cancellationToken)
+			=> Task.CompletedTask;
+	}
+}
