Horde: Other files for experimental plugin.

#rnx

diff --git a/Engine/Source/Programs/Horde/Plugins/Experimental/HordeServer.Experimental/ExperimentalConfig.cs b/Engine/Source/Programs/Horde/Plugins/Experimental/HordeServer.Experimental/ExperimentalConfig.cs
new file mode 100644
--- /dev/null
+++ b/Engine/Source/Programs/Horde/Plugins/Experimental/HordeServer.Experimental/ExperimentalConfig.cs
@@ -0,0 +1,17 @@
+// Copyright Epic Games, Inc. All Rights Reserved.
+
+using HordeServer.Plugins;
+
+namespace HordeServer
+{
+	/// <summary>
+	/// Configuration for the experimental plugin
+	/// </summary>
+	public class ExperimentalConfig : IPluginConfig
+	{
+		/// <inheritdoc/>
+		public void PostLoad(PluginConfigOptions configOptions)
+		{
+		}
+	}
+}
diff --git a/Engine/Source/Programs/Horde/Plugins/Experimental/HordeServer.Experimental/ExperimentalPlugin.cs b/Engine/Source/Programs/Horde/Plugins/Experimental/HordeServer.Experimental/ExperimentalPlugin.cs
new file mode 100644
--- /dev/null
+++ b/Engine/Source/Programs/Horde/Plugins/Experimental/HordeServer.Experimental/ExperimentalPlugin.cs
@@ -0,0 +1,72 @@
+// Copyright Epic Games, Inc. All Rights Reserved.
+
+using System.Diagnostics.CodeAnalysis;
+using HordeServer.Notifications;
+using HordeServer.Notifications.Sinks;
+using HordeServer.Plugins;
+using Microsoft.AspNetCore.Builder;
+using Microsoft.Extensions.DependencyInjection;
+
+namespace HordeServer
+{
+	/// <summary>
+	/// Entry point for the experimental plugin
+	/// </summary>
+	[Plugin("Experimental", EnabledByDefault = false, GlobalConfigType = typeof(ExperimentalConfig), ServerConfigType = typeof(ExperimentalServerConfig))]
+	public class ExperimentalPlugin : IPluginStartup
+	{
+		/// <summary>
+		/// Constructor
+		/// </summary>
+		public ExperimentalPlugin()
+		{
+		}
+
+		/// <inheritdoc/>
+		public void Configure(IApplicationBuilder app)
+		{
+		}
+
+		/// <inheritdoc/>
+		public void ConfigureServices(IServiceCollection services)
+		{
+			services.AddSingleton<INotificationSink, ExperimentalSlackNotificationSink>();
+		}
+	}
+
+	/// <summary>
+	/// Helper methods for experimental plugin config
+	/// </summary>
+	public static class ExperimentalPluginExtensions
+	{
+		/// <summary>
+		/// Configures the build plugin
+		/// </summary>
+		public static void AddExperimentalConfig(this IDictionary<PluginName, IPluginConfig> dictionary, ExperimentalConfig experimentalConfig)
+			=> dictionary[new PluginName("Experimental")] = experimentalConfig;
+
+		/// <summary>
+		/// Gets configuration for the build plugin
+		/// </summary>
+		public static ExperimentalConfig GetExperimentalConfig(this IDictionary<PluginName, IPluginConfig> dictionary)
+			=> (ExperimentalConfig)dictionary[new PluginName("Experimental")];
+
+		/// <summary>
+		/// Gets configuration for the build plugin
+		/// </summary>
+		public static bool TryGetExperimentalConfig(this IDictionary<PluginName, IPluginConfig> dictionary, [NotNullWhen(true)] out ExperimentalConfig? experimentalConfig)
+		{
+			IPluginConfig? pluginConfig;
+			if (dictionary.TryGetValue(new PluginName("Experimental"), out pluginConfig) && pluginConfig is ExperimentalConfig newExperimentalConfig)
+			{
+				experimentalConfig = newExperimentalConfig;
+				return true;
+			}
+			else
+			{
+				experimentalConfig = null;
+				return false;
+			}
+		}
+	}
+}
diff --git a/Engine/Source/Programs/Horde/Plugins/Experimental/HordeServer.Experimental/ExperimentalServerConfig.cs b/Engine/Source/Programs/Horde/Plugins/Experimental/HordeServer.Experimental/ExperimentalServerConfig.cs
new file mode 100644
--- /dev/null
+++ b/Engine/Source/Programs/Horde/Plugins/Experimental/HordeServer.Experimental/ExperimentalServerConfig.cs
@@ -0,0 +1,13 @@
+// Copyright Epic Games, Inc. All Rights Reserved.
+
+using HordeServer.Plugins;
+
+namespace HordeServer
+{
+	/// <summary>
+	/// Server configuration for the experimental plugin
+	/// </summary>
+	public class ExperimentalServerConfig : PluginServerConfig
+	{
+	}
+}
diff --git a/Engine/Source/Programs/Horde/Plugins/Experimental/HordeServer.Experimental/HordeServer.Experimental.csproj b/Engine/Source/Programs/Horde/Plugins/Experimental/HordeServer.Experimental/HordeServer.Experimental.csproj
new file mode 100644
--- /dev/null
+++ b/Engine/Source/Programs/Horde/Plugins/Experimental/HordeServer.Experimental/HordeServer.Experimental.csproj
@@ -0,0 +1,46 @@
+﻿<Project Sdk="Microsoft.NET.Sdk">
+
+  <PropertyGroup>
+    <TargetFramework>net8.0</TargetFramework>
+    <ImplicitUsings>enable</ImplicitUsings>
+    <Nullable>enable</Nullable>
+    <RootNamespace>HordeServer</RootNamespace>
+    <RunAnalyzersDuringBuild>False</RunAnalyzersDuringBuild>
+    <Configurations>Debug;Release;Analyze</Configurations>
+  </PropertyGroup>
+
+  <PropertyGroup Condition=" '$(Configuration)|$(Platform)' == 'Analyze|AnyCPU' ">
+    <RunAnalyzersDuringBuild>True</RunAnalyzersDuringBuild>
+    <TreatWarningsAsErrors>false</TreatWarningsAsErrors>
+    <EnforceCodeStyleInBuild>True</EnforceCodeStyleInBuild>
+    <GenerateDocumentationFile>True</GenerateDocumentationFile>
+    <DocumentationFile></DocumentationFile>
+  </PropertyGroup>
+
+  <ItemGroup>
+    <PackageReference Include="Grpc.AspNetCore" Version="2.59.0" />
+    <PackageReference Include="Grpc.AspNetCore.Server.Reflection" Version="2.59.0" />
+    <PackageReference Include="Grpc.Net.Client" Version="2.59.0" />
+    <PackageReference Include="Grpc.Tools" Version="2.59.0">
+      <PrivateAssets>all</PrivateAssets>
+      <IncludeAssets>runtime; build; native; contentfiles; analyzers; buildtransitive</IncludeAssets>
+    </PackageReference>
+    <PackageReference Include="Polly" Version="8.2.0" />
+    <PackageReference Include="TimeZoneConverter" Version="6.1.0" />
+  </ItemGroup>
+
+  <ItemGroup>
+    <ProjectReference Include="..\..\..\..\Shared\EpicGames.Core\EpicGames.Core.csproj" />
+    <ProjectReference Include="..\..\..\..\Shared\EpicGames.Horde\EpicGames.Horde.csproj" />
+    <ProjectReference Include="..\..\..\..\Shared\EpicGames.IoHash\EpicGames.IoHash.csproj" />
+    <ProjectReference Include="..\..\..\..\Shared\EpicGames.Perforce.Managed\EpicGames.Perforce.Managed.csproj" />
+    <ProjectReference Include="..\..\..\..\Shared\EpicGames.Perforce\EpicGames.Perforce.csproj" />
+    <ProjectReference Include="..\..\..\..\Shared\EpicGames.Redis\EpicGames.Redis.csproj" />
+    <ProjectReference Include="..\..\..\..\Shared\EpicGames.Slack\EpicGames.Slack.csproj" />
+    <ProjectReference Include="..\..\..\HordeServer.Shared\HordeServer.Shared.csproj" />
+    <ProjectReference Include="..\..\Build\HordeServer.Build\HordeServer.Build.csproj" />
+    <ProjectReference Include="..\..\Compute\HordeServer.Compute\HordeServer.Compute.csproj" />
+    <ProjectReference Include="..\..\Tools\HordeServer.Tools\HordeServer.Tools.csproj" />
+  </ItemGroup>
+
+</Project>
