Horde: Enable the experimental plugin for local development.

#rnx

diff --git a/Engine/Source/Programs/Horde/HordeServer/appsettings.Local.json b/Engine/Source/Programs/Horde/HordeServer/appsettings.Local.json
--- a/Engine/Source/Programs/Horde/HordeServer/appsettings.Local.json
+++ b/Engine/Source/Programs/Horde/HordeServer/appsettings.Local.json
@@ -23,6 +23,13 @@
 
 			// Don't replicate content from Perforce commits
 			"ReplicateContent": false
+		},
+
+		// Enable the experimental plugin for local development
+		"Plugins": {
+			"Experimental": {
+				"Enabled": true
+			}
 		}
 	},
 	"Serilog": {
