Update nuget packages

#rnx

diff --git a/Engine/Source/Programs/AutomationTool.Tests/AutomationTool.Tests.csproj b/Engine/Source/Programs/AutomationTool.Tests/AutomationTool.Tests.csproj
--- a/Engine/Source/Programs/AutomationTool.Tests/AutomationTool.Tests.csproj
+++ b/Engine/Source/Programs/AutomationTool.Tests/AutomationTool.Tests.csproj
@@ -12,8 +12,8 @@
 
   <ItemGroup>
     <PackageReference Include="Microsoft.NET.Test.Sdk" Version="17.11.1" />
-    <PackageReference Include="MSTest.TestAdapter" Version="3.6.1" />
-    <PackageReference Include="MSTest.TestFramework" Version="3.6.1" />
+    <PackageReference Include="MSTest.TestAdapter" Version="3.6.2" />
+    <PackageReference Include="MSTest.TestFramework" Version="3.6.2" />
     <PackageReference Include="coverlet.collector" Version="6.0.2">
       <PrivateAssets>all</PrivateAssets>
       <IncludeAssets>runtime; build; native; contentfiles; analyzers; buildtransitive</IncludeAssets>
diff --git a/Engine/Source/Programs/Horde/Drivers/JobDriver.Tests/JobDriver.Tests.csproj b/Engine/Source/Programs/Horde/Drivers/JobDriver.Tests/JobDriver.Tests.csproj
--- a/Engine/Source/Programs/Horde/Drivers/JobDriver.Tests/JobDriver.Tests.csproj
+++ b/Engine/Source/Programs/Horde/Drivers/JobDriver.Tests/JobDriver.Tests.csproj
@@ -22,8 +22,8 @@
     <PackageReference Include="Microsoft.Extensions.Logging.Abstractions" Version="8.0.2" />
     <PackageReference Include="Microsoft.NET.Test.Sdk" Version="17.11.1" />
     <PackageReference Include="Moq" Version="4.20.70" />
-    <PackageReference Include="MSTest.TestAdapter" Version="3.6.1" />
-    <PackageReference Include="MSTest.TestFramework" Version="3.6.1" />
+    <PackageReference Include="MSTest.TestAdapter" Version="3.6.2" />
+    <PackageReference Include="MSTest.TestFramework" Version="3.6.2" />
     <PackageReference Include="coverlet.collector" Version="6.0.2">
       <PrivateAssets>all</PrivateAssets>
       <IncludeAssets>runtime; build; native; contentfiles; analyzers; buildtransitive</IncludeAssets>
diff --git a/Engine/Source/Programs/Horde/HordeAgent.Tests/HordeAgent.Tests.csproj b/Engine/Source/Programs/Horde/HordeAgent.Tests/HordeAgent.Tests.csproj
--- a/Engine/Source/Programs/Horde/HordeAgent.Tests/HordeAgent.Tests.csproj
+++ b/Engine/Source/Programs/Horde/HordeAgent.Tests/HordeAgent.Tests.csproj
@@ -23,8 +23,8 @@
     <PackageReference Include="Microsoft.Extensions.Logging.Abstractions" Version="8.0.2" />
     <PackageReference Include="Microsoft.NET.Test.Sdk" Version="17.11.1" />
     <PackageReference Include="Moq" Version="4.20.70" />
-    <PackageReference Include="MSTest.TestAdapter" Version="3.6.1" />
-    <PackageReference Include="MSTest.TestFramework" Version="3.6.1" />
+    <PackageReference Include="MSTest.TestAdapter" Version="3.6.2" />
+    <PackageReference Include="MSTest.TestFramework" Version="3.6.2" />
     <PackageReference Include="coverlet.collector" Version="6.0.2">
       <PrivateAssets>all</PrivateAssets>
       <IncludeAssets>runtime; build; native; contentfiles; analyzers; buildtransitive</IncludeAssets>
diff --git a/Engine/Source/Programs/Horde/HordeServer.Tests.Shared/HordeServer.Tests.Shared.csproj b/Engine/Source/Programs/Horde/HordeServer.Tests.Shared/HordeServer.Tests.Shared.csproj
--- a/Engine/Source/Programs/Horde/HordeServer.Tests.Shared/HordeServer.Tests.Shared.csproj
+++ b/Engine/Source/Programs/Horde/HordeServer.Tests.Shared/HordeServer.Tests.Shared.csproj
@@ -25,8 +25,8 @@
     <PackageReference Include="Microsoft.AspNetCore.Mvc.Testing" Version="8.0.10" />
     <PackageReference Include="Microsoft.NET.Test.Sdk" Version="17.11.1" />
     <PackageReference Include="Moq" Version="4.20.70" />
-    <PackageReference Include="MSTest.TestAdapter" Version="3.6.1" />
-    <PackageReference Include="MSTest.TestFramework" Version="3.6.1" />
+    <PackageReference Include="MSTest.TestAdapter" Version="3.6.2" />
+    <PackageReference Include="MSTest.TestFramework" Version="3.6.2" />
     <PackageReference Include="Microsoft.Extensions.Caching.Memory" Version="8.0.1" />
     <PackageReference Include="coverlet.collector" Version="6.0.2">
       <PrivateAssets>all</PrivateAssets>
diff --git a/Engine/Source/Programs/Horde/HordeServer.Tests/HordeServer.Tests.csproj b/Engine/Source/Programs/Horde/HordeServer.Tests/HordeServer.Tests.csproj
--- a/Engine/Source/Programs/Horde/HordeServer.Tests/HordeServer.Tests.csproj
+++ b/Engine/Source/Programs/Horde/HordeServer.Tests/HordeServer.Tests.csproj
@@ -22,8 +22,8 @@
     <PackageReference Include="Microsoft.AspNetCore.Mvc.Testing" Version="8.0.10" />
     <PackageReference Include="Microsoft.NET.Test.Sdk" Version="17.11.1" />
     <PackageReference Include="Moq" Version="4.20.70" />
-    <PackageReference Include="MSTest.TestAdapter" Version="3.6.1" />
-    <PackageReference Include="MSTest.TestFramework" Version="3.6.1" />
+    <PackageReference Include="MSTest.TestAdapter" Version="3.6.2" />
+    <PackageReference Include="MSTest.TestFramework" Version="3.6.2" />
     <PackageReference Include="Microsoft.Extensions.Caching.Memory" Version="8.0.1" />
     <PackageReference Include="coverlet.collector" Version="6.0.2">
       <PrivateAssets>all</PrivateAssets>
diff --git a/Engine/Source/Programs/Horde/Plugins/Analytics/HordeServer.Analytics.Tests/HordeServer.Analytics.Tests.csproj b/Engine/Source/Programs/Horde/Plugins/Analytics/HordeServer.Analytics.Tests/HordeServer.Analytics.Tests.csproj
--- a/Engine/Source/Programs/Horde/Plugins/Analytics/HordeServer.Analytics.Tests/HordeServer.Analytics.Tests.csproj
+++ b/Engine/Source/Programs/Horde/Plugins/Analytics/HordeServer.Analytics.Tests/HordeServer.Analytics.Tests.csproj
@@ -24,8 +24,8 @@
       <IncludeAssets>runtime; build; native; contentfiles; analyzers; buildtransitive</IncludeAssets>
     </PackageReference>
     <PackageReference Include="Microsoft.NET.Test.Sdk" Version="17.11.1" />
-    <PackageReference Include="MSTest.TestAdapter" Version="3.6.1" />
-    <PackageReference Include="MSTest.TestFramework" Version="3.6.1" />
+    <PackageReference Include="MSTest.TestAdapter" Version="3.6.2" />
+    <PackageReference Include="MSTest.TestFramework" Version="3.6.2" />
   </ItemGroup>
 
   <ItemGroup>
diff --git a/Engine/Source/Programs/Horde/Plugins/Build/HordeServer.Build.Tests/HordeServer.Build.Tests.csproj b/Engine/Source/Programs/Horde/Plugins/Build/HordeServer.Build.Tests/HordeServer.Build.Tests.csproj
--- a/Engine/Source/Programs/Horde/Plugins/Build/HordeServer.Build.Tests/HordeServer.Build.Tests.csproj
+++ b/Engine/Source/Programs/Horde/Plugins/Build/HordeServer.Build.Tests/HordeServer.Build.Tests.csproj
@@ -25,8 +25,8 @@
       <IncludeAssets>runtime; build; native; contentfiles; analyzers; buildtransitive</IncludeAssets>
     </PackageReference>
     <PackageReference Include="Microsoft.NET.Test.Sdk" Version="17.11.1" />
-    <PackageReference Include="MSTest.TestAdapter" Version="3.6.1" />
-    <PackageReference Include="MSTest.TestFramework" Version="3.6.1" />
+    <PackageReference Include="MSTest.TestAdapter" Version="3.6.2" />
+    <PackageReference Include="MSTest.TestFramework" Version="3.6.2" />
     <PackageReference Include="Microsoft.AspNetCore.Mvc.Testing" Version="8.0.10" />
     <PackageReference Include="Moq" Version="4.20.70" />
     <PackageReference Include="Microsoft.Extensions.Caching.Memory" Version="8.0.1" />
diff --git a/Engine/Source/Programs/Horde/Plugins/Compute/HordeServer.Compute.Tests/HordeServer.Compute.Tests.csproj b/Engine/Source/Programs/Horde/Plugins/Compute/HordeServer.Compute.Tests/HordeServer.Compute.Tests.csproj
--- a/Engine/Source/Programs/Horde/Plugins/Compute/HordeServer.Compute.Tests/HordeServer.Compute.Tests.csproj
+++ b/Engine/Source/Programs/Horde/Plugins/Compute/HordeServer.Compute.Tests/HordeServer.Compute.Tests.csproj
@@ -25,8 +25,8 @@
       <IncludeAssets>runtime; build; native; contentfiles; analyzers; buildtransitive</IncludeAssets>
     </PackageReference>
     <PackageReference Include="Microsoft.NET.Test.Sdk" Version="17.11.1" />
-    <PackageReference Include="MSTest.TestAdapter" Version="3.6.1" />
-    <PackageReference Include="MSTest.TestFramework" Version="3.6.1" />
+    <PackageReference Include="MSTest.TestAdapter" Version="3.6.2" />
+    <PackageReference Include="MSTest.TestFramework" Version="3.6.2" />
   </ItemGroup>
 
   <ItemGroup>
diff --git a/Engine/Source/Programs/Horde/Plugins/Ddc/HordeServer.Ddc.Tests/HordeServer.Ddc.Tests.csproj b/Engine/Source/Programs/Horde/Plugins/Ddc/HordeServer.Ddc.Tests/HordeServer.Ddc.Tests.csproj
--- a/Engine/Source/Programs/Horde/Plugins/Ddc/HordeServer.Ddc.Tests/HordeServer.Ddc.Tests.csproj
+++ b/Engine/Source/Programs/Horde/Plugins/Ddc/HordeServer.Ddc.Tests/HordeServer.Ddc.Tests.csproj
@@ -25,8 +25,8 @@
       <IncludeAssets>runtime; build; native; contentfiles; analyzers; buildtransitive</IncludeAssets>
     </PackageReference>
     <PackageReference Include="Microsoft.NET.Test.Sdk" Version="17.11.1" />
-    <PackageReference Include="MSTest.TestAdapter" Version="3.6.1" />
-    <PackageReference Include="MSTest.TestFramework" Version="3.6.1" />
+    <PackageReference Include="MSTest.TestAdapter" Version="3.6.2" />
+    <PackageReference Include="MSTest.TestFramework" Version="3.6.2" />
   </ItemGroup>
 
   <ItemGroup>
diff --git a/Engine/Source/Programs/Horde/Plugins/Experimental/HordeServer.Experimental.Tests/HordeServer.Experimental.Tests.csproj b/Engine/Source/Programs/Horde/Plugins/Experimental/HordeServer.Experimental.Tests/HordeServer.Experimental.Tests.csproj
--- a/Engine/Source/Programs/Horde/Plugins/Experimental/HordeServer.Experimental.Tests/HordeServer.Experimental.Tests.csproj
+++ b/Engine/Source/Programs/Horde/Plugins/Experimental/HordeServer.Experimental.Tests/HordeServer.Experimental.Tests.csproj
@@ -25,8 +25,8 @@
       <IncludeAssets>runtime; build; native; contentfiles; analyzers; buildtransitive</IncludeAssets>
     </PackageReference>
     <PackageReference Include="Microsoft.NET.Test.Sdk" Version="17.11.1" />
-    <PackageReference Include="MSTest.TestAdapter" Version="3.6.1" />
-    <PackageReference Include="MSTest.TestFramework" Version="3.6.1" />
+    <PackageReference Include="MSTest.TestAdapter" Version="3.6.2" />
+    <PackageReference Include="MSTest.TestFramework" Version="3.6.2" />
     <PackageReference Include="Microsoft.AspNetCore.Mvc.Testing" Version="8.0.10" />
     <PackageReference Include="Moq" Version="4.20.70" />
     <PackageReference Include="Microsoft.Extensions.Caching.Memory" Version="8.0.1" />
diff --git a/Engine/Source/Programs/Horde/Plugins/Storage/HordeServer.Storage.Tests/HordeServer.Storage.Tests.csproj b/Engine/Source/Programs/Horde/Plugins/Storage/HordeServer.Storage.Tests/HordeServer.Storage.Tests.csproj
--- a/Engine/Source/Programs/Horde/Plugins/Storage/HordeServer.Storage.Tests/HordeServer.Storage.Tests.csproj
+++ b/Engine/Source/Programs/Horde/Plugins/Storage/HordeServer.Storage.Tests/HordeServer.Storage.Tests.csproj
@@ -25,8 +25,8 @@
       <IncludeAssets>runtime; build; native; contentfiles; analyzers; buildtransitive</IncludeAssets>
     </PackageReference>
     <PackageReference Include="Microsoft.NET.Test.Sdk" Version="17.11.1" />
-    <PackageReference Include="MSTest.TestAdapter" Version="3.6.1" />
-    <PackageReference Include="MSTest.TestFramework" Version="3.6.1" />
+    <PackageReference Include="MSTest.TestAdapter" Version="3.6.2" />
+    <PackageReference Include="MSTest.TestFramework" Version="3.6.2" />
   </ItemGroup>
 
   <ItemGroup>
diff --git a/Engine/Source/Programs/Horde/Plugins/Tools/HordeServer.Tools.Tests/HordeServer.Tools.Tests.csproj b/Engine/Source/Programs/Horde/Plugins/Tools/HordeServer.Tools.Tests/HordeServer.Tools.Tests.csproj
--- a/Engine/Source/Programs/Horde/Plugins/Tools/HordeServer.Tools.Tests/HordeServer.Tools.Tests.csproj
+++ b/Engine/Source/Programs/Horde/Plugins/Tools/HordeServer.Tools.Tests/HordeServer.Tools.Tests.csproj
@@ -25,8 +25,8 @@
       <IncludeAssets>runtime; build; native; contentfiles; analyzers; buildtransitive</IncludeAssets>
     </PackageReference>
     <PackageReference Include="Microsoft.NET.Test.Sdk" Version="17.11.1" />
-    <PackageReference Include="MSTest.TestAdapter" Version="3.6.1" />
-    <PackageReference Include="MSTest.TestFramework" Version="3.6.1" />
+    <PackageReference Include="MSTest.TestAdapter" Version="3.6.2" />
+    <PackageReference Include="MSTest.TestFramework" Version="3.6.2" />
   </ItemGroup>
 
   <ItemGroup>
diff --git a/Engine/Source/Programs/Shared/EpicGames.BuildGraph.Tests/EpicGames.BuildGraph.Tests.csproj b/Engine/Source/Programs/Shared/EpicGames.BuildGraph.Tests/EpicGames.BuildGraph.Tests.csproj
--- a/Engine/Source/Programs/Shared/EpicGames.BuildGraph.Tests/EpicGames.BuildGraph.Tests.csproj
+++ b/Engine/Source/Programs/Shared/EpicGames.BuildGraph.Tests/EpicGames.BuildGraph.Tests.csproj
@@ -46,10 +46,10 @@
 
   <ItemGroup>
     <PackageReference Include="Microsoft.NET.Test.Sdk" Version="17.11.1" />
-    <PackageReference Include="MSTest.TestAdapter" Version="3.6.1">
+    <PackageReference Include="MSTest.TestAdapter" Version="3.6.2">
       <TreatAsUsed>true</TreatAsUsed>
     </PackageReference>
-    <PackageReference Include="MSTest.TestFramework" Version="3.6.1" />
+    <PackageReference Include="MSTest.TestFramework" Version="3.6.2" />
     <PackageReference Include="coverlet.collector" Version="6.0.2">
       <PrivateAssets>all</PrivateAssets>
       <IncludeAssets>runtime; build; native; contentfiles; analyzers; buildtransitive</IncludeAssets>
diff --git a/Engine/Source/Programs/Shared/EpicGames.Core.Tests/EpicGames.Core.Tests.csproj b/Engine/Source/Programs/Shared/EpicGames.Core.Tests/EpicGames.Core.Tests.csproj
--- a/Engine/Source/Programs/Shared/EpicGames.Core.Tests/EpicGames.Core.Tests.csproj
+++ b/Engine/Source/Programs/Shared/EpicGames.Core.Tests/EpicGames.Core.Tests.csproj
@@ -50,10 +50,10 @@
 
   <ItemGroup>
     <PackageReference Include="Microsoft.NET.Test.Sdk" Version="17.11.1" />
-    <PackageReference Include="MSTest.TestAdapter" Version="3.6.1">
+    <PackageReference Include="MSTest.TestAdapter" Version="3.6.2">
       <TreatAsUsed>true</TreatAsUsed>
     </PackageReference>
-    <PackageReference Include="MSTest.TestFramework" Version="3.6.1" />
+    <PackageReference Include="MSTest.TestFramework" Version="3.6.2" />
     <PackageReference Include="coverlet.collector" Version="6.0.2">
       <PrivateAssets>all</PrivateAssets>
       <IncludeAssets>runtime; build; native; contentfiles; analyzers; buildtransitive</IncludeAssets>
diff --git a/Engine/Source/Programs/Shared/EpicGames.Core/EpicGames.Core.csproj b/Engine/Source/Programs/Shared/EpicGames.Core/EpicGames.Core.csproj
--- a/Engine/Source/Programs/Shared/EpicGames.Core/EpicGames.Core.csproj
+++ b/Engine/Source/Programs/Shared/EpicGames.Core/EpicGames.Core.csproj
@@ -52,7 +52,7 @@
     </When>
   </Choose>
   <ItemGroup>
-    <PackageReference Include="JetBrains.Annotations" Version="2024.2.0" />
+    <PackageReference Include="JetBrains.Annotations" Version="2024.3.0" />
     <PackageReference Include="Microsoft.Extensions.Configuration" Version="8.0.0" />
     <PackageReference Include="Microsoft.Extensions.Http.Polly" Version="8.0.10" />
     <PackageReference Include="Microsoft.Extensions.Logging.Console" Version="8.0.1" />
diff --git a/Engine/Source/Programs/Shared/EpicGames.Horde.Tests/EpicGames.Horde.Tests.csproj b/Engine/Source/Programs/Shared/EpicGames.Horde.Tests/EpicGames.Horde.Tests.csproj
--- a/Engine/Source/Programs/Shared/EpicGames.Horde.Tests/EpicGames.Horde.Tests.csproj
+++ b/Engine/Source/Programs/Shared/EpicGames.Horde.Tests/EpicGames.Horde.Tests.csproj
@@ -32,10 +32,10 @@
 
   <ItemGroup>
     <PackageReference Include="Microsoft.NET.Test.Sdk" Version="17.11.1" />
-    <PackageReference Include="MSTest.TestAdapter" Version="3.6.1">
+    <PackageReference Include="MSTest.TestAdapter" Version="3.6.2">
       <TreatAsUsed>true</TreatAsUsed>
     </PackageReference>
-    <PackageReference Include="MSTest.TestFramework" Version="3.6.1" />
+    <PackageReference Include="MSTest.TestFramework" Version="3.6.2" />
     <PackageReference Include="coverlet.collector" Version="6.0.2">
       <PrivateAssets>all</PrivateAssets>
       <IncludeAssets>runtime; build; native; contentfiles; analyzers; buildtransitive</IncludeAssets>
diff --git a/Engine/Source/Programs/Shared/EpicGames.IoHash.Tests/EpicGames.IoHash.Tests.csproj b/Engine/Source/Programs/Shared/EpicGames.IoHash.Tests/EpicGames.IoHash.Tests.csproj
--- a/Engine/Source/Programs/Shared/EpicGames.IoHash.Tests/EpicGames.IoHash.Tests.csproj
+++ b/Engine/Source/Programs/Shared/EpicGames.IoHash.Tests/EpicGames.IoHash.Tests.csproj
@@ -19,10 +19,10 @@
 
   <ItemGroup>
     <PackageReference Include="Microsoft.NET.Test.Sdk" Version="17.11.1" />
-    <PackageReference Include="MSTest.TestAdapter" Version="3.6.1">
+    <PackageReference Include="MSTest.TestAdapter" Version="3.6.2">
       <TreatAsUsed>true</TreatAsUsed>
     </PackageReference>
-    <PackageReference Include="MSTest.TestFramework" Version="3.6.1" />
+    <PackageReference Include="MSTest.TestFramework" Version="3.6.2" />
     <PackageReference Include="coverlet.collector" Version="6.0.2">
       <PrivateAssets>all</PrivateAssets>
       <IncludeAssets>runtime; build; native; contentfiles; analyzers; buildtransitive</IncludeAssets>
diff --git a/Engine/Source/Programs/Shared/EpicGames.Perforce.Fixture/EpicGames.Perforce.Fixture.csproj b/Engine/Source/Programs/Shared/EpicGames.Perforce.Fixture/EpicGames.Perforce.Fixture.csproj
--- a/Engine/Source/Programs/Shared/EpicGames.Perforce.Fixture/EpicGames.Perforce.Fixture.csproj
+++ b/Engine/Source/Programs/Shared/EpicGames.Perforce.Fixture/EpicGames.Perforce.Fixture.csproj
@@ -34,10 +34,10 @@
   <ItemGroup>
     <PackageReference Include="Microsoft.Extensions.Logging.Console" Version="8.0.1" />
     <PackageReference Include="Microsoft.NET.Test.Sdk" Version="17.11.1" />
-    <PackageReference Include="MSTest.TestAdapter" Version="3.6.1">
+    <PackageReference Include="MSTest.TestAdapter" Version="3.6.2">
       <TreatAsUsed>true</TreatAsUsed>
     </PackageReference>
-    <PackageReference Include="MSTest.TestFramework" Version="3.6.1" />
+    <PackageReference Include="MSTest.TestFramework" Version="3.6.2" />
     <PackageReference Include="coverlet.collector" Version="6.0.2">
       <PrivateAssets>all</PrivateAssets>
       <IncludeAssets>runtime; build; native; contentfiles; analyzers; buildtransitive</IncludeAssets>
diff --git a/Engine/Source/Programs/Shared/EpicGames.Perforce.Managed.Tests/EpicGames.Perforce.Managed.Tests.csproj b/Engine/Source/Programs/Shared/EpicGames.Perforce.Managed.Tests/EpicGames.Perforce.Managed.Tests.csproj
--- a/Engine/Source/Programs/Shared/EpicGames.Perforce.Managed.Tests/EpicGames.Perforce.Managed.Tests.csproj
+++ b/Engine/Source/Programs/Shared/EpicGames.Perforce.Managed.Tests/EpicGames.Perforce.Managed.Tests.csproj
@@ -34,10 +34,10 @@
     <PackageReference Include="Microsoft.NET.Test.Sdk" Version="17.11.1">
       <TreatAsUsed>true</TreatAsUsed>
     </PackageReference>
-    <PackageReference Include="MSTest.TestAdapter" Version="3.6.1">
+    <PackageReference Include="MSTest.TestAdapter" Version="3.6.2">
       <TreatAsUsed>true</TreatAsUsed>
     </PackageReference>
-    <PackageReference Include="MSTest.TestFramework" Version="3.6.1">
+    <PackageReference Include="MSTest.TestFramework" Version="3.6.2">
       <TreatAsUsed>true</TreatAsUsed>
     </PackageReference>
     <PackageReference Include="coverlet.collector" Version="6.0.2">
diff --git a/Engine/Source/Programs/Shared/EpicGames.Perforce.Tests/EpicGames.Perforce.Tests.csproj b/Engine/Source/Programs/Shared/EpicGames.Perforce.Tests/EpicGames.Perforce.Tests.csproj
--- a/Engine/Source/Programs/Shared/EpicGames.Perforce.Tests/EpicGames.Perforce.Tests.csproj
+++ b/Engine/Source/Programs/Shared/EpicGames.Perforce.Tests/EpicGames.Perforce.Tests.csproj
@@ -31,10 +31,10 @@
 
   <ItemGroup>
     <PackageReference Include="Microsoft.NET.Test.Sdk" Version="17.11.1" />
-    <PackageReference Include="MSTest.TestAdapter" Version="3.6.1">
+    <PackageReference Include="MSTest.TestAdapter" Version="3.6.2">
       <TreatAsUsed>true</TreatAsUsed>
     </PackageReference>
-    <PackageReference Include="MSTest.TestFramework" Version="3.6.1" />
+    <PackageReference Include="MSTest.TestFramework" Version="3.6.2" />
     <PackageReference Include="coverlet.collector" Version="6.0.2">
       <PrivateAssets>all</PrivateAssets>
       <IncludeAssets>runtime; build; native; contentfiles; analyzers; buildtransitive</IncludeAssets>
diff --git a/Engine/Source/Programs/Shared/EpicGames.Redis.Tests/EpicGames.Redis.Tests.csproj b/Engine/Source/Programs/Shared/EpicGames.Redis.Tests/EpicGames.Redis.Tests.csproj
--- a/Engine/Source/Programs/Shared/EpicGames.Redis.Tests/EpicGames.Redis.Tests.csproj
+++ b/Engine/Source/Programs/Shared/EpicGames.Redis.Tests/EpicGames.Redis.Tests.csproj
@@ -32,10 +32,10 @@
 
   <ItemGroup>
     <PackageReference Include="Microsoft.NET.Test.Sdk" Version="17.11.1" />
-    <PackageReference Include="MSTest.TestAdapter" Version="3.6.1">
+    <PackageReference Include="MSTest.TestAdapter" Version="3.6.2">
       <TreatAsUsed>true</TreatAsUsed>
     </PackageReference>
-    <PackageReference Include="MSTest.TestFramework" Version="3.6.1" />
+    <PackageReference Include="MSTest.TestFramework" Version="3.6.2" />
     <PackageReference Include="coverlet.collector" Version="6.0.2">
       <PrivateAssets>all</PrivateAssets>
       <IncludeAssets>runtime; build; native; contentfiles; analyzers; buildtransitive</IncludeAssets>
diff --git a/Engine/Source/Programs/Shared/EpicGames.Serialization.Tests/EpicGames.Serialization.Tests.csproj b/Engine/Source/Programs/Shared/EpicGames.Serialization.Tests/EpicGames.Serialization.Tests.csproj
--- a/Engine/Source/Programs/Shared/EpicGames.Serialization.Tests/EpicGames.Serialization.Tests.csproj
+++ b/Engine/Source/Programs/Shared/EpicGames.Serialization.Tests/EpicGames.Serialization.Tests.csproj
@@ -35,10 +35,10 @@
 
   <ItemGroup>
     <PackageReference Include="Microsoft.NET.Test.Sdk" Version="17.11.1" />
-    <PackageReference Include="MSTest.TestAdapter" Version="3.6.1">
+    <PackageReference Include="MSTest.TestAdapter" Version="3.6.2">
       <TreatAsUsed>true</TreatAsUsed>
     </PackageReference>
-    <PackageReference Include="MSTest.TestFramework" Version="3.6.1" />
+    <PackageReference Include="MSTest.TestFramework" Version="3.6.2" />
     <PackageReference Include="coverlet.collector" Version="6.0.2">
       <PrivateAssets>all</PrivateAssets>
       <IncludeAssets>runtime; build; native; contentfiles; analyzers; buildtransitive</IncludeAssets>
diff --git a/Engine/Source/Programs/Shared/EpicGames.Tracing.Tests/EpicGames.Tracing.Tests.csproj b/Engine/Source/Programs/Shared/EpicGames.Tracing.Tests/EpicGames.Tracing.Tests.csproj
--- a/Engine/Source/Programs/Shared/EpicGames.Tracing.Tests/EpicGames.Tracing.Tests.csproj
+++ b/Engine/Source/Programs/Shared/EpicGames.Tracing.Tests/EpicGames.Tracing.Tests.csproj
@@ -30,10 +30,10 @@
 
   <ItemGroup>
     <PackageReference Include="Microsoft.NET.Test.Sdk" Version="17.11.1" />
-    <PackageReference Include="MSTest.TestAdapter" Version="3.6.1">
+    <PackageReference Include="MSTest.TestAdapter" Version="3.6.2">
       <TreatAsUsed>true</TreatAsUsed>
     </PackageReference>
-    <PackageReference Include="MSTest.TestFramework" Version="3.6.1" />
+    <PackageReference Include="MSTest.TestFramework" Version="3.6.2" />
     <PackageReference Include="coverlet.collector" Version="6.0.2">
       <PrivateAssets>all</PrivateAssets>
       <IncludeAssets>runtime; build; native; contentfiles; analyzers; buildtransitive</IncludeAssets>
diff --git a/Engine/Source/Programs/Shared/UnrealBuildTool.Tests/UnrealBuildTool.Tests.csproj b/Engine/Source/Programs/Shared/UnrealBuildTool.Tests/UnrealBuildTool.Tests.csproj
--- a/Engine/Source/Programs/Shared/UnrealBuildTool.Tests/UnrealBuildTool.Tests.csproj
+++ b/Engine/Source/Programs/Shared/UnrealBuildTool.Tests/UnrealBuildTool.Tests.csproj
@@ -18,10 +18,10 @@
 
   <ItemGroup>
     <PackageReference Include="Microsoft.NET.Test.Sdk" Version="17.11.1" />
-    <PackageReference Include="MSTest.TestAdapter" Version="3.6.1">
+    <PackageReference Include="MSTest.TestAdapter" Version="3.6.2">
       <TreatAsUsed>true</TreatAsUsed>
     </PackageReference>
-    <PackageReference Include="MSTest.TestFramework" Version="3.6.1" />
+    <PackageReference Include="MSTest.TestFramework" Version="3.6.2" />
     <PackageReference Include="coverlet.collector" Version="6.0.2">
       <PrivateAssets>all</PrivateAssets>
       <IncludeAssets>runtime; build; native; contentfiles; analyzers; buildtransitive</IncludeAssets>
