Horde: Fix static analysis warning.

#rnx

diff --git a/Engine/Source/Programs/Horde/Plugins/Experimental/HordeServer.Experimental/Notifications/Sinks/ExperimentalSlackNotificationSink.cs b/Engine/Source/Programs/Horde/Plugins/Experimental/HordeServer.Experimental/Notifications/Sinks/ExperimentalSlackNotificationSink.cs
--- a/Engine/Source/Programs/Horde/Plugins/Experimental/HordeServer.Experimental/Notifications/Sinks/ExperimentalSlackNotificationSink.cs
+++ b/Engine/Source/Programs/Horde/Plugins/Experimental/HordeServer.Experimental/Notifications/Sinks/ExperimentalSlackNotificationSink.cs
@@ -13,7 +13,6 @@ using HordeServer.Streams;
 using HordeServer.Users;
 using Microsoft.Extensions.Logging;
 using Microsoft.Extensions.Options;
-using ZstdSharp.Unsafe;
 
 namespace HordeServer.Notifications.Sinks
 {
