Horde: Fix PerforceServiceCache.SubscribeAsync() stalling out when more than 10 commits are parsed in a single iteration, due to changes being enumerated in the incorrect order. Now uses a custom query to search for commits in the correct order, and only uses cached commit metadata.

diff --git a/Engine/Source/Programs/Horde/Plugins/Build/HordeServer.Build/Perforce/PerforceServiceCache.cs b/Engine/Source/Programs/Horde/Plugins/Build/HordeServer.Build/Perforce/PerforceServiceCache.cs
--- a/Engine/Source/Programs/Horde/Plugins/Build/HordeServer.Build/Perforce/PerforceServiceCache.cs
+++ b/Engine/Source/Programs/Horde/Plugins/Build/HordeServer.Build/Perforce/PerforceServiceCache.cs
@@ -784,18 +784,10 @@ namespace HordeServer.Perforce
 
 						Task task = updateEvent.Task;
 
-						int numResults = 10;
-						await foreach (ICommit commit in FindAsync(minChange, includeMinChange: false, maxResults: numResults, tags: tags, cancellationToken: cancellationToken))
+						await foreach(ICommit commit in FindCachedCommitsInAscendingOrderAsync(minChange, tags, cancellationToken))
 						{
 							yield return commit;
 							minChange = commit.Id.GetPerforceChange();
-							numResults--;
-						}
-
-						if (numResults == 0)
-						{
-							// Query again; received the max requested number of changes.
-							continue;
 						}
 
 						if (subscription == null)
@@ -815,6 +807,36 @@ namespace HordeServer.Perforce
 					}
 				}
 			}
+
+			async IAsyncEnumerable<ICommit> FindCachedCommitsInAscendingOrderAsync(int minChange, IReadOnlyList<CommitTag>? tags, [EnumeratorCancellation] CancellationToken cancellationToken)
+			{
+				FilterDefinition<CachedCommitDoc> filter = Builders<CachedCommitDoc>.Filter.Eq(x => x.StreamId, StreamConfig.Id);
+				filter &= Builders<CachedCommitDoc>.Filter.Gt(x => x.Number, minChange);
+
+				if (tags != null && tags.Count > 0)
+				{
+					if (tags.Count == 1)
+					{
+						filter &= Builders<CachedCommitDoc>.Filter.AnyEq(x => x.CommitTags, tags[0]);
+					}
+					else
+					{
+						filter &= Builders<CachedCommitDoc>.Filter.AnyIn(x => x.CommitTags, tags);
+					}
+				}
+
+				using (IAsyncCursor<CachedCommitDoc> cursor = await _owner._commits.Find(filter).SortBy(x => x.Number).ToCursorAsync(cancellationToken))
+				{
+					while (await cursor.MoveNextAsync(cancellationToken))
+					{
+						foreach (CachedCommitDoc commit in cursor.Current)
+						{
+							commit.PostLoad(_owner, StreamConfig);
+							yield return commit;
+						}
+					}
+				}
+			}
 		}
 
 		/// <inheritdoc/>
diff --git a/Engine/Source/Programs/Shared/EpicGames.Perforce/Records/ChangesRecord.cs b/Engine/Source/Programs/Shared/EpicGames.Perforce/Records/ChangesRecord.cs
--- a/Engine/Source/Programs/Shared/EpicGames.Perforce/Records/ChangesRecord.cs
+++ b/Engine/Source/Programs/Shared/EpicGames.Perforce/Records/ChangesRecord.cs
@@ -1,12 +1,14 @@
 // Copyright Epic Games, Inc. All Rights Reserved.
 
 using System;
+using System.Diagnostics;
 
 namespace EpicGames.Perforce
 {
 	/// <summary>
 	/// Contains summary information about a changelist
 	/// </summary>
+	[DebuggerDisplay("{Number}: {Description}")]
 	public class ChangesRecord
 	{
 		/// <summary>
