Horde: Show steps with issues inline when there is a label selected
#skipci

diff --git a/Engine/Source/Programs/Horde/HordeDashboard/src/components/jobDetailsV2/JobDetailLabels.tsx b/Engine/Source/Programs/Horde/HordeDashboard/src/components/jobDetailsV2/JobDetailLabels.tsx
--- a/Engine/Source/Programs/Horde/HordeDashboard/src/components/jobDetailsV2/JobDetailLabels.tsx
+++ b/Engine/Source/Programs/Horde/HordeDashboard/src/components/jobDetailsV2/JobDetailLabels.tsx
@@ -1,14 +1,15 @@
 import { ConstrainMode, DefaultButton, DetailsList, DetailsListLayoutMode, IColumn, SelectionMode, Stack, Text } from "@fluentui/react";
 import { observer } from "mobx-react-lite";
 import { useNavigate, useLocation } from "react-router-dom";
-import { LabelOutcome, LabelState } from "../../backend/Api";
+import { GetStepResponse, JobData, JobStepOutcome, LabelOutcome, LabelState, StepData } from "../../backend/Api";
 import { JobLabel } from "../../backend/JobDetails";
 import { getLabelColor } from "../../styles/colors";
 import { useQuery } from "../JobDetailCommon";
 import { JobDataView, JobDetailsV2 } from "./JobDetailsViewCommon";
 import { getHordeStyling } from "../../styles/Styles";
 import dashboard from "../../backend/Dashboard";
-
+import { Link } from "react-router-dom";
+import { StepStatusIcon } from "../StatusIcon";
 
 export const LabelsPanelV2: React.FC<{ jobDetails: JobDetailsV2, dataView: JobDataView }> = observer(({ jobDetails, dataView }) => {
 
@@ -123,6 +124,72 @@ export const LabelsPanelV2: React.FC<{ jobDetails: JobDetailsV2, dataView: JobDa
       return null;
    }
 
+
+   const label = jobDetails.filter.label;
+
+   const job = jobDetails.jobData!;
+
+   type StepItem = {
+      step: StepData;
+      name: string;
+   }
+
+   const jobId = job.id;
+
+   if (!job.batches) {
+      return null;
+   }
+
+   let steps: GetStepResponse[] = [];
+
+   if (label) {
+      steps = job.batches.map(b => b.steps).flat().filter(step => !!step.startTime && (step.outcome === JobStepOutcome.Warnings || step.outcome === JobStepOutcome.Failure));
+      steps = steps.filter(s => label?.steps.indexOf(s.id) === -1 ? false : true);
+   }
+
+   const renderSteps = () => {
+
+      if (!steps || !steps.length) {
+         return null;
+      }
+
+      const onRenderCell = (stepItem: StepItem): JSX.Element => {
+
+         const step = stepItem.step;
+
+         const stepUrl = `/job/${jobId}?step=${step.id}`;
+
+         return <Stack key={`step_${step.id}_job_${job.id}_${stepItem.name}`}>
+            <Stack horizontal>
+               <Link to={stepUrl} onClick={(ev) => { ev.stopPropagation(); }}><div style={{ cursor: "pointer" }}>
+                  <Stack horizontal verticalAlign="center">
+                     <StepStatusIcon step={step} style={{ fontSize: 13 }} />
+                     <Text styles={{ root: { fontSize: 13, paddingRight: 4, paddingTop: 2, userSelect: "none" } }}>{`${stepItem.name}`}</Text>
+                  </Stack>
+               </div></Link>
+            </Stack>
+         </Stack>;
+      };
+
+
+      let items = steps.map(step => {
+
+
+         return {
+            step: step,
+            name: step.name
+         }
+
+      })
+
+      items = items.filter(item => !!item);
+
+      const render = items.map(s => onRenderCell(s!));
+
+      return <Stack horizontal wrap tokens={{ childrenGap: 12 }}>{render}</Stack>;
+
+   }
+
    const buildColumns = (): IColumn[] => {
 
       const widths: Record<string, number> = {
@@ -216,6 +283,14 @@ export const LabelsPanelV2: React.FC<{ jobDetails: JobDetailsV2, dataView: JobDa
                onShouldVirtualize={() => { return false; }}
             />
          </Stack>
+         <Stack style={{ paddingTop: 12, paddingLeft: 12 }} horizontal>
+            <Stack>
+               <Text style={{ fontFamily: "Horde Open Sans SemiBold" }}>Issues</Text>
+            </Stack>
+            <Stack style={{paddingLeft: 100}}>
+               {renderSteps()}
+            </Stack>
+         </Stack>
       </Stack>
    </Stack>;
 
