Horde: Include a flag in template responses indicating whether the user can run it.

diff --git a/Engine/Source/Programs/Horde/Plugins/Build/HordeServer.Build/Streams/StreamMessages.cs b/Engine/Source/Programs/Horde/Plugins/Build/HordeServer.Build/Streams/StreamMessages.cs
--- a/Engine/Source/Programs/Horde/Plugins/Build/HordeServer.Build/Streams/StreamMessages.cs
+++ b/Engine/Source/Programs/Horde/Plugins/Build/HordeServer.Build/Streams/StreamMessages.cs
@@ -260,6 +260,11 @@ namespace HordeServer.Streams
 		/// </summary>
 		public List<ChangeQueryConfig>? DefaultChange { get; set; }
 
+		/// <summary>
+		/// Whether the user is allowed to create jobs from this template
+		/// </summary>
+		public bool CanRun { get; set; }
+
 		/// <summary>
 		/// Constructor
 		/// </summary>
@@ -268,7 +273,7 @@ namespace HordeServer.Streams
 		/// <param name="template">The actual template</param>
 		/// <param name="stepStates">The template step states</param>
 		/// <param name="schedulerTimeZone">The scheduler time zone</param>
-		public GetTemplateRefResponse(TemplateId id, ITemplateRef templateRef, ITemplate template, List<GetTemplateStepStateResponse>? stepStates, TimeZoneInfo schedulerTimeZone)
+		public GetTemplateRefResponse(TemplateId id, ITemplateRef templateRef, ITemplate template, List<GetTemplateStepStateResponse>? stepStates, TimeZoneInfo schedulerTimeZone, bool canRun)
 			: base(template)
 		{
 			Id = id.ToString();
@@ -281,6 +286,7 @@ namespace HordeServer.Streams
 			ChainedJobs = (templateRef.Config.ChainedJobs != null && templateRef.Config.ChainedJobs.Count > 0) ? templateRef.Config.ChainedJobs : null;
 			StepStates = stepStates;
 			DefaultChange = templateRef.Config.DefaultChange;
+			CanRun = canRun;
 		}
 	}
 
diff --git a/Engine/Source/Programs/Horde/Plugins/Build/HordeServer.Build/Streams/StreamsController.cs b/Engine/Source/Programs/Horde/Plugins/Build/HordeServer.Build/Streams/StreamsController.cs
--- a/Engine/Source/Programs/Horde/Plugins/Build/HordeServer.Build/Streams/StreamsController.cs
+++ b/Engine/Source/Programs/Horde/Plugins/Build/HordeServer.Build/Streams/StreamsController.cs
@@ -218,8 +218,10 @@ namespace HordeServer.Streams
 						}
 					}
 
+					bool canRun = templateRef.Config.Authorize(JobAclAction.CreateJob, User);
+
 					ITemplate? template = await _templateCollection.GetOrAddAsync(pair.Value.Config);
-					apiTemplateRefs.Add(new GetTemplateRefResponse(pair.Key, pair.Value, template, stepStates, _timeZone));
+					apiTemplateRefs.Add(new GetTemplateRefResponse(pair.Key, pair.Value, template, stepStates, _timeZone, canRun));
 				}
 			}
 
