Horde: Add a canUpdate flag to job responses indicated whether the user is authorized to modify the job.

#rnx

diff --git a/Engine/Source/Programs/Horde/Plugins/Build/HordeServer.Build/Jobs/JobsController.cs b/Engine/Source/Programs/Horde/Plugins/Build/HordeServer.Build/Jobs/JobsController.cs
--- a/Engine/Source/Programs/Horde/Plugins/Build/HordeServer.Build/Jobs/JobsController.cs
+++ b/Engine/Source/Programs/Horde/Plugins/Build/HordeServer.Build/Jobs/JobsController.cs
@@ -516,7 +516,13 @@ namespace HordeServer.Jobs
 				abortedByUserInfo = (await _userCollection.GetCachedUserAsync(job.AbortedByUserId.Value, cancellationToken))?.ToThinApiResponse();
 			}
 
-			GetJobResponse response = CreateGetJobResponse(job, startedByUserInfo, abortedByUserInfo);
+			bool canUpdate = false;
+			if (_buildConfig.Value.TryGetStream(job.StreamId, out StreamConfig? streamConfig))
+			{
+				canUpdate = streamConfig.Authorize(JobAclAction.UpdateJob, User);
+			}
+
+			GetJobResponse response = CreateGetJobResponse(job, startedByUserInfo, abortedByUserInfo, canUpdate);
 			if (includeBatches || includeLabels)
 			{
 				if (includeBatches)
@@ -601,7 +607,7 @@ namespace HordeServer.Jobs
 				&& artifact.Type != ArtifactType.StepTestData;
 		}
 
-		static GetJobResponse CreateGetJobResponse(IJob job, GetThinUserInfoResponse? startedByUserInfo, GetThinUserInfoResponse? abortedByUserInfo)
+		static GetJobResponse CreateGetJobResponse(IJob job, GetThinUserInfoResponse? startedByUserInfo, GetThinUserInfoResponse? abortedByUserInfo, bool canUpdate)
 		{
 			GetJobResponse response = new GetJobResponse(job.Id, job.StreamId, job.TemplateId, job.Name);
 			response.CommitId = job.CommitId;
@@ -631,6 +637,7 @@ namespace HordeServer.Jobs
 			response.UpdateTime = new DateTimeOffset(job.UpdateTimeUtc);
 			response.UseArtifactsV2 = true;
 			response.UpdateIssues = job.UpdateIssues;
+			response.CanUpdate = canUpdate;
 			return response;
 		}
 
diff --git a/Engine/Source/Programs/Shared/EpicGames.Horde/Jobs/JobMessages.cs b/Engine/Source/Programs/Shared/EpicGames.Horde/Jobs/JobMessages.cs
--- a/Engine/Source/Programs/Shared/EpicGames.Horde/Jobs/JobMessages.cs
+++ b/Engine/Source/Programs/Shared/EpicGames.Horde/Jobs/JobMessages.cs
@@ -630,6 +630,11 @@ namespace EpicGames.Horde.Jobs
 		/// </summary>
 		public bool UpdateIssues { get; set; }
 
+		/// <summary>
+		/// Whether the current user is allowed to update this job
+		/// </summary>
+		public bool CanUpdate { get; set; } = true;
+
 		/// <summary>
 		/// Constructor
 		/// </summary>
