Horde: Add stats for artifact upload performance.

#rnx

diff --git a/Engine/Source/Programs/AutomationTool/BuildGraph/Tasks/CreateArtifactTask.cs b/Engine/Source/Programs/AutomationTool/BuildGraph/Tasks/CreateArtifactTask.cs
--- a/Engine/Source/Programs/AutomationTool/BuildGraph/Tasks/CreateArtifactTask.cs
+++ b/Engine/Source/Programs/AutomationTool/BuildGraph/Tasks/CreateArtifactTask.cs
@@ -3,6 +3,7 @@
 using System;
 using System.Collections.Generic;
 using System.Diagnostics;
+using System.IO;
 using System.Linq;
 using System.Threading;
 using System.Threading.Tasks;
@@ -191,10 +192,13 @@ namespace AutomationTool.Tasks
 				throw new AutomationException($"Unable to create artifact {name} with given file list.");
 			}
 
+			// Find the total file size
+			long totalSize = files.Sum(x => x.ToFileInfo().Length);
+
 			// Create the new artifact
 			IHordeClient hordeClient = CommandUtils.ServiceProvider.GetRequiredService<IHordeClient>();
 			IArtifactBuilder artifact = await hordeClient.Artifacts.CreateAsync(name, type, _parameters.Description, streamId, commitId, keys, metadata);
-			Logger.LogInformation("Creating artifact {ArtifactId} '{ArtifactName}' ({ArtifactType}) with namespace {NamespaceId}, ref {RefName} ({Link})", artifact.Id, name, type, artifact.NamespaceId, artifact.RefName, new Uri(hordeClient.ServerUrl, $"/api/v1/storage/{artifact.NamespaceId}/refs/{artifact.RefName}").ToString());
+			Logger.LogInformation("Creating artifact {ArtifactId} '{ArtifactName}' ({ArtifactType}) from {NumFiles} files ({TotalSize:n1}mb). Namespace: {NamespaceId}, ref {RefName} ({Link})", artifact.Id, name, type,files.Count, totalSize / (1024.0 * 1024.0), artifact.NamespaceId, artifact.RefName, new Uri(hordeClient.ServerUrl, $"/api/v1/storage/{artifact.NamespaceId}/refs/{artifact.RefName}").ToString());
 
 			// Upload the files
 			Stopwatch timer = Stopwatch.StartNew();
@@ -271,7 +275,7 @@ namespace AutomationTool.Tasks
 
 			// Finalize the artifact
 			await artifact.CompleteAsync(rootRef);
-			Logger.LogInformation("Uploaded artifact {ArtifactId} in {Time:n1}s", artifact.Id, timer.Elapsed.TotalSeconds);
+			Logger.LogInformation("Uploaded artifact {ArtifactId} in {Time:n1}s ({Rate:n1}mb/s, {RateMbps:n1}mbps)", artifact.Id, timer.Elapsed.TotalSeconds, totalSize / (timer.Elapsed.TotalSeconds * 1024.0 * 1024.0), (totalSize * 8.0) / (timer.Elapsed.TotalSeconds * 1024.0 * 1024.0));
 		}
 
 		static async Task<FileEntry?> FindFileAsync(IBlobRef<DirectoryNode> rootDir, string path, CancellationToken cancellationToken = default)
diff --git a/Engine/Source/Programs/Horde/Drivers/JobDriver/Execution/JobExecutor.cs b/Engine/Source/Programs/Horde/Drivers/JobDriver/Execution/JobExecutor.cs
--- a/Engine/Source/Programs/Horde/Drivers/JobDriver/Execution/JobExecutor.cs
+++ b/Engine/Source/Programs/Horde/Drivers/JobDriver/Execution/JobExecutor.cs
@@ -1138,12 +1138,7 @@ namespace JobDriver.Execution
 			return result;
 		}
 		
-		protected async Task CreateArtifactsAsync(JobStepId stepId, ArtifactName name, ArtifactType type, DirectoryReference baseDir, IEnumerable<(string, FileReference)> files, ILogger logger, CancellationToken cancellationToken)
-		{
-			await CreateArtifactAsync(stepId, name, type, baseDir, files.Select(x => x.Item2), logger, cancellationToken);
-		}
-
-		protected async Task CreateArtifactAsync(JobStepId stepId, ArtifactName name, ArtifactType type, DirectoryReference baseDir, IEnumerable<FileReference> files, ILogger logger, CancellationToken cancellationToken)
+		protected async Task CreateArtifactAsync(JobStepId stepId, ArtifactName name, ArtifactType type, DirectoryReference baseDir, IReadOnlyList<FileInfo> files, ILogger logger, CancellationToken cancellationToken)
 		{
 			try
 			{
@@ -1153,11 +1148,14 @@ namespace JobDriver.Execution
 				artifactRequest.Name = name.ToString();
 				artifactRequest.Type = type.ToString();
 
+				long totalSize = files.Sum(x => x.Length);
+
 				RpcCreateJobArtifactResponseV2 artifact = await JobRpc.CreateArtifactV2Async(artifactRequest, cancellationToken: cancellationToken);
 				ArtifactId artifactId = ArtifactId.Parse(artifact.Id);
-				Logger.LogInformation("Created artifact {ArtifactId} '{ArtifactName}' ({ArtifactType}) with ref {RefName} ({Link})", artifactId, name, type, artifact.RefName, $"{HordeClient.ServerUrl}/api/v1/storage/{artifact.NamespaceId}/refs/{artifact.RefName}");
+				Logger.LogInformation("Creating artifact {ArtifactId} '{ArtifactName}' ({ArtifactType}) from {NumFiles:n0} files ({TotalSize:n1}mb). Namespace {NamespaceId}, ref {RefName} ({Link})", artifactId, name, type, files.Count, totalSize / (1024.0 * 1024.0), artifact.NamespaceId, artifact.RefName, $"{HordeClient.ServerUrl}/api/v1/storage/{artifact.NamespaceId}/refs/{artifact.RefName}");
 
 				IStorageNamespace storage = GetStorageNamespace(new NamespaceId(artifact.NamespaceId), artifact.Token);
+				Stopwatch timer = Stopwatch.StartNew();
 
 				IHashedBlobRef<DirectoryNode> rootRef;
 				await using (IBlobWriter blobWriter = storage.CreateBlobWriter(new RefName(artifact.RefName)))
@@ -1176,7 +1174,8 @@ namespace JobDriver.Execution
 				}
 				await storage.WriteRefAsync(new RefName(artifact.RefName), rootRef, cancellationToken: cancellationToken);
 
-				Logger.LogInformation("Uploaded artifact {ArtifactId}", artifactId);
+				timer.Stop();
+				logger.LogInformation("Uploaded artifact {ArtifactId} in {Time:n1}s ({Rate:n1}mb/s, {RateMbps:n1}mbps)", artifact.Id, timer.Elapsed.TotalSeconds, totalSize / (timer.Elapsed.TotalSeconds * 1024.0 * 1024.0), (totalSize * 8.0) / (timer.Elapsed.TotalSeconds * 1024.0 * 1024.0));
 			}
 			catch (Exception ex)
 			{
@@ -1362,9 +1361,11 @@ namespace JobDriver.Execution
 				artifactRequest.Name = TempStorage.GetArtifactNameForNode(step.Name).ToString();
 				artifactRequest.Type = ArtifactType.StepOutput.ToString();
 
+				long totalSize = outputFileToBlockName.Keys.Sum(x => x.ToFileInfo().Length);
+
 				RpcCreateJobArtifactResponseV2 artifact = await JobRpc.CreateArtifactV2Async(artifactRequest, cancellationToken: cancellationToken);
 				ArtifactId artifactId = ArtifactId.Parse(artifact.Id);
-				logger.LogInformation("Created artifact {ArtifactId} '{ArtifactName}' ({ArtifactType}) with ref {RefName} ({RefUrl})", artifactId, artifactRequest.Name, ArtifactType.StepOutput, artifact.RefName, $"{HordeClient.ServerUrl.ToString().TrimEnd('/')}/api/v1/storage/{artifact.NamespaceId}/refs/{artifact.RefName}");
+				logger.LogInformation("Creating artifact {ArtifactId} '{ArtifactName}' ({ArtifactType}) from {NumFiles:n0} ({TotalSize:n1}mb). Namespace: {NamespaceId}, ref {RefName} ({RefUrl})", artifactId, artifactRequest.Name, ArtifactType.StepOutput, outputFileToBlockName.Count, totalSize / (1024.0 * 1024.0), artifact.NamespaceId, artifact.RefName, $"{HordeClient.ServerUrl.ToString().TrimEnd('/')}/api/v1/storage/{artifact.NamespaceId}/refs/{artifact.RefName}");
 				DecorateSpanWithArtifact(span, Prefix, artifactRequest, artifact);
 
 				IStorageNamespace storage = GetStorageNamespace(new NamespaceId(artifact.NamespaceId), artifact.Token);
@@ -1415,7 +1416,7 @@ namespace JobDriver.Execution
 
 				// Write the final node
 				await storage.WriteRefAsync(artifact.RefName, outputNodeRef, new RefOptions(), cancellationToken: cancellationToken);
-				logger.LogInformation("Upload took {Time:n1}s", timer.Elapsed.TotalSeconds);
+				logger.LogInformation("Uploaded artifact {ArtifactId} in {Time:n1}s ({Rate:n1}mb/s, {RateMbps:n1}mbps)", artifact.Id, timer.Elapsed.TotalSeconds, totalSize / (timer.Elapsed.TotalSeconds * 1024.0 * 1024.0), (totalSize * 8.0) / (timer.Elapsed.TotalSeconds * 1024.0 * 1024.0));
 			}
 
 			// Create all the named artifacts. TODO: Merge this with regular temp storage artifacts?
@@ -1448,9 +1449,13 @@ namespace JobDriver.Execution
 					artifactRequest.Keys.AddRange(graphArtifact.Keys);
 					artifactRequest.Metadata.AddRange(graphArtifact.Metadata);
 
+					List<FileInfo> fileList = files.Select(x => x.ToFileInfo()).ToList();
+					long totalSize = fileList.Sum(x => x.Length);
+
 					RpcCreateJobArtifactResponseV2 artifact = await JobRpc.CreateArtifactV2Async(artifactRequest, cancellationToken: cancellationToken);
 					ArtifactId artifactId = ArtifactId.Parse(artifact.Id);
-					logger.LogInformation("Created artifact {ArtifactId} '{ArtifactName}' ({ArtifactType}) with ref {RefName} ({RefUrl})", artifactId, artifactRequest.Name, ArtifactType.StepOutput, artifact.RefName, $"{HordeClient.ServerUrl.ToString().TrimEnd('/')}/api/v1/storage/{artifact.NamespaceId}/refs/{artifact.RefName}");
+					logger.LogInformation("Creating artifact {ArtifactId} '{ArtifactName}' ({ArtifactType}) from {NumFiles:n0} files ({TotalSize:n1}mb). Namespace {NamespaceId}, ref {RefName} ({RefUrl})", artifactId, artifactRequest.Name, ArtifactType.StepOutput, fileList.Count, totalSize / (1024.0 * 1024.0), artifact.NamespaceId, artifact.RefName, $"{HordeClient.ServerUrl.ToString().TrimEnd('/')}/api/v1/storage/{artifact.NamespaceId}/refs/{artifact.RefName}");
+
 					DecorateSpanWithArtifact(span, Prefix, artifactRequest, artifact);
 
 					IStorageNamespace storage = GetStorageNamespace(new NamespaceId(artifact.NamespaceId), artifact.Token);
@@ -1478,7 +1483,9 @@ namespace JobDriver.Execution
 
 					// Write the final node
 					await storage.WriteRefAsync(artifact.RefName, outputNodeRef, new RefOptions(), cancellationToken: cancellationToken);
-					logger.LogInformation("Upload took {Time:n1}s", timer.Elapsed.TotalSeconds);
+
+					timer.Stop();
+					logger.LogInformation("Uploaded artifact {ArtifactId} in {Time:n1}s ({Rate:n1}mb/s, {RateMbps:n1}mbps)", artifact.Id, timer.Elapsed.TotalSeconds, totalSize / (timer.Elapsed.TotalSeconds * 1024.0 * 1024.0), (totalSize * 8.0) / (timer.Elapsed.TotalSeconds * 1024.0 * 1024.0));
 				}
 			}
 
@@ -1918,38 +1925,40 @@ namespace JobDriver.Execution
 				filter.Flush();
 			}
 
-			if (DirectoryReference.Exists(telemetryDir))
+			DirectoryInfo telemetryDirInfo = telemetryDir.ToDirectoryInfo();
+			if (telemetryDirInfo.Exists)
 			{
-				List<(string, FileReference)> telemetryFiles = new List<(string, FileReference)>();
+				List<FileInfo> telemetryFiles = new List<FileInfo>();
 				try
 				{
-					await CreateTelemetryFilesAsync(step.Name, telemetryDir, telemetryFiles, buildGraphSpan, jobLogger, cancellationToken);
+					await CreateTelemetryFilesAsync(step.Name, telemetryDirInfo, telemetryFiles, buildGraphSpan, jobLogger, cancellationToken);
 				}
 				catch (Exception ex)
 				{
 					jobLogger.LogInformation(ex, "Unable to parse/upload telemetry files: {Message}", ex.Message);
 				}
 
-				await CreateArtifactsAsync(step.StepId, TempStorage.GetArtifactNameForNode(step.Name), ArtifactType.StepTrace, workspaceDir, telemetryFiles, jobLogger, cancellationToken);
+				await CreateArtifactAsync(step.StepId, TempStorage.GetArtifactNameForNode(step.Name), ArtifactType.StepTrace, workspaceDir, telemetryFiles, jobLogger, cancellationToken);
 
-				foreach ((_, FileReference telemetryFile) in telemetryFiles)
+				foreach (FileInfo telemetryFile in telemetryFiles)
 				{
 					FileUtils.ForceDeleteFile(telemetryFile);
 				}
 			}
 
-			if (DirectoryReference.Exists(testDataDir))
+			DirectoryInfo testDataDirInfo = testDataDir.ToDirectoryInfo();
+			if (testDataDirInfo.Exists)
 			{
-				List<(string, FileReference)> testDataFiles = new List<(string, FileReference)>();
+				List<FileInfo> testDataFiles = new List<FileInfo>();
 
 				Dictionary<string, object> combinedTestData = new Dictionary<string, object>();
-				foreach (FileReference testDataFile in DirectoryReference.EnumerateFiles(testDataDir, "*.json", SearchOption.AllDirectories))
+				foreach (FileInfo testDataFile in testDataDirInfo.EnumerateFiles("*.json", SearchOption.AllDirectories))
 				{
 					jobLogger.LogInformation("Reading test data {TestDataFile}", testDataFile);
-					testDataFiles.Add(($"TestData/{testDataFile.MakeRelativeTo(testDataDir)}", testDataFile));
+					testDataFiles.Add(testDataFile);
 
 					TestData testData;
-					using (FileStream stream = FileReference.Open(testDataFile, FileMode.Open))
+					using (FileStream stream = testDataFile.OpenRead())
 					{
 						JsonSerializerOptions options = new JsonSerializerOptions { PropertyNameCaseInsensitive = true };
 						testData = (await JsonSerializer.DeserializeAsync<TestData>(stream, options, cancellationToken))!;
@@ -1972,15 +1981,16 @@ namespace JobDriver.Execution
 				jobLogger.LogInformation("Found {NumResults} test results", combinedTestData.Count);
 				await UploadTestDataAsync(step.StepId, combinedTestData);
 
-				await CreateArtifactsAsync(step.StepId, TempStorage.GetArtifactNameForNode(step.Name), ArtifactType.StepTestData, workspaceDir, testDataFiles, jobLogger, cancellationToken);
+				await CreateArtifactAsync(step.StepId, TempStorage.GetArtifactNameForNode(step.Name), ArtifactType.StepTestData, workspaceDir, testDataFiles, jobLogger, cancellationToken);
 			}
 
-			if (DirectoryReference.Exists(logDir))
+			DirectoryInfo logDirInfo = logDir.ToDirectoryInfo();
+			if (logDirInfo.Exists)
 			{
-				List<FileReference> artifactFiles = DirectoryReference.EnumerateFiles(logDir, "*", SearchOption.AllDirectories).ToList();
+				FileInfo[] artifactFiles = logDirInfo.GetFiles("*", SearchOption.AllDirectories);
 				await CreateArtifactAsync(step.StepId, TempStorage.GetArtifactNameForNode(step.Name), ArtifactType.StepSaved, workspaceDir, artifactFiles, jobLogger, cancellationToken);
 
-				foreach (FileReference reportFile in artifactFiles.Where(x => x.HasExtension(".report.json")))
+				foreach (FileReference reportFile in artifactFiles.Select(x => new FileReference(x)).Where(x => x.HasExtension(".report.json")))
 				{
 					try
 					{
@@ -2032,18 +2042,20 @@ namespace JobDriver.Execution
 			return exitCode;
 		}
 
-		async Task CreateTelemetryFilesAsync(string stepName, DirectoryReference telemetryDir, List<(string, FileReference)> telemetryFiles, TelemetrySpan buildGraphSpan, ILogger jobLogger, CancellationToken cancellationToken)
+		async Task CreateTelemetryFilesAsync(string stepName, DirectoryInfo telemetryDir, List<FileInfo> telemetryFiles, TelemetrySpan buildGraphSpan, ILogger jobLogger, CancellationToken cancellationToken)
 		{
 			List<TraceEventList> telemetryList = new List<TraceEventList>();
-			foreach (FileReference telemetryFile in DirectoryReference.EnumerateFiles(telemetryDir, "*.json"))
+			foreach (FileInfo telemetryFile in telemetryDir.EnumerateFiles("*.json"))
 			{
-				jobLogger.LogInformation("Reading telemetry from {File}", telemetryFile);
-				byte[] data = await FileReference.ReadAllBytesAsync(telemetryFile, cancellationToken);
+				FileReference telemetryFileRef = new FileReference(telemetryFile);
+
+				jobLogger.LogInformation("Reading telemetry from {File}", telemetryFileRef);
+				byte[] data = await FileReference.ReadAllBytesAsync(telemetryFileRef, cancellationToken);
 
 				TraceEventList telemetry = JsonSerializer.Deserialize<TraceEventList>(data.AsSpan())!;
 				if (telemetry.Spans.Count > 0)
 				{
-					string defaultServiceName = telemetryFile.GetFileNameWithoutAnyExtensions();
+					string defaultServiceName = telemetryFileRef.GetFileNameWithoutAnyExtensions();
 					foreach (TraceEvent span in telemetry.Spans)
 					{
 						span.Service ??= defaultServiceName;
@@ -2051,7 +2063,7 @@ namespace JobDriver.Execution
 					telemetryList.Add(telemetry);
 				}
 
-				telemetryFiles.Add(($"Telemetry/{telemetryFile.GetFileName()}", telemetryFile));
+				telemetryFiles.Add(telemetryFile);
 			}
 
 			List<TraceEvent> telemetrySpans = new List<TraceEvent>();
@@ -2067,7 +2079,6 @@ namespace JobDriver.Execution
 				}
 			}
 
-			FileReference? traceFile = null;
 			if (telemetrySpans.Count > 0)
 			{
 				TraceSpan rootSpan = new TraceSpan();
@@ -2139,13 +2150,13 @@ namespace JobDriver.Execution
 				rootSpan.Start = rootSpan.Children!.First().Start;
 				rootSpan.Finish = rootSpan.Children!.Last().Finish;
 
-				traceFile = FileReference.Combine(telemetryDir, "Trace.json");
+				FileReference traceFile = FileReference.Combine(new DirectoryReference(telemetryDir), "Trace.json");
 				using (FileStream stream = FileReference.Open(traceFile, FileMode.Create))
 				{
 					JsonSerializerOptions options = new JsonSerializerOptions { DefaultIgnoreCondition = JsonIgnoreCondition.WhenWritingNull };
 					await JsonSerializer.SerializeAsync(stream, rootSpan, options, cancellationToken);
 				}
-				telemetryFiles.Add(("Trace.json", traceFile));
+				telemetryFiles.Add(traceFile.ToFileInfo());
 
 				CreateTracingData(buildGraphSpan, rootSpan);
 			}
