Horde: Add support for named blob uploads on tools endpoint. Also fix missing logging from tool deploy task.

#rnx

diff --git a/Engine/Source/Programs/AutomationTool/BuildGraph/Tasks/DeployToolTask.cs b/Engine/Source/Programs/AutomationTool/BuildGraph/Tasks/DeployToolTask.cs
--- a/Engine/Source/Programs/AutomationTool/BuildGraph/Tasks/DeployToolTask.cs
+++ b/Engine/Source/Programs/AutomationTool/BuildGraph/Tasks/DeployToolTask.cs
@@ -146,11 +146,14 @@ namespace AutomationTool.Tasks
 					options.AccessToken = settings.Token;
 				}
 			});
+			serviceCollection.AddLogging(builder => builder.AddProvider(CommandUtils.ServiceProvider.GetRequiredService<ILoggerProvider>()));
 			serviceCollection.AddHttpClient();
 			serviceCollection.AddHorde();
 
 			await using ServiceProvider serviceProvider = serviceCollection.BuildServiceProvider();
-			IHordeClient hordeClient = serviceProvider.GetRequiredService<IHordeClient>();
+			IHordeClient hordeClient = (settings == null) 
+				? CommandUtils.ServiceProvider.GetRequiredService<IHordeClient>()
+				: serviceProvider.GetRequiredService<IHordeClient>();
 
 			using HordeHttpClient hordeHttpClient = hordeClient.CreateHttpClient();
 
@@ -162,7 +165,7 @@ namespace AutomationTool.Tasks
 			IHashedBlobRef handle;
 
 			IStorageNamespace storageNamespace = hordeClient.GetStorageNamespace(toolId);
-			await using (IBlobWriter blobWriter = storageNamespace.CreateBlobWriter(serializerOptions: serializerOptions))
+			await using (IBlobWriter blobWriter = storageNamespace.CreateBlobWriter($"{toolId}", serializerOptions: serializerOptions))
 			{
 				DirectoryNode sandbox = new DirectoryNode();
 				if (_parameters.File != null)
diff --git a/Engine/Source/Programs/Horde/Plugins/Tools/HordeServer.Tools/Tools/ToolsController.cs b/Engine/Source/Programs/Horde/Plugins/Tools/HordeServer.Tools/Tools/ToolsController.cs
--- a/Engine/Source/Programs/Horde/Plugins/Tools/HordeServer.Tools/Tools/ToolsController.cs
+++ b/Engine/Source/Programs/Horde/Plugins/Tools/HordeServer.Tools/Tools/ToolsController.cs
@@ -65,6 +65,35 @@ namespace HordeServer.Tools
 			return await StorageController.WriteBlobAsync(storageBackend, null, request, cancellationToken);
 		}
 
+		/// <summary>
+		/// Uploads blob data for a new tool deployment.
+		/// </summary>
+		/// <param name="id">Identifier of the tool to upload</param>
+		/// <param name="locator">Locator for the blob</param>
+		/// <param name="request">Upload request</param>
+		/// <param name="cancellationToken">Cancellation token for the operation</param>
+		[HttpPut]
+		[Route("/api/v1/tools/{id}/blobs/{*locator}")]
+		public async Task<ActionResult<WriteBlobResponse>> WriteBlobAsync(ToolId id, BlobLocator locator, WriteBlobRequest request, CancellationToken cancellationToken = default)
+		{
+			ITool? tool = await _toolCollection.GetAsync(id, cancellationToken);
+			if (tool == null)
+			{
+				return NotFound(id);
+			}
+			if (!tool.Authorize(ToolAclAction.UploadTool, User))
+			{
+				return Forbid(ToolAclAction.UploadTool, id);
+			}
+			if (!locator.ToString().StartsWith($"{id}/", StringComparison.Ordinal))
+			{
+				return BadRequest("Tool blobs must start with tool id");
+			}
+
+			IStorageBackend storageBackend = tool.GetStorageBackend();
+			return await StorageController.WriteBlobAsync(storageBackend, locator, request, cancellationToken);
+		}
+
 		/// <summary>
 		/// Create a new deployment of the given tool.
 		/// </summary>
diff --git a/Engine/Source/Programs/Shared/EpicGames.Horde/Storage/Backends/HttpStorageBackend.cs b/Engine/Source/Programs/Shared/EpicGames.Horde/Storage/Backends/HttpStorageBackend.cs
--- a/Engine/Source/Programs/Shared/EpicGames.Horde/Storage/Backends/HttpStorageBackend.cs
+++ b/Engine/Source/Programs/Shared/EpicGames.Horde/Storage/Backends/HttpStorageBackend.cs
@@ -207,7 +207,7 @@ namespace EpicGames.Horde.Storage.Backends
 						if (!response.IsSuccessStatusCode)
 						{
 							string responseText = await response.Content.ReadAsStringAsync(cancellationToken);
-							throw new StorageException($"Upload to {request.RequestUri} failed ({response.StatusCode}). Response: {responseText}");
+							throw new StorageException($"{request.Method} to {request.RequestUri} failed ({response.StatusCode}). Response: {responseText}");
 						}
 
 						WriteBlobResponse? data = await response.Content.ReadFromJsonAsync<WriteBlobResponse>(cancellationToken: cancellationToken);
