Horde: Static analysis fixes.

#rnx

diff --git a/Engine/Source/Programs/Horde/HordeCmd/Commands/Compute/ComputeRun.cs b/Engine/Source/Programs/Horde/HordeCmd/Commands/Compute/ComputeRun.cs
--- a/Engine/Source/Programs/Horde/HordeCmd/Commands/Compute/ComputeRun.cs
+++ b/Engine/Source/Programs/Horde/HordeCmd/Commands/Compute/ComputeRun.cs
@@ -159,7 +159,7 @@ namespace Horde.Commands.Compute
 
 		static async Task<BlobLocator> CreateSandboxAsync(FileReference taskFile, IStorageNamespace storage, CancellationToken cancellationToken)
 		{
-			await using IBlobWriter writer = storage.CreateBlobWriter();
+			await using IBlobWriter writer = storage.CreateBlobWriter(cancellationToken: cancellationToken);
 
 			IHashedBlobRef<DirectoryNode> sandbox = await writer.WriteFilesAsync(taskFile.Directory, cancellationToken: cancellationToken);
 			await writer.FlushAsync(cancellationToken);
diff --git a/Engine/Source/Programs/Horde/Plugins/Ddc/HordeServer.Ddc/Ddc/BlobService.cs b/Engine/Source/Programs/Horde/Plugins/Ddc/HordeServer.Ddc/Ddc/BlobService.cs
--- a/Engine/Source/Programs/Horde/Plugins/Ddc/HordeServer.Ddc/Ddc/BlobService.cs
+++ b/Engine/Source/Programs/Horde/Plugins/Ddc/HordeServer.Ddc/Ddc/BlobService.cs
@@ -177,7 +177,7 @@ namespace HordeServer.Ddc
 			IStorageNamespace storageNamespace = _storageClient.GetNamespace(ns);
 
 			IHashedBlobRef blobRef;
-			await using (IBlobWriter writer = storageNamespace.CreateBlobWriter())
+			await using (IBlobWriter writer = storageNamespace.CreateBlobWriter(cancellationToken: cancellationToken))
 			{
 				int contentLength = (int)content.Length;
 				Memory<byte> memory = writer.GetMemory(contentLength).Slice(0, contentLength);
diff --git a/Engine/Source/Programs/Horde/Plugins/Tools/HordeServer.Tools/Tools/ToolCollection.cs b/Engine/Source/Programs/Horde/Plugins/Tools/HordeServer.Tools/Tools/ToolCollection.cs
--- a/Engine/Source/Programs/Horde/Plugins/Tools/HordeServer.Tools/Tools/ToolCollection.cs
+++ b/Engine/Source/Programs/Horde/Plugins/Tools/HordeServer.Tools/Tools/ToolCollection.cs
@@ -337,7 +337,7 @@ namespace HordeServer.Tools
 			IStorageNamespace client = _storageService.GetNamespace(toolConfig.NamespaceId);
 
 			IHashedBlobRef<DirectoryNode> nodeRef;
-			await using (IBlobWriter writer = client.CreateBlobWriter($"{tool.Id}/{deploymentId}"))
+			await using (IBlobWriter writer = client.CreateBlobWriter($"{tool.Id}/{deploymentId}", cancellationToken: cancellationToken))
 			{
 				DirectoryNode directoryNode = new DirectoryNode();
 				await directoryNode.CopyFromZipStreamAsync(stream, writer, new ChunkingOptions(), cancellationToken: cancellationToken);
diff --git a/Engine/Source/Programs/Shared/EpicGames.Horde.Tests/ComputeSocketTests.cs b/Engine/Source/Programs/Shared/EpicGames.Horde.Tests/ComputeSocketTests.cs
--- a/Engine/Source/Programs/Shared/EpicGames.Horde.Tests/ComputeSocketTests.cs
+++ b/Engine/Source/Programs/Shared/EpicGames.Horde.Tests/ComputeSocketTests.cs
@@ -172,7 +172,7 @@ namespace EpicGames.Horde.Tests
 					}
 
 					BundleStorageNamespace storage = BundleStorageNamespace.CreateInMemory(NullLogger.Instance);
-					await using (IBlobWriter blobWriter = storage.CreateBlobWriter())
+					await using (IBlobWriter blobWriter = storage.CreateBlobWriter(cancellationToken: : cancellationToken))
 					{
 						FileReference file = FileReference.Combine(tempDir, "subdir/hello.txt");
 						if (FileReference.Exists(file))
