Horde: Improve diagnostics from JobExecutor (use typed artifact ids, some messages written to wrong channel)

#rnx

diff --git a/Engine/Source/Programs/Horde/Drivers/JobDriver/Execution/JobExecutor.cs b/Engine/Source/Programs/Horde/Drivers/JobDriver/Execution/JobExecutor.cs
--- a/Engine/Source/Programs/Horde/Drivers/JobDriver/Execution/JobExecutor.cs
+++ b/Engine/Source/Programs/Horde/Drivers/JobDriver/Execution/JobExecutor.cs
@@ -1152,7 +1152,7 @@ namespace JobDriver.Execution
 
 				RpcCreateJobArtifactResponseV2 artifact = await JobRpc.CreateArtifactV2Async(artifactRequest, cancellationToken: cancellationToken);
 				ArtifactId artifactId = ArtifactId.Parse(artifact.Id);
-				Logger.LogInformation("Creating artifact {ArtifactId} '{ArtifactName}' ({ArtifactType}) from {NumFiles:n0} files ({TotalSize:n1}mb). Namespace {NamespaceId}, ref {RefName} ({Link})", artifactId, name, type, files.Count, totalSize / (1024.0 * 1024.0), artifact.NamespaceId, artifact.RefName, $"{HordeClient.ServerUrl}/api/v1/storage/{artifact.NamespaceId}/refs/{artifact.RefName}");
+				logger.LogInformation("Creating artifact {ArtifactId} '{ArtifactName}' ({ArtifactType}) from {NumFiles:n0} files ({TotalSize:n1}mb). Namespace {NamespaceId}, ref {RefName} ({Link})", artifactId, name, type, files.Count, totalSize / (1024.0 * 1024.0), artifact.NamespaceId, artifact.RefName, $"{HordeClient.ServerUrl}/api/v1/storage/{artifact.NamespaceId}/refs/{artifact.RefName}");
 
 				IStorageNamespace storage = GetStorageNamespace(new NamespaceId(artifact.NamespaceId), artifact.Token);
 				Stopwatch timer = Stopwatch.StartNew();
@@ -1175,11 +1175,11 @@ namespace JobDriver.Execution
 				await storage.WriteRefAsync(new RefName(artifact.RefName), rootRef, cancellationToken: cancellationToken);
 
 				timer.Stop();
-				logger.LogInformation("Uploaded artifact {ArtifactId} in {Time:n1}s ({Rate:n1}mb/s, {RateMbps:n1}mbps)", artifact.Id, timer.Elapsed.TotalSeconds, totalSize / (timer.Elapsed.TotalSeconds * 1024.0 * 1024.0), (totalSize * 8.0) / (timer.Elapsed.TotalSeconds * 1024.0 * 1024.0));
+				logger.LogInformation("Uploaded artifact {ArtifactId} in {Time:n1}s ({Rate:n1}mb/s, {RateMbps:n1}mbps)", artifactId, timer.Elapsed.TotalSeconds, totalSize / (timer.Elapsed.TotalSeconds * 1024.0 * 1024.0), (totalSize * 8.0) / (timer.Elapsed.TotalSeconds * 1024.0 * 1024.0));
 			}
 			catch (Exception ex)
 			{
-				Logger.LogInformation(ex, "Error creating artifact '{Type}'", type);
+				logger.LogInformation(ex, "Error creating artifact '{Type}'", type);
 			}
 		}
 
@@ -1365,7 +1365,7 @@ namespace JobDriver.Execution
 
 				RpcCreateJobArtifactResponseV2 artifact = await JobRpc.CreateArtifactV2Async(artifactRequest, cancellationToken: cancellationToken);
 				ArtifactId artifactId = ArtifactId.Parse(artifact.Id);
-				logger.LogInformation("Creating artifact {ArtifactId} '{ArtifactName}' ({ArtifactType}) from {NumFiles:n0} ({TotalSize:n1}mb). Namespace: {NamespaceId}, ref {RefName} ({RefUrl})", artifactId, artifactRequest.Name, ArtifactType.StepOutput, outputFileToBlockName.Count, totalSize / (1024.0 * 1024.0), artifact.NamespaceId, artifact.RefName, $"{HordeClient.ServerUrl.ToString().TrimEnd('/')}/api/v1/storage/{artifact.NamespaceId}/refs/{artifact.RefName}");
+				logger.LogInformation("Creating artifact {ArtifactId} '{ArtifactName}' ({ArtifactType}) from {NumFiles:n0} files ({TotalSize:n1}mb). Namespace: {NamespaceId}, ref {RefName} ({RefUrl})", artifactId, artifactRequest.Name, ArtifactType.StepOutput, outputFileToBlockName.Count, totalSize / (1024.0 * 1024.0), artifact.NamespaceId, artifact.RefName, $"{HordeClient.ServerUrl.ToString().TrimEnd('/')}/api/v1/storage/{artifact.NamespaceId}/refs/{artifact.RefName}");
 				DecorateSpanWithArtifact(span, Prefix, artifactRequest, artifact);
 
 				IStorageNamespace storage = GetStorageNamespace(new NamespaceId(artifact.NamespaceId), artifact.Token);
@@ -1416,7 +1416,7 @@ namespace JobDriver.Execution
 
 				// Write the final node
 				await storage.WriteRefAsync(artifact.RefName, outputNodeRef, new RefOptions(), cancellationToken: cancellationToken);
-				logger.LogInformation("Uploaded artifact {ArtifactId} in {Time:n1}s ({Rate:n1}mb/s, {RateMbps:n1}mbps)", artifact.Id, timer.Elapsed.TotalSeconds, totalSize / (timer.Elapsed.TotalSeconds * 1024.0 * 1024.0), (totalSize * 8.0) / (timer.Elapsed.TotalSeconds * 1024.0 * 1024.0));
+				logger.LogInformation("Uploaded artifact {ArtifactId} in {Time:n1}s ({Rate:n1}mb/s, {RateMbps:n1}mbps)", artifactId, timer.Elapsed.TotalSeconds, totalSize / (timer.Elapsed.TotalSeconds * 1024.0 * 1024.0), (totalSize * 8.0) / (timer.Elapsed.TotalSeconds * 1024.0 * 1024.0));
 			}
 
 			// Create all the named artifacts. TODO: Merge this with regular temp storage artifacts?
@@ -1485,7 +1485,7 @@ namespace JobDriver.Execution
 					await storage.WriteRefAsync(artifact.RefName, outputNodeRef, new RefOptions(), cancellationToken: cancellationToken);
 
 					timer.Stop();
-					logger.LogInformation("Uploaded artifact {ArtifactId} in {Time:n1}s ({Rate:n1}mb/s, {RateMbps:n1}mbps)", artifact.Id, timer.Elapsed.TotalSeconds, totalSize / (timer.Elapsed.TotalSeconds * 1024.0 * 1024.0), (totalSize * 8.0) / (timer.Elapsed.TotalSeconds * 1024.0 * 1024.0));
+					logger.LogInformation("Uploaded artifact {ArtifactId} in {Time:n1}s ({Rate:n1}mb/s, {RateMbps:n1}mbps)", artifactId, timer.Elapsed.TotalSeconds, totalSize / (timer.Elapsed.TotalSeconds * 1024.0 * 1024.0), (totalSize * 8.0) / (timer.Elapsed.TotalSeconds * 1024.0 * 1024.0));
 				}
 			}
 
diff --git a/Engine/Source/Programs/Shared/EpicGames.Horde.Tests/ComputeSocketTests.cs b/Engine/Source/Programs/Shared/EpicGames.Horde.Tests/ComputeSocketTests.cs
--- a/Engine/Source/Programs/Shared/EpicGames.Horde.Tests/ComputeSocketTests.cs
+++ b/Engine/Source/Programs/Shared/EpicGames.Horde.Tests/ComputeSocketTests.cs
@@ -172,7 +172,7 @@ namespace EpicGames.Horde.Tests
 					}
 
 					BundleStorageNamespace storage = BundleStorageNamespace.CreateInMemory(NullLogger.Instance);
-					await using (IBlobWriter blobWriter = storage.CreateBlobWriter(cancellationToken: : cancellationToken))
+					await using (IBlobWriter blobWriter = storage.CreateBlobWriter(cancellationToken: cancellationToken))
 					{
 						FileReference file = FileReference.Combine(tempDir, "subdir/hello.txt");
 						if (FileReference.Exists(file))
