Horde: Add an upload endpoint for log blobs which takes client-determined identifiers.

#rnx

diff --git a/Engine/Source/Programs/Horde/Plugins/Compute/HordeServer.Compute/Logs/LogsController.cs b/Engine/Source/Programs/Horde/Plugins/Compute/HordeServer.Compute/Logs/LogsController.cs
--- a/Engine/Source/Programs/Horde/Plugins/Compute/HordeServer.Compute/Logs/LogsController.cs
+++ b/Engine/Source/Programs/Horde/Plugins/Compute/HordeServer.Compute/Logs/LogsController.cs
@@ -130,6 +130,37 @@ namespace HordeServer.Logs
 			return await StorageController.WriteBlobAsync(storageBackend, null, request, cancellationToken);
 		}
 
+		/// <summary>
+		/// Uploads a blob for a log file. See /api/v1/storage/XXX/blobs.
+		/// </summary>
+		/// <param name="logId">Id of the log file to get information about</param>
+		/// <param name="locator">Locator for the new blob</param>
+		/// <param name="request">Request for the upload</param>
+		/// <param name="cancellationToken">Cancellation token for the request</param>
+		/// <returns>Information about the requested project</returns>
+		[HttpPut]
+		[Route("/api/v1/logs/{logId}/blobs/{*locator}")]
+		[ProducesResponseType(typeof(WriteBlobResponse), 200)]
+		public async Task<ActionResult<WriteBlobResponse>> WriteLogBlobAsync(LogId logId, BlobLocator locator, WriteBlobRequest request, CancellationToken cancellationToken = default)
+		{
+			ILog? log = await _logCollection.GetAsync(logId, cancellationToken);
+			if (log == null)
+			{
+				return NotFound();
+			}
+			if (!await AuthorizeAsync(log, LogAclAction.WriteLogData, User, cancellationToken))
+			{
+				return Forbid();
+			}
+			if (!locator.ToString().StartsWith($"{logId}/", StringComparison.Ordinal))
+			{
+				return BadRequest($"Log locator must start with log id");
+			}
+
+			IStorageBackend storageBackend = _storageService.CreateBackend(log.NamespaceId);
+			return await StorageController.WriteBlobAsync(storageBackend, locator, request, cancellationToken);
+		}
+
 		/// <summary>
 		/// Retrieve raw data for a log file
 		/// </summary>
diff --git a/Engine/Source/Programs/Shared/EpicGames.Horde/Logs/ServerLogger.cs b/Engine/Source/Programs/Shared/EpicGames.Horde/Logs/ServerLogger.cs
--- a/Engine/Source/Programs/Shared/EpicGames.Horde/Logs/ServerLogger.cs
+++ b/Engine/Source/Programs/Shared/EpicGames.Horde/Logs/ServerLogger.cs
@@ -55,7 +55,7 @@ namespace EpicGames.Horde.Logs
 			_internalLogger = internalLogger;
 			_builder = new LogBuilder(LogFormat.Json, internalLogger);
 			_store = _hordeClient.GetStorageNamespace(logId);
-			_writer = _store.CreateBlobWriter();
+			_writer = _store.CreateBlobWriter($"{logId}");
 
 			_tailTaskStop = new AsyncEvent();
 			_tailTask = Task.Run(() => TickTailAsync());
