Horde: Update release notes.

#rnx

diff --git a/Engine/Source/Programs/Horde/Docs/ReleaseNotes.md b/Engine/Source/Programs/Horde/Docs/ReleaseNotes.md
--- a/Engine/Source/Programs/Horde/Docs/ReleaseNotes.md
+++ b/Engine/Source/Programs/Horde/Docs/ReleaseNotes.md
@@ -1,5 +1,20 @@
 # Release Notes
 
+## 2024-11-06
+
+* Add support for parallel uploads of bundles. Bundle locators are now assigned by the client, allowing the next bundle to be assembled before the previous one has finished uploading. (37808051)
+* Include a flag in template responses indicating whether the user can run it. (37792772)
+* Show steps with issues inline when there is a label selected (37791398)
+* Fix PerforceServiceCache.SubscribeAsync() stalling out when more than 10 commits are parsed in a single iteration, due to changes being enumerated in the incorrect order. Now uses a custom query to search for commits in the correct order, and only uses cached commit metadata. (37790919)
+* Add flag for registering agent for uninstall in Windows. If the agent is installed via Unreal Toolbox but later uninstalled via Windows add/remove programs, Toolbox won't be notified. Until the uninstall command can be routed through Toolbox, this flag prevents Windows from handling uninstallation directly. (37743944)
+* Do not try to parse portable PDB files. These are not supported by symstore. (37743471)
+* [UBT] CppCompileEnvironment to use CppCompileWarning instead of raw properties (37696325)
+* Add option for using uninstall location from registry (37693256)
+* Remove EpicGames.Perforce dependency from agent (37693017)
+* Include toolbox.json in agent with process elevation (37692524)
+* Fix storage client initialization in tool upload command. (37650068)
+* Fix DI error with Horde client tool commands (37646970)
+
 ## 2024-10-31
 
 * Use lowercase for processor architecture in agent property names (37575294)
