Horde: Add enforcement of new template canRun and job response can update properties
#skipci

diff --git a/Engine/Source/Programs/Horde/HordeDashboard/src/backend/Api.ts b/Engine/Source/Programs/Horde/HordeDashboard/src/backend/Api.ts
--- a/Engine/Source/Programs/Horde/HordeDashboard/src/backend/Api.ts
+++ b/Engine/Source/Programs/Horde/HordeDashboard/src/backend/Api.ts
@@ -1320,6 +1320,9 @@ export type GetJobResponse = {
 
 	/** Custom list of targets for the job, will take precedence over all argument and parameter targets */
 	targets?: string[];
+
+	/**  Whether the current user is allowed to update this job */		
+	canUpdate?: boolean;
 }
 
 /**Request used to update a jobstep */
@@ -2236,6 +2239,9 @@ export type GetTemplateRefResponse = GetTemplateResponseBase & {
 	/** List of queries for the default changelist */
 	defaultChange?: ChangeQueryConfig[];
 
+	/** Whether the current user can run the template or not */
+	canRun?: boolean;
+
 }
 
 /** Specifies defaults for running a preflight */
diff --git a/Engine/Source/Programs/Horde/HordeDashboard/src/components/build/NewBuildV2.tsx b/Engine/Source/Programs/Horde/HordeDashboard/src/components/build/NewBuildV2.tsx
--- a/Engine/Source/Programs/Horde/HordeDashboard/src/components/build/NewBuildV2.tsx
+++ b/Engine/Source/Programs/Horde/HordeDashboard/src/components/build/NewBuildV2.tsx
@@ -977,8 +977,11 @@ class BuildOptions {
                   return;
                }
                templates.push(t);
-            });
+            });            
          }
+
+         templates = templates.filter(t => !!t.canRun);
+         
       } else {
          const query = new URLSearchParams(window.location.search);
          const allowtemplatechange = query.get("allowtemplatechange") === "true";
@@ -2032,7 +2035,7 @@ const BuildModal: React.FC<{ setUseLegacyDialog: (value: boolean) => void }> = o
                   </Stack>}
 
                   <Stack grow />
-                  <PrimaryButton text="Start Job" disabled={!template || options.readOnly || options.preview} onClick={() => { onSubmit(); }} />
+                  <PrimaryButton text="Start Job" disabled={!template || !template.canRun || options.readOnly || options.preview} onClick={() => { onSubmit(); }} />
                   <DefaultButton text="Cancel" disabled={false} onClick={() => { options.onClose(); }} />
                </Stack>
             </Stack>
diff --git a/Engine/Source/Programs/Horde/HordeDashboard/src/components/jobDetailsV2/JobOperationsBar.tsx b/Engine/Source/Programs/Horde/HordeDashboard/src/components/jobDetailsV2/JobOperationsBar.tsx
--- a/Engine/Source/Programs/Horde/HordeDashboard/src/components/jobDetailsV2/JobOperationsBar.tsx
+++ b/Engine/Source/Programs/Horde/HordeDashboard/src/components/jobDetailsV2/JobOperationsBar.tsx
@@ -55,7 +55,7 @@ export const JobOperations: React.FC<{ jobDetails: JobDetailsV2 }> = observer(({
       return null;
    }
 
-   const abortDisabled = jobData.state === JobState.Complete;
+   const abortDisabled = jobData.state === JobState.Complete || !jobDetails.template?.canRun;
 
    const seenSteps = new Set<string>();
    const runningSteps = new Set<string>();
@@ -96,8 +96,8 @@ export const JobOperations: React.FC<{ jobDetails: JobDetailsV2 }> = observer(({
    })
 
    
-
-   const retryFailedStepsDisabled = !failedSteps.length;
+   const canRun = jobDetails.template?.canRun ? true : false;
+   const retryFailedStepsDisabled = !failedSteps.length || !canRun;
 
    const pinned = dashboard.jobPinned(jobId);
 
@@ -120,6 +120,7 @@ export const JobOperations: React.FC<{ jobDetails: JobDetailsV2 }> = observer(({
    opsList.push({
       key: 'jobops_edit',
       text: "Edit",
+      disabled: !jobData.canUpdate,
       iconProps: { iconName: "Edit" },
       onClick: () => { setEditShown(true); }
    });
@@ -147,6 +148,7 @@ export const JobOperations: React.FC<{ jobDetails: JobDetailsV2 }> = observer(({
       key: 'jobops_runagain',
       text: "Run Again",
       iconProps: { iconName: "Duplicate" },
+      disabled: !canRun,
       onClick: () => { setParametersState(ParameterState.Clone); }
    });
 
@@ -379,9 +381,9 @@ const StepOperations: React.FC<{ jobDetails: JobDetailsV2, stepId: string }> = o
    }
 
 
-   const canRunDisabled = !step?.allowRetry || !!step.retriedByUserInfo || !!jobData.abortedByUserInfo;
-   const canTryFix = jobDetails.template?.allowPreflights;
-   let canBisect = (step.outcome === JobStepOutcome.Failure || step.outcome === JobStepOutcome.Warnings) && !!jobDetails.template?.allowPreflights;
+   const canRunDisabled = !step?.allowRetry || !!step.retriedByUserInfo || !!jobData.abortedByUserInfo || !jobDetails.template?.canRun;
+   const canTryFix = jobDetails.template?.allowPreflights && !!jobDetails.template?.canRun;
+   let canBisect = (step.outcome === JobStepOutcome.Failure || step.outcome === JobStepOutcome.Warnings) && !!jobDetails.template?.allowPreflights && !!jobDetails.template?.canRun;
 
    if (getSiteConfig().environment === "production") {
       canBisect = false;
@@ -402,7 +404,7 @@ const StepOperations: React.FC<{ jobDetails: JobDetailsV2, stepId: string }> = o
       key: 'stepops_abort',
       text: "Cancel Step",
       iconProps: { iconName: "Cross" },
-      disabled: !!step.finishTime,
+      disabled: !!step.finishTime || !jobData.canUpdate,
       onClick: () => { setShown({ abortShown: true }); }
    });
 
