Horde: Fix skipped steps showing that they ran for 0s
#skipci

diff --git a/Engine/Source/Programs/Horde/HordeDashboard/src/components/jobDetailsV2/StepDetailSummary.tsx b/Engine/Source/Programs/Horde/HordeDashboard/src/components/jobDetailsV2/StepDetailSummary.tsx
--- a/Engine/Source/Programs/Horde/HordeDashboard/src/components/jobDetailsV2/StepDetailSummary.tsx
+++ b/Engine/Source/Programs/Horde/HordeDashboard/src/components/jobDetailsV2/StepDetailSummary.tsx
@@ -147,6 +147,10 @@ const getStepSummaryMarkdown = (jobDetails: JobDetailsV2, stepId: string): strin
          }
       }
 
+      if (step.state === JobStepState.Skipped) {
+         runningText = "";
+      }
+
       if (runningText) {
          text.push(runningText);
       }
