Horde: Fix cached commit source decrementing number of remaining results to return even if commits are filtered out of the resulting list by commit tags.

Prevents issues with new jobs not being able to identify code changelist to run at.

diff --git a/Engine/Source/Programs/Horde/Plugins/Build/HordeServer.Build/Perforce/PerforceServiceCache.cs b/Engine/Source/Programs/Horde/Plugins/Build/HordeServer.Build/Perforce/PerforceServiceCache.cs
--- a/Engine/Source/Programs/Horde/Plugins/Build/HordeServer.Build/Perforce/PerforceServiceCache.cs
+++ b/Engine/Source/Programs/Horde/Plugins/Build/HordeServer.Build/Perforce/PerforceServiceCache.cs
@@ -697,7 +697,13 @@ namespace HordeServer.Perforce
 						// Expand the range of cached changes if necessary
 						if ((maxResults == null || maxResults.Value > 0) && maxChange > 0)
 						{
-							await foreach (ICommit commit in base.FindAsync(minChange, true, maxChange, true, maxResults, null, cancellationToken))
+							int? maxUnfilteredResults = null;
+							if (tags == null || tags.Count == 0)
+							{
+								maxUnfilteredResults = maxResults;
+							}
+
+							await foreach (ICommit commit in base.FindAsync(minChange, true, maxChange, true, maxUnfilteredResults, null, cancellationToken))
 							{
 								CachedCommitDoc cachedCommit = await CachedCommitDoc.FromCommitAsync(commit, cancellationToken);
 								await _owner.AddCachedCommitAsync(cachedCommit, cancellationToken);
@@ -708,15 +714,6 @@ namespace HordeServer.Perforce
 								{
 									yield return cachedCommit;
 								}
-
-								maxResults = maxResults.HasValue ? maxResults.Value - 1 : null;
-							}
-
-							if (maxResults == null || maxResults.Value > 0)
-							{
-								int newMinChange = minChange ?? 0;
-								_owner._logger.LogDebug("Extending range for {StreamId} cache to {Change}..", StreamConfig.Id, newMinChange);
-								await _owner._mongoService.UpdateSingletonAsync<CacheState>(x => TryUpdateRange(x, StreamConfig, newMinChange, maxChange), cancellationToken);
 							}
 						}
 					}
