Horde: Update release notes.

#rnx

diff --git a/Engine/Source/Programs/Horde/Docs/ReleaseNotes.md b/Engine/Source/Programs/Horde/Docs/ReleaseNotes.md
--- a/Engine/Source/Programs/Horde/Docs/ReleaseNotes.md
+++ b/Engine/Source/Programs/Horde/Docs/ReleaseNotes.md
@@ -2,6 +2,7 @@
 
 ## 2024-11-06
 
+* Fix cached commit source decrementing number of remaining results to return even if commits are filtered out of the resulting list by commit tags. Prevents issues with new jobs not being able to identify code changelist to run at. (37833532)
 * Add support for parallel uploads of bundles. Bundle locators are now assigned by the client, allowing the next bundle to be assembled before the previous one has finished uploading. (37808051)
 * Include a flag in template responses indicating whether the user can run it. (37792772)
 * Show steps with issues inline when there is a label selected (37791398)
