Horde: Decrement the number of remaining results to return when falling back to slow path for querying commits.

#rnx

diff --git a/Engine/Source/Programs/Horde/Plugins/Build/HordeServer.Build/Perforce/PerforceServiceCache.cs b/Engine/Source/Programs/Horde/Plugins/Build/HordeServer.Build/Perforce/PerforceServiceCache.cs
--- a/Engine/Source/Programs/Horde/Plugins/Build/HordeServer.Build/Perforce/PerforceServiceCache.cs
+++ b/Engine/Source/Programs/Horde/Plugins/Build/HordeServer.Build/Perforce/PerforceServiceCache.cs
@@ -713,6 +713,15 @@ namespace HordeServer.Perforce
 								if (tags == null || tags.Any(x => cachedCommit.CommitTags.Contains(x)))
 								{
 									yield return cachedCommit;
+
+									if (maxResults != null)
+									{
+										maxResults = maxResults.Value - 1;
+										if (maxResults.Value == 0)
+										{
+											break;
+										}
+									}
 								}
 							}
 						}
