Horde: Strip unused files during publish for win-x64

diff --git a/Engine/Source/Programs/Horde/Drivers/JobDriver/JobDriver.csproj b/Engine/Source/Programs/Horde/Drivers/JobDriver/JobDriver.csproj
--- a/Engine/Source/Programs/Horde/Drivers/JobDriver/JobDriver.csproj
+++ b/Engine/Source/Programs/Horde/Drivers/JobDriver/JobDriver.csproj
@@ -111,6 +111,21 @@
     <ProjectReference Include="..\..\..\Shared\EpicGames.Perforce\EpicGames.Perforce.csproj" />
   </ItemGroup>
 
+  <Target Name="DeleteUnusedFiles-Win64" AfterTargets="Publish" Condition="'$(Configuration)' == 'Release' AND '$(RuntimeIdentifier)' == 'win-x64'">
+    <ItemGroup>
+      <FilesToDelete Include="$(PublishDir)**/*.dylib" />
+      <FilesToDelete Include="$(PublishDir)**/*.so.9" />
+      <FilesToDelete Include="$(PublishDir)**/*.so" />
+      <FilesToDelete Include="$(PublishDir)**/*.xml" />
+      <FilesToDelete Include="$(PublishDir)**/oo2core*win32.dll" />
+      <FilesToDelete Include="$(PublishDir)**/oo2core*winuwp*.dll" />
+    </ItemGroup>
+    <Delete Files="@(FilesToDelete)">
+      <Output TaskParameter="DeletedFiles" ItemName="DeletedList"/>
+    </Delete>
+    <Message Text="Deleted files:%0A@(DeletedList->'  %(Identity)','%0A')" Importance="high" />
+  </Target>
+
   <ItemGroup>
     <Content Include="appsettings.json" Condition="Exists('appsettings.json')">
       <Generator></Generator>
diff --git a/Engine/Source/Programs/Horde/HordeAgent/HordeAgent.csproj b/Engine/Source/Programs/Horde/HordeAgent/HordeAgent.csproj
--- a/Engine/Source/Programs/Horde/HordeAgent/HordeAgent.csproj
+++ b/Engine/Source/Programs/Horde/HordeAgent/HordeAgent.csproj
@@ -204,6 +204,8 @@
       <FilesToDelete Include="$(PublishDir)**/*.so.9" />
       <FilesToDelete Include="$(PublishDir)**/*.so" />
       <FilesToDelete Include="$(PublishDir)**/*.xml" />
+      <FilesToDelete Include="$(PublishDir)**/oo2core*win32.dll" />
+      <FilesToDelete Include="$(PublishDir)**/oo2core*winuwp*.dll" />
     </ItemGroup>
     <Delete Files="@(FilesToDelete)">
       <Output TaskParameter="DeletedFiles" ItemName="DeletedList"/>
