Horde: Show critical log levels as errors in the log view
#skipci

diff --git a/Engine/Source/Programs/Horde/HordeDashboard/src/backend/Api.ts b/Engine/Source/Programs/Horde/HordeDashboard/src/backend/Api.ts
--- a/Engine/Source/Programs/Horde/HordeDashboard/src/backend/Api.ts
+++ b/Engine/Source/Programs/Horde/HordeDashboard/src/backend/Api.ts
@@ -258,7 +258,8 @@ export enum LogType {
 export enum LogLevel {
 	Informational = "Informational",
 	Warning = "Warning",
-	Error = "Error"
+	Error = "Error",
+	Critical = "Critical"
 }
 
 export type LogLine = {
diff --git a/Engine/Source/Programs/Horde/HordeDashboard/src/components/LogView.tsx b/Engine/Source/Programs/Horde/HordeDashboard/src/components/LogView.tsx
--- a/Engine/Source/Programs/Horde/HordeDashboard/src/components/LogView.tsx
+++ b/Engine/Source/Programs/Horde/HordeDashboard/src/components/LogView.tsx
@@ -733,7 +733,7 @@ export const LogList: React.FC<{ logId: string }> = observer(({ logId }) => {
 
          const line = item.line;
          const warning = line?.level === LogLevel.Warning;
-         const error = line?.level === LogLevel.Error;
+         const error = (line?.level === LogLevel.Error) || (line?.level === LogLevel.Critical);
 
          const ev = logSource.events.find(event => event.lineIndex === item.lineNumber - 1);
 
