Horde: Rename IStorageWriter.WriteRef/DeleteRef to AddRef/RemoveRef.

#rnx

diff --git a/Engine/Source/Programs/Horde/Drivers/JobDriver/Execution/JobExecutor.cs b/Engine/Source/Programs/Horde/Drivers/JobDriver/Execution/JobExecutor.cs
--- a/Engine/Source/Programs/Horde/Drivers/JobDriver/Execution/JobExecutor.cs
+++ b/Engine/Source/Programs/Horde/Drivers/JobDriver/Execution/JobExecutor.cs
@@ -855,7 +855,7 @@ namespace JobDriver.Execution
 
 					rootNodeRef = await blobWriter.WriteBlobAsync(rootNode, cancellationToken);
 				}
-				await storage.WriteRefAsync(artifact.RefName, rootNodeRef, new RefOptions(), cancellationToken);
+				await storage.AddRefAsync(artifact.RefName, rootNodeRef, new RefOptions(), cancellationToken);
 
 				logger.LogInformation("Upload took {Time:n1}s", timer.Elapsed.TotalSeconds);
 			}
@@ -1172,7 +1172,7 @@ namespace JobDriver.Execution
 						throw;
 					}
 				}
-				await storage.WriteRefAsync(new RefName(artifact.RefName), rootRef, cancellationToken: cancellationToken);
+				await storage.AddRefAsync(new RefName(artifact.RefName), rootRef, cancellationToken: cancellationToken);
 
 				timer.Stop();
 				logger.LogInformation("Uploaded artifact {ArtifactId} in {Time:n1}s ({Rate:n1}mb/s, {RateMbps:n1}mbps)", artifactId, timer.Elapsed.TotalSeconds, totalSize / (timer.Elapsed.TotalSeconds * 1024.0 * 1024.0), (totalSize * 8.0) / (timer.Elapsed.TotalSeconds * 1024.0 * 1024.0));
@@ -1415,7 +1415,7 @@ namespace JobDriver.Execution
 				}
 
 				// Write the final node
-				await storage.WriteRefAsync(artifact.RefName, outputNodeRef, new RefOptions(), cancellationToken: cancellationToken);
+				await storage.AddRefAsync(artifact.RefName, outputNodeRef, new RefOptions(), cancellationToken: cancellationToken);
 				logger.LogInformation("Uploaded artifact {ArtifactId} in {Time:n1}s ({Rate:n1}mb/s, {RateMbps:n1}mbps)", artifactId, timer.Elapsed.TotalSeconds, totalSize / (timer.Elapsed.TotalSeconds * 1024.0 * 1024.0), (totalSize * 8.0) / (timer.Elapsed.TotalSeconds * 1024.0 * 1024.0));
 			}
 
@@ -1482,7 +1482,7 @@ namespace JobDriver.Execution
 					}
 
 					// Write the final node
-					await storage.WriteRefAsync(artifact.RefName, outputNodeRef, new RefOptions(), cancellationToken: cancellationToken);
+					await storage.AddRefAsync(artifact.RefName, outputNodeRef, new RefOptions(), cancellationToken: cancellationToken);
 
 					timer.Stop();
 					logger.LogInformation("Uploaded artifact {ArtifactId} in {Time:n1}s ({Rate:n1}mb/s, {RateMbps:n1}mbps)", artifactId, timer.Elapsed.TotalSeconds, totalSize / (timer.Elapsed.TotalSeconds * 1024.0 * 1024.0), (totalSize * 8.0) / (timer.Elapsed.TotalSeconds * 1024.0 * 1024.0));
diff --git a/Engine/Source/Programs/Horde/HordeCmd/Commands/Archives/ArchiveCreate.cs b/Engine/Source/Programs/Horde/HordeCmd/Commands/Archives/ArchiveCreate.cs
--- a/Engine/Source/Programs/Horde/HordeCmd/Commands/Archives/ArchiveCreate.cs
+++ b/Engine/Source/Programs/Horde/HordeCmd/Commands/Archives/ArchiveCreate.cs
@@ -109,7 +109,7 @@ namespace Horde.Commands.Archives
 				else
 				{
 					logger.LogInformation("Writing ref {Ref}", Ref);
-					await store.WriteRefAsync(new RefName(Ref!), nodeRef);
+					await store.AddRefAsync(new RefName(Ref!), nodeRef);
 				}
 			}
 
diff --git a/Engine/Source/Programs/Horde/Plugins/Build/HordeServer.Build.Tests/Artifacts/UnsyncCacheTests.cs b/Engine/Source/Programs/Horde/Plugins/Build/HordeServer.Build.Tests/Artifacts/UnsyncCacheTests.cs
--- a/Engine/Source/Programs/Horde/Plugins/Build/HordeServer.Build.Tests/Artifacts/UnsyncCacheTests.cs
+++ b/Engine/Source/Programs/Horde/Plugins/Build/HordeServer.Build.Tests/Artifacts/UnsyncCacheTests.cs
@@ -40,7 +40,7 @@ namespace HordeServer.Tests.Artifacts
 			}
 
 			RefName refName = new RefName("test");
-			await storageNamespace.WriteRefAsync(refName, directoryRef);
+			await storageNamespace.AddRefAsync(refName, directoryRef);
 
 			Mock<IStorageService> factory = new Mock<IStorageService>();
 			factory.Setup(x => x.TryGetNamespace(ns)).Returns(storageNamespace);
diff --git a/Engine/Source/Programs/Horde/Plugins/Build/HordeServer.Build.Tests/Issues/IssueServiceTests.cs b/Engine/Source/Programs/Horde/Plugins/Build/HordeServer.Build.Tests/Issues/IssueServiceTests.cs
--- a/Engine/Source/Programs/Horde/Plugins/Build/HordeServer.Build.Tests/Issues/IssueServiceTests.cs
+++ b/Engine/Source/Programs/Horde/Plugins/Build/HordeServer.Build.Tests/Issues/IssueServiceTests.cs
@@ -68,7 +68,7 @@ namespace HordeServer.Tests.Issues
 				await using (IBlobWriter writer = _storageNamespace.CreateBlobWriter())
 				{
 					IHashedBlobRef<LogNode> handle = await _builder.FlushAsync(writer, true, CancellationToken.None);
-					await _storageNamespace.WriteRefAsync(_log.RefName, handle);
+					await _storageNamespace.AddRefAsync(_log.RefName, handle);
 				}
 			}
 
@@ -354,7 +354,7 @@ namespace HordeServer.Tests.Issues
 				builder.WriteData(data);
 
 				IHashedBlobRef<LogNode> handle = await builder.FlushAsync(writer, true, CancellationToken.None);
-				await storageNamespace.WriteRefAsync(log!.RefName, handle);
+				await storageNamespace.AddRefAsync(log!.RefName, handle);
 			}
 
 			await log.AddEventsAsync(new List<NewLogEventData> { new NewLogEventData { LineIndex = 0, LineCount = 1, Severity = severity } }, CancellationToken.None);
diff --git a/Engine/Source/Programs/Horde/Plugins/Build/HordeServer.Build/Artifacts/ArtifactCollection.cs b/Engine/Source/Programs/Horde/Plugins/Build/HordeServer.Build/Artifacts/ArtifactCollection.cs
--- a/Engine/Source/Programs/Horde/Plugins/Build/HordeServer.Build/Artifacts/ArtifactCollection.cs
+++ b/Engine/Source/Programs/Horde/Plugins/Build/HordeServer.Build/Artifacts/ArtifactCollection.cs
@@ -182,7 +182,7 @@ namespace HordeServer.Artifacts
 					{
 						await writer.AddAliasAsync(aliasInfo.Name, aliasInfo.BlobRef, aliasInfo.Rank, aliasInfo.Data, cancellationToken);
 					}
-					await writer.WriteRefAsync(_artifact.RefName, blobRef, cancellationToken: cancellationToken);
+					await writer.AddRefAsync(_artifact.RefName, blobRef, cancellationToken: cancellationToken);
 				}
 				return _artifact;
 			}
@@ -427,7 +427,7 @@ namespace HordeServer.Artifacts
 			_logger.LogInformation("Deleting {StreamId} artifact {ArtifactId}, ref {RefName} (created {CreateTime})", artifact.StreamId, artifact.Id, artifact.RefName, artifact.CreatedAtUtc);
 
 			IStorageNamespace storageNamespace = _storageService.GetNamespace(artifact.NamespaceId);
-			await storageNamespace.DeleteRefAsync(artifact.RefName, cancellationToken);
+			await storageNamespace.RemoveRefAsync(artifact.RefName, cancellationToken);
 
 			FilterDefinition<ArtifactDocument> filter = Builders<ArtifactDocument>.Filter.Eq(x => x.Id, artifact.Id);
 			await _artifactCollection.DeleteOneAsync(filter, cancellationToken);
@@ -442,7 +442,7 @@ namespace HordeServer.Artifacts
 				{
 					// Delete the ref allowing the storage service to expire this data
 					_logger.LogInformation("Expiring {StreamId} artifact {ArtifactId}, ref {RefName} (created {CreateTime})", artifact.StreamId, artifact.Id, artifact.RefName, artifact.CreatedAtUtc);
-					await storageNamespace.DeleteRefAsync(artifact.RefName, cancellationToken);
+					await storageNamespace.RemoveRefAsync(artifact.RefName, cancellationToken);
 				}
 
 				FilterDefinition<ArtifactDocument> filter = Builders<ArtifactDocument>.Filter.In(x => x.Id, artifactGroup.Select(x => x.Id));
diff --git a/Engine/Source/Programs/Horde/Plugins/Build/HordeServer.Build/Replicators/PerforceReplicator.cs b/Engine/Source/Programs/Horde/Plugins/Build/HordeServer.Build/Replicators/PerforceReplicator.cs
--- a/Engine/Source/Programs/Horde/Plugins/Build/HordeServer.Build/Replicators/PerforceReplicator.cs
+++ b/Engine/Source/Programs/Horde/Plugins/Build/HordeServer.Build/Replicators/PerforceReplicator.cs
@@ -324,8 +324,8 @@ namespace HordeServer.Replicators
 			{
 				_logger.LogInformation("Resetting replication for {ReplicatorId}", replicator.Id);
 
-				await store.DeleteRefAsync(refName, cancellationToken);
-				await store.DeleteRefAsync(incRefName, cancellationToken);
+				await store.RemoveRefAsync(refName, cancellationToken);
+				await store.RemoveRefAsync(incRefName, cancellationToken);
 
 				UpdateReplicatorOptions updateOptions = new UpdateReplicatorOptions { Reset = false, LastChange = 0, CurrentChange = 0 };
 				replicator = await UpdateReplicatorAsync(replicator, updateOptions, cancellationToken);
@@ -335,7 +335,7 @@ namespace HordeServer.Replicators
 			{
 				_logger.LogInformation("Forcing current change for {ReplicatorId} to {Change}", replicator.Id, replicator.NextChange.Value);
 
-				await store.DeleteRefAsync(incRefName, cancellationToken);
+				await store.RemoveRefAsync(incRefName, cancellationToken);
 
 				UpdateReplicatorOptions updateOptions = new UpdateReplicatorOptions { NextChange = 0, CurrentChange = replicator.NextChange.Value };
 				replicator = await UpdateReplicatorAsync(replicator, updateOptions, cancellationToken);
@@ -355,7 +355,7 @@ namespace HordeServer.Replicators
 					_logger.LogInformation("Starting {ReplicatorId} replication from latest CL {Change}", replicator.Id, nextCommit.Id);
 				}
 
-				await store.DeleteRefAsync(incRefName, cancellationToken);
+				await store.RemoveRefAsync(incRefName, cancellationToken);
 
 				UpdateReplicatorOptions updateOptions = new UpdateReplicatorOptions { CurrentChange = nextCommit.Id.GetPerforceChange() };
 				replicator = await UpdateReplicatorAsync(replicator, updateOptions, cancellationToken);
@@ -713,7 +713,7 @@ namespace HordeServer.Replicators
 				IHashedBlobRef<StateNode> stateNodeRef = await directoryWriter.WriteBlobAsync(stateNode, cancellationToken);
 				await directoryWriter.FlushAsync(cancellationToken);
 
-				await store.WriteRefAsync(incRefName, stateNodeRef, cancellationToken: cancellationToken);
+				await store.AddRefAsync(incRefName, stateNodeRef, cancellationToken: cancellationToken);
 				rootUpdate.Clear();
 
 				if (replicator.Clean || !String.IsNullOrEmpty(replicator.CurrentError))
@@ -739,7 +739,7 @@ namespace HordeServer.Replicators
 			IHashedBlobRef<RedirectNode<CommitNode>> redirectNodeRef = await commitWriter.WriteBlobAsync(redirectNode, cancellationToken);
 
 			await commitWriter.FlushAsync(cancellationToken);
-			await store.WriteRefAsync(refName, redirectNodeRef, cancellationToken: cancellationToken);
+			await store.AddRefAsync(refName, redirectNodeRef, cancellationToken: cancellationToken);
 
 			// Update the replicator state
 			UpdateReplicatorOptions completeUpdateOptions = new UpdateReplicatorOptions
diff --git a/Engine/Source/Programs/Horde/Plugins/Compute/HordeServer.Compute.Tests/Logs/LogIndexingTests.cs b/Engine/Source/Programs/Horde/Plugins/Compute/HordeServer.Compute.Tests/Logs/LogIndexingTests.cs
--- a/Engine/Source/Programs/Horde/Plugins/Compute/HordeServer.Compute.Tests/Logs/LogIndexingTests.cs
+++ b/Engine/Source/Programs/Horde/Plugins/Compute/HordeServer.Compute.Tests/Logs/LogIndexingTests.cs
@@ -43,7 +43,7 @@ namespace HordeServer.Tests.Logs
 		public async Task<ILog> FlushAsync(bool complete = true)
 		{
 			IHashedBlobRef<LogNode> handle = await _builder.FlushAsync(_blobWriter, complete, CancellationToken.None);
-			await _storageNamespace.WriteRefAsync(_logFile.RefName, handle);
+			await _storageNamespace.AddRefAsync(_logFile.RefName, handle);
 			_logFile = await _logFile.UpdateLineCountAsync(_lineCount, complete, CancellationToken.None);
 			return _logFile;
 		}
diff --git a/Engine/Source/Programs/Horde/Plugins/Compute/HordeServer.Compute/Logs/LogCollection.cs b/Engine/Source/Programs/Horde/Plugins/Compute/HordeServer.Compute/Logs/LogCollection.cs
--- a/Engine/Source/Programs/Horde/Plugins/Compute/HordeServer.Compute/Logs/LogCollection.cs
+++ b/Engine/Source/Programs/Horde/Plugins/Compute/HordeServer.Compute/Logs/LogCollection.cs
@@ -307,7 +307,7 @@ namespace HordeServer.Logs
 			if (logDocument.UseNewStorageBackend)
 			{
 				IStorageNamespace storageNamespace = _storageService.GetNamespace(logDocument.NamespaceId);
-				await storageNamespace.DeleteRefAsync(logDocument.RefName, cancellationToken);
+				await storageNamespace.RemoveRefAsync(logDocument.RefName, cancellationToken);
 			}
 		}
 
diff --git a/Engine/Source/Programs/Horde/Plugins/Compute/HordeServer.Compute/Logs/LogRpcService.cs b/Engine/Source/Programs/Horde/Plugins/Compute/HordeServer.Compute/Logs/LogRpcService.cs
--- a/Engine/Source/Programs/Horde/Plugins/Compute/HordeServer.Compute/Logs/LogRpcService.cs
+++ b/Engine/Source/Programs/Horde/Plugins/Compute/HordeServer.Compute/Logs/LogRpcService.cs
@@ -62,7 +62,7 @@ namespace HordeServer.Logs
 			}
 
 			IHashedBlobRef target = store.CreateBlobRef(hash, new BlobLocator(request.TargetLocator));
-			await store.WriteRefAsync(new RefName(request.LogId), target);
+			await store.AddRefAsync(new RefName(request.LogId), target);
 
 			await log.UpdateLineCountAsync(request.LineCount, request.Complete, context.CancellationToken);
 
diff --git a/Engine/Source/Programs/Horde/Plugins/Ddc/HordeServer.Ddc/Ddc/RefService.cs b/Engine/Source/Programs/Horde/Plugins/Ddc/HordeServer.Ddc/Ddc/RefService.cs
--- a/Engine/Source/Programs/Horde/Plugins/Ddc/HordeServer.Ddc/Ddc/RefService.cs
+++ b/Engine/Source/Programs/Horde/Plugins/Ddc/HordeServer.Ddc/Ddc/RefService.cs
@@ -31,7 +31,7 @@ namespace HordeServer.Ddc
 				return false;
 			}
 
-			await storageNamespace.DeleteRefAsync(refName, cancellationToken);
+			await storageNamespace.RemoveRefAsync(refName, cancellationToken);
 			return true;
 		}
 
@@ -94,7 +94,7 @@ namespace HordeServer.Ddc
 					refNodeRef = await writer.WriteBlobAsync(refNode, cancellationToken);
 				}
 
-				await storageNamespace.WriteRefAsync(refName, refNodeRef, cancellationToken: cancellationToken);
+				await storageNamespace.AddRefAsync(refName, refNodeRef, cancellationToken: cancellationToken);
 			}
 
 			return (missingReferences, missingBlobs);
diff --git a/Engine/Source/Programs/Horde/Plugins/Storage/HordeServer.Storage.Tests/Storage/BlobStoreTests.cs b/Engine/Source/Programs/Horde/Plugins/Storage/HordeServer.Storage.Tests/Storage/BlobStoreTests.cs
--- a/Engine/Source/Programs/Horde/Plugins/Storage/HordeServer.Storage.Tests/Storage/BlobStoreTests.cs
+++ b/Engine/Source/Programs/Horde/Plugins/Storage/HordeServer.Storage.Tests/Storage/BlobStoreTests.cs
@@ -132,7 +132,7 @@ namespace HordeServer.Tests.Storage
 			for (int idx = 0; idx < 2; idx++)
 			{
 				RefName refName = new RefName("hello");
-				await store.WriteRefAsync(refName, store.CreateBlobRef(locator3));
+				await store.AddRefAsync(refName, store.CreateBlobRef(locator3));
 				IHashedBlobRef refTarget = await store.ReadRefAsync(refName);
 				Assert.AreEqual(locator3.Locator, refTarget.GetLocator());
 			}
@@ -146,9 +146,9 @@ namespace HordeServer.Tests.Storage
 			Blob blob1 = new Blob(new byte[] { 1, 2, 3 }, Array.Empty<HashedBlobRefValue>());
 			HashedBlobRefValue target = await WriteBlobAsync(store, blob1);
 
-			await store.WriteRefAsync("test-ref-1", store.CreateBlobRef(target));
-			await store.WriteRefAsync("test-ref-2", store.CreateBlobRef(target), new RefOptions { Lifetime = TimeSpan.FromMinutes(30.0), Extend = true });
-			await store.WriteRefAsync("test-ref-3", store.CreateBlobRef(target), new RefOptions { Lifetime = TimeSpan.FromMinutes(30.0), Extend = false });
+			await store.AddRefAsync("test-ref-1", store.CreateBlobRef(target));
+			await store.AddRefAsync("test-ref-2", store.CreateBlobRef(target), new RefOptions { Lifetime = TimeSpan.FromMinutes(30.0), Extend = true });
+			await store.AddRefAsync("test-ref-3", store.CreateBlobRef(target), new RefOptions { Lifetime = TimeSpan.FromMinutes(30.0), Extend = false });
 
 			Assert.AreEqual(target, await TryReadRefValueAsync(store, "test-ref-1"));
 			Assert.AreEqual(target, await TryReadRefValueAsync(store, "test-ref-2"));
diff --git a/Engine/Source/Programs/Horde/Plugins/Storage/HordeServer.Storage.Tests/Storage/GcServiceTests.cs b/Engine/Source/Programs/Horde/Plugins/Storage/HordeServer.Storage.Tests/Storage/GcServiceTests.cs
--- a/Engine/Source/Programs/Horde/Plugins/Storage/HordeServer.Storage.Tests/Storage/GcServiceTests.cs
+++ b/Engine/Source/Programs/Horde/Plugins/Storage/HordeServer.Storage.Tests/Storage/GcServiceTests.cs
@@ -44,7 +44,7 @@ namespace HordeServer.Tests.Storage
 				if (roots.Add(blobs[blobIdx]))
 				{
 					IHashedBlobRef handle = blobs[blobIdx];
-					await store.WriteRefAsync(new RefName($"ref-{idx}"), handle);
+					await store.AddRefAsync(new RefName($"ref-{idx}"), handle);
 				}
 			}
 
diff --git a/Engine/Source/Programs/Horde/Plugins/Tools/HordeServer.Tools/Tools/ToolCollection.cs b/Engine/Source/Programs/Horde/Plugins/Tools/HordeServer.Tools/Tools/ToolCollection.cs
--- a/Engine/Source/Programs/Horde/Plugins/Tools/HordeServer.Tools/Tools/ToolCollection.cs
+++ b/Engine/Source/Programs/Horde/Plugins/Tools/HordeServer.Tools/Tools/ToolCollection.cs
@@ -364,7 +364,7 @@ namespace HordeServer.Tools
 
 			// Write a ref for the deployment so the blobs aren't GC'd
 			RefName refName = new RefName($"{tool.Id}/{deploymentId}");
-			await storageNamespace.WriteRefAsync(refName, content, cancellationToken: cancellationToken);
+			await storageNamespace.AddRefAsync(refName, content, cancellationToken: cancellationToken);
 
 			// Create the new deployment object
 			ToolDeploymentDocument deployment = new ToolDeploymentDocument(deploymentId, options, toolConfig.NamespaceId, refName);
@@ -410,7 +410,7 @@ namespace HordeServer.Tools
 
 				ToolDeploymentDocument removeDeployment = tool.Deployments[0];
 				IStorageNamespace client = _storageService.GetNamespace(removeDeployment.NamespaceId);
-				await client.DeleteRefAsync(removeDeployment.RefName, cancellationToken);
+				await client.RemoveRefAsync(removeDeployment.RefName, cancellationToken);
 			}
 
 			// Add the new deployment
diff --git a/Engine/Source/Programs/Shared/EpicGames.Horde.Tests/BundleStoreTests.cs b/Engine/Source/Programs/Shared/EpicGames.Horde.Tests/BundleStoreTests.cs
--- a/Engine/Source/Programs/Shared/EpicGames.Horde.Tests/BundleStoreTests.cs
+++ b/Engine/Source/Programs/Shared/EpicGames.Horde.Tests/BundleStoreTests.cs
@@ -257,7 +257,7 @@ namespace EpicGames.Horde.Tests
 
 					await writer.FlushAsync();
 
-					await store.WriteRefAsync(new RefName("test"), rootRef);
+					await store.AddRefAsync(new RefName("test"), rootRef);
 				}
 			}
 
@@ -302,7 +302,7 @@ namespace EpicGames.Horde.Tests
 				await root.UpdateAsync(fileUpdates, writer);
 
 				IHashedBlobRef<DirectoryNode> rootRef = await writer.WriteBlobAsync(root);
-				await store.WriteRefAsync(new RefName("test"), rootRef);
+				await store.AddRefAsync(new RefName("test"), rootRef);
 
 				await CheckFileTreeAsync(root);
 			}
@@ -403,7 +403,7 @@ namespace EpicGames.Horde.Tests
 				root.AddFile("test", FileEntryFlags.None, fileWriter.Length, chunkedData);
 
 				IHashedBlobRef<DirectoryNode> rootRef = await writer.WriteBlobAsync(root);
-				await store.WriteRefAsync(new RefName("test"), rootRef);
+				await store.AddRefAsync(new RefName("test"), rootRef);
 
 				await CheckLargeFileTreeAsync(root, data);
 			}
diff --git a/Engine/Source/Programs/Shared/EpicGames.Horde.Tests/BundleTests.cs b/Engine/Source/Programs/Shared/EpicGames.Horde.Tests/BundleTests.cs
--- a/Engine/Source/Programs/Shared/EpicGames.Horde.Tests/BundleTests.cs
+++ b/Engine/Source/Programs/Shared/EpicGames.Horde.Tests/BundleTests.cs
@@ -277,7 +277,7 @@ namespace EpicGames.Horde.Tests
 						next.AddDirectory(new DirectoryEntry($"node{idx}", 0, rootRef));
 						rootRef = await writer.WriteBlobAsync(next);
 					}
-					await storage.WriteRefAsync(refName, rootRef);
+					await storage.AddRefAsync(refName, rootRef);
 				}
 
 				MemoryStorageBackend memoryStore = (MemoryStorageBackend)storage.Backend;
diff --git a/Engine/Source/Programs/Shared/EpicGames.Horde.Tests/ChunkingTests.cs b/Engine/Source/Programs/Shared/EpicGames.Horde.Tests/ChunkingTests.cs
--- a/Engine/Source/Programs/Shared/EpicGames.Horde.Tests/ChunkingTests.cs
+++ b/Engine/Source/Programs/Shared/EpicGames.Horde.Tests/ChunkingTests.cs
@@ -69,7 +69,7 @@ namespace EpicGames.Horde.Tests
 				directory.AddFile("test.foo", FileEntryFlags.None, 0, chunkedData);
 
 				IHashedBlobRef handle = await writer.WriteBlobAsync(directory);
-				await store.WriteRefAsync(RefName, handle);
+				await store.AddRefAsync(RefName, handle);
 			}
 		}
 
diff --git a/Engine/Source/Programs/Shared/EpicGames.Horde.Tests/StorageClientTests.cs b/Engine/Source/Programs/Shared/EpicGames.Horde.Tests/StorageClientTests.cs
--- a/Engine/Source/Programs/Shared/EpicGames.Horde.Tests/StorageClientTests.cs
+++ b/Engine/Source/Programs/Shared/EpicGames.Horde.Tests/StorageClientTests.cs
@@ -80,7 +80,7 @@ namespace EpicGames.Horde.Tests
 			{
 				nodeRef = await writer.WriteBlobAsync(new TestNode(123));
 			}
-			await store.WriteRefAsync("hello", nodeRef);
+			await store.AddRefAsync("hello", nodeRef);
 
 			TestNode output = await store.ReadRefTargetAsync<TestNode>("hello");
 			Assert.AreEqual(123, output.Value);
@@ -99,7 +99,7 @@ namespace EpicGames.Horde.Tests
 				IHashedBlobRef<TestNode> nodeRef1 = await writer.WriteBlobAsync(new TestNode(123));
 				nodeRef2 = await writer.WriteBlobAsync(new TestNode(456, nodeRef1));
 			}
-			await store.WriteRefAsync("hello", nodeRef2);
+			await store.AddRefAsync("hello", nodeRef2);
 
 			TestNode output2 = await store.ReadRefTargetAsync<TestNode>("hello");
 			Assert.AreEqual(456, output2.Value);
diff --git a/Engine/Source/Programs/Shared/EpicGames.Horde/Artifacts/HttpArtifactCollection.cs b/Engine/Source/Programs/Shared/EpicGames.Horde/Artifacts/HttpArtifactCollection.cs
--- a/Engine/Source/Programs/Shared/EpicGames.Horde/Artifacts/HttpArtifactCollection.cs
+++ b/Engine/Source/Programs/Shared/EpicGames.Horde/Artifacts/HttpArtifactCollection.cs
@@ -104,7 +104,7 @@ namespace EpicGames.Horde.Artifacts
 					{
 						await writer.AddAliasAsync(aliasInfo.Name, aliasInfo.BlobRef, aliasInfo.Rank, aliasInfo.Data, cancellationToken);
 					}
-					await writer.WriteRefAsync(_artifact.RefName, blobRef, cancellationToken: cancellationToken);
+					await writer.AddRefAsync(_artifact.RefName, blobRef, cancellationToken: cancellationToken);
 				}
 				return _artifact;
 			}
diff --git a/Engine/Source/Programs/Shared/EpicGames.Horde/Storage/IStorageNamespace.cs b/Engine/Source/Programs/Shared/EpicGames.Horde/Storage/IStorageNamespace.cs
--- a/Engine/Source/Programs/Shared/EpicGames.Horde/Storage/IStorageNamespace.cs
+++ b/Engine/Source/Programs/Shared/EpicGames.Horde/Storage/IStorageNamespace.cs
@@ -362,10 +362,10 @@ namespace EpicGames.Horde.Storage
 		/// <param name="options">Options for the new ref</param>
 		/// <param name="cancellationToken">Cancellation token for the operation</param>
 		/// <returns>Unique identifier for the blob</returns>
-		public static async Task WriteRefAsync(this IStorageNamespace store, RefName name, IHashedBlobRef target, RefOptions? options = null, CancellationToken cancellationToken = default)
+		public static async Task AddRefAsync(this IStorageNamespace store, RefName name, IHashedBlobRef target, RefOptions? options = null, CancellationToken cancellationToken = default)
 		{
 			await using IStorageWriter writer = store.CreateWriter(cancellationToken);
-			await writer.WriteRefAsync(name, target, options, cancellationToken);
+			await writer.AddRefAsync(name, target, options, cancellationToken);
 		}
 
 		/// <summary>
@@ -374,12 +374,22 @@ namespace EpicGames.Horde.Storage
 		/// <param name="store">The store instance to write to</param>
 		/// <param name="name">The ref identifier</param>
 		/// <param name="cancellationToken">Cancellation token for the operation</param>
-		public static async Task DeleteRefAsync(this IStorageNamespace store, RefName name, CancellationToken cancellationToken = default)
+		public static async Task RemoveRefAsync(this IStorageNamespace store, RefName name, CancellationToken cancellationToken = default)
 		{
 			await using IStorageWriter writer = store.CreateWriter(cancellationToken);
-			await writer.DeleteRefAsync(name, cancellationToken);
+			await writer.RemoveRefAsync(name, cancellationToken);
 		}
 
+		/// <inheritdoc cref="AddRefAsync(IStorageNamespace, RefName, IHashedBlobRef, RefOptions?, CancellationToken)"/>
+		[Obsolete("Use AddRefAsync() instead")]
+		public static Task WriteRefAsync(this IStorageNamespace store, RefName name, IHashedBlobRef target, RefOptions? options = null, CancellationToken cancellationToken = default)
+			=> AddRefAsync(store, name, target, options, cancellationToken);
+
+		/// <inheritdoc cref="RemoveRefAsync(IStorageNamespace, RefName, CancellationToken)"/>
+		[Obsolete("Use RemoveRefAsync() instead")]
+		public static Task DeleteRefAsync(this IStorageNamespace store, RefName name, CancellationToken cancellationToken = default)
+			=> RemoveRefAsync(store, name, cancellationToken);
+
 		#endregion
 
 		/// <summary>
diff --git a/Engine/Source/Programs/Shared/EpicGames.Horde/Storage/IStorageWriter.cs b/Engine/Source/Programs/Shared/EpicGames.Horde/Storage/IStorageWriter.cs
--- a/Engine/Source/Programs/Shared/EpicGames.Horde/Storage/IStorageWriter.cs
+++ b/Engine/Source/Programs/Shared/EpicGames.Horde/Storage/IStorageWriter.cs
@@ -51,21 +51,21 @@ namespace EpicGames.Horde.Storage
 		#region Refs
 
 		/// <summary>
-		/// Writes a new ref to the store
+		/// Adds a new ref to the store
 		/// </summary>
 		/// <param name="name">Ref to write</param>
 		/// <param name="target">Handle to the target blob</param>
 		/// <param name="options">Options for the new ref</param>
 		/// <param name="cancellationToken">Cancellation token for the operation</param>
 		/// <returns>Unique identifier for the blob</returns>
-		Task WriteRefAsync(RefName name, IHashedBlobRef target, RefOptions? options = null, CancellationToken cancellationToken = default);
+		Task AddRefAsync(RefName name, IHashedBlobRef target, RefOptions? options = null, CancellationToken cancellationToken = default);
 
 		/// <summary>
-		/// Reads data for a ref from the store
+		/// Removes a ref from the store
 		/// </summary>
 		/// <param name="name">The ref identifier</param>
 		/// <param name="cancellationToken">Cancellation token for the operation</param>
-		Task DeleteRefAsync(RefName name, CancellationToken cancellationToken = default);
+		Task RemoveRefAsync(RefName name, CancellationToken cancellationToken = default);
 
 		#endregion
 
@@ -125,8 +125,10 @@ namespace EpicGames.Horde.Storage
 
 		#endregion
 
+		#region Refs
+
 		/// <inheritdoc/>
-		public async Task WriteRefAsync(RefName name, IHashedBlobRef target, RefOptions? options = null, CancellationToken cancellationToken = default)
+		public async Task AddRefAsync(RefName name, IHashedBlobRef target, RefOptions? options = null, CancellationToken cancellationToken = default)
 		{
 			await target.FlushAsync(cancellationToken);
 			options = (options != null) ? new RefOptions(options.Lifetime, options.Extend) : null;
@@ -135,12 +137,14 @@ namespace EpicGames.Horde.Storage
 		}
 
 		/// <inheritdoc/>
-		public async Task DeleteRefAsync(RefName name, CancellationToken cancellationToken = default)
+		public async Task RemoveRefAsync(RefName name, CancellationToken cancellationToken = default)
 		{
 			_request.RemoveRefs.Add(new RemoveRefRequest { RefName = name });
 			await CheckForFlushAsync(cancellationToken);
 		}
 
+		#endregion
+
 		async Task CheckForFlushAsync(CancellationToken cancellationToken = default)
 		{
 			int numItems = _request.AddAliases.Count + _request.RemoveAliases.Count + _request.AddRefs.Count + _request.RemoveRefs.Count;
diff --git a/Engine/Source/Programs/UnrealBuildTool/Artifacts/HordeStorageArtifactCache.cs b/Engine/Source/Programs/UnrealBuildTool/Artifacts/HordeStorageArtifactCache.cs
--- a/Engine/Source/Programs/UnrealBuildTool/Artifacts/HordeStorageArtifactCache.cs
+++ b/Engine/Source/Programs/UnrealBuildTool/Artifacts/HordeStorageArtifactCache.cs
@@ -480,7 +480,7 @@ namespace UnrealBuildTool.Artifacts
 					await writer.FlushAsync();
 
 					// Save the collection
-					await _store.WriteRefAsync(refName, nodeRef, cancellationToken: cancellationToken);
+					await _store.AddRefAsync(refName, nodeRef, cancellationToken: cancellationToken);
 				}, cancellationToken));
 			}
 			return tasks;
