Horde: Add logging for metadata updates.

#rnx

diff --git a/Engine/Source/Programs/Horde/Plugins/Storage/HordeServer.Storage/Storage/StorageService.cs b/Engine/Source/Programs/Horde/Plugins/Storage/HordeServer.Storage/Storage/StorageService.cs
--- a/Engine/Source/Programs/Horde/Plugins/Storage/HordeServer.Storage/Storage/StorageService.cs
+++ b/Engine/Source/Programs/Horde/Plugins/Storage/HordeServer.Storage/Storage/StorageService.cs
@@ -186,28 +186,50 @@ namespace HordeServer.Storage
 
 			/// <inheritdoc/>
 			public Task WriteRefAsync(RefName name, HashedBlobRefValue value, RefOptions? options = null, CancellationToken cancellationToken = default)
-				=> _outer.WriteRefAsync(NamespaceId, name, value, options, cancellationToken);
+				=> _outer.AddRefAsync(NamespaceId, name, value, options, cancellationToken);
 
 			/// <inheritdoc/>
 			public Task<bool> DeleteRefAsync(RefName name, CancellationToken cancellationToken = default)
-				=> _outer.DeleteRefAsync(NamespaceId, name, cancellationToken);
+				=> _outer.RemoveRefAsync(NamespaceId, name, cancellationToken);
 
 			#endregion
 
 			/// <inheritdoc/>
 			public async Task UpdateMetadataAsync(UpdateMetadataRequest request, CancellationToken cancellationToken = default)
 			{
-				await Parallel.ForEachAsync(request.AddAliases, cancellationToken, async (addAlias, cancellationToken)
-					=> await _outer.AddAliasAsync(NamespaceId, addAlias.Name, addAlias.Target, addAlias.Rank, addAlias.Data, cancellationToken));
-
-				await Parallel.ForEachAsync(request.RemoveAliases, cancellationToken, async (removeAlias, cancellationToken)
-					=> await _outer.RemoveAliasAsync(NamespaceId, removeAlias.Name, removeAlias.Target, cancellationToken));
+				List<Func<CancellationToken, Task>> actions = new List<Func<CancellationToken, Task>>();
+				foreach (AddAliasRequest addAlias in request.AddAliases)
+				{
+					actions.Add(ctx => _outer.AddAliasAsync(NamespaceId, addAlias.Name, addAlias.Target, addAlias.Rank, addAlias.Data, ctx));
+				}
+				foreach (RemoveAliasRequest addAlias in request.RemoveAliases)
+				{
+					actions.Add(ctx => _outer.RemoveAliasAsync(NamespaceId, addAlias.Name, addAlias.Target, ctx));
+				}
+				foreach (AddRefRequest addRef in request.AddRefs)
+				{
+					actions.Add(ctx => _outer.AddRefAsync(NamespaceId, addRef.RefName, new HashedBlobRefValue(addRef.Hash, addRef.Target), addRef.Options, ctx));
+				}
+				foreach (RemoveRefRequest removeRef in request.RemoveRefs)
+				{
+					actions.Add(ctx => _outer.RemoveRefAsync(NamespaceId, removeRef.RefName, ctx));
+				}
 
-				await Parallel.ForEachAsync(request.AddRefs, cancellationToken, async (addRef, cancellationToken) 
-					=> await _outer.WriteRefAsync(NamespaceId, addRef.RefName, new HashedBlobRefValue(addRef.Hash, addRef.Target), addRef.Options, cancellationToken));
+				await Parallel.ForEachAsync(actions, cancellationToken, (action, cancellationToken)
+					=> WrapAsync(() => action(cancellationToken)));
+			}
 
-				await Parallel.ForEachAsync(request.RemoveRefs, cancellationToken, async (removeRef, cancellationToken)
-					=> await _outer.DeleteRefAsync(NamespaceId, removeRef.RefName, cancellationToken));
+			async ValueTask WrapAsync(Func<Task> func)
+			{
+				try
+				{
+					await func();
+				}
+				catch (Exception ex)
+				{
+					_outer._logger.LogError(ex, "Error during metadata update: {Message}", ex.Message);
+					throw;
+				}
 			}
 
 			/// <inheritdoc/>
@@ -761,7 +783,7 @@ namespace HordeServer.Storage
 
 			// Compute missing import info, by searching for blobs with an ObjectId timestamp after the last import compute cycle
 			ObjectId latestInfoId = ObjectId.GenerateNewId(ingestTimeUtc);
-			for(; ;)
+			for (; ; )
 			{
 				// Check for the reset flag being set
 				if (gcState.Reset)
@@ -850,6 +872,8 @@ namespace HordeServer.Storage
 			FilterDefinition<BlobInfo> filter = Builders<BlobInfo>.Filter.Expr(x => x.NamespaceId == blobInfo.NamespaceId && x.Path == blobPath);
 			UpdateDefinition<BlobInfo> update = Builders<BlobInfo>.Update.Push(x => x.Aliases, new AliasInfo(name, blobFragment, data.ToArray(), rank));
 			await _blobCollection.UpdateOneAsync(filter, update, cancellationToken: cancellationToken);
+
+			_logger.LogDebug("Added alias {Name} to {Target}", name, target);
 		}
 
 		/// <summary>
@@ -868,6 +892,8 @@ namespace HordeServer.Storage
 			FilterDefinition<BlobInfo> filter = Builders<BlobInfo>.Filter.Expr(x => x.NamespaceId == namespaceId && x.Path == blobPath);
 			UpdateDefinition<BlobInfo> update = Builders<BlobInfo>.Update.PullFilter(x => x.Aliases, Builders<AliasInfo>.Filter.Expr(x => x.Name == name && x.Fragment == blobFragment));
 			await _blobCollection.UpdateOneAsync(filter, update, cancellationToken: cancellationToken);
+
+			_logger.LogDebug("Removed alias {Name} from {Target}", name, target);
 		}
 
 		/// <summary>
@@ -964,26 +990,28 @@ namespace HordeServer.Storage
 		}
 
 		/// <inheritdoc/>
-		async Task<bool> DeleteRefAsync(NamespaceId namespaceId, RefName name, CancellationToken cancellationToken = default)
+		async Task<bool> RemoveRefAsync(NamespaceId namespaceId, RefName name, CancellationToken cancellationToken = default)
 		{
 			FilterDefinition<RefInfo> filter = Builders<RefInfo>.Filter.Expr(x => x.NamespaceId == namespaceId && x.Name == name);
-			return await DeleteRefInternalAsync(namespaceId, name, filter, cancellationToken);
+			return await RemoveRefInternalAsync(namespaceId, name, filter, cancellationToken);
 		}
 
 		/// <summary>
 		/// Deletes a ref document that has reached its expiry time
 		/// </summary>
-		async Task DeleteExpiredRefAsync(RefInfo refDocument, CancellationToken cancellationToken = default)
+		async Task RemoveExpiredRefAsync(RefInfo refDocument, CancellationToken cancellationToken = default)
 		{
 			FilterDefinition<RefInfo> filter = Builders<RefInfo>.Filter.Expr(x => x.Id == refDocument.Id && x.ExpiresAtUtc == refDocument.ExpiresAtUtc);
-			await DeleteRefInternalAsync(refDocument.NamespaceId, refDocument.Name, filter, cancellationToken);
+			await RemoveRefInternalAsync(refDocument.NamespaceId, refDocument.Name, filter, cancellationToken);
 		}
 
-		async Task<bool> DeleteRefInternalAsync(NamespaceId namespaceId, RefName name, FilterDefinition<RefInfo> filter, CancellationToken cancellationToken = default)
+		async Task<bool> RemoveRefInternalAsync(NamespaceId namespaceId, RefName name, FilterDefinition<RefInfo> filter, CancellationToken cancellationToken = default)
 		{
 			RefInfo? oldRefInfo = await _refCollection.FindOneAndDeleteAsync<RefInfo>(filter, cancellationToken: cancellationToken);
 			AddRefToCache(namespaceId, name, default);
 
+			_logger.LogDebug("Removed ref {Name} from {Target}", name, oldRefInfo?.Target);
+
 			if (oldRefInfo != null)
 			{
 				_logger.LogInformation("Deleted ref {NamespaceId}:{RefName}", namespaceId, name);
@@ -1015,7 +1043,7 @@ namespace HordeServer.Storage
 				DateTime utcNow = _clock.UtcNow;
 				if (entry.Value.HasExpired(utcNow))
 				{
-					await DeleteExpiredRefAsync(entry.Value, cancellationToken);
+					await RemoveExpiredRefAsync(entry.Value, cancellationToken);
 					return default;
 				}
 				if (entry.Value.RequiresTouch(utcNow))
@@ -1028,7 +1056,7 @@ namespace HordeServer.Storage
 		}
 
 		/// <inheritdoc/>
-		async Task WriteRefAsync(NamespaceId namespaceId, RefName name, HashedBlobRefValue value, RefOptions? options = null, CancellationToken cancellationToken = default)
+		async Task AddRefAsync(NamespaceId namespaceId, RefName name, HashedBlobRefValue value, RefOptions? options = null, CancellationToken cancellationToken = default)
 		{
 			string path = value.Locator.BaseLocator.ToString();
 
@@ -1056,7 +1084,7 @@ namespace HordeServer.Storage
 				AddGcCheckRecord(namespaceId, oldRefInfo.TargetBlobId, gcState.LastImportBlobInfoId);
 			}
 
-			_logger.LogInformation("Updated ref {NamespaceId}:{RefName}", namespaceId, name);
+			_logger.LogDebug("Added ref {NamespaceId}:{Name} to {Target}", namespaceId, name, value);
 			AddRefToCache(namespaceId, name, newRefInfo);
 		}
 
@@ -1393,7 +1421,7 @@ namespace HordeServer.Storage
 			long scannedCount = 0;
 			LengthScanState lengthScanState = await _lengthScanState.GetAsync(cancellationToken);
 
-			for(; ;)
+			for (; ; )
 			{
 				// Reset the state
 				if (lengthScanState.Reset)
