Horde: Fix CLI tool upload using named blobs

diff --git a/Engine/Source/Programs/Horde/HordeCmd/Commands/Tools/ToolUpload.cs b/Engine/Source/Programs/Horde/HordeCmd/Commands/Tools/ToolUpload.cs
--- a/Engine/Source/Programs/Horde/HordeCmd/Commands/Tools/ToolUpload.cs
+++ b/Engine/Source/Programs/Horde/HordeCmd/Commands/Tools/ToolUpload.cs
@@ -31,7 +31,7 @@ class ToolUpload(IHordeClient hordeClient) : Command
 		IStorageNamespace storageNamespace = hordeClient.GetStorageNamespace(ToolId);
 		
 		IHashedBlobRef<DirectoryNode> target;
-		await using (IBlobWriter writer = storageNamespace.CreateBlobWriter())
+		await using (IBlobWriter writer = storageNamespace.CreateBlobWriter(ToolId.ToString(), null))
 		{
 			target = await writer.WriteFilesAsync(InputDir);
 		}
