Horde: Fix agent uninstall directory path when registry path is not used

diff --git a/Engine/Source/Programs/Horde/HordeAgent/Commands/Setup/SetupInstall.cs b/Engine/Source/Programs/Horde/HordeAgent/Commands/Setup/SetupInstall.cs
--- a/Engine/Source/Programs/Horde/HordeAgent/Commands/Setup/SetupInstall.cs
+++ b/Engine/Source/Programs/Horde/HordeAgent/Commands/Setup/SetupInstall.cs
@@ -112,7 +112,7 @@ class SetupInstallCommand : Command
 		}
 		
 		DirectoryReference sourceDir = new(installParams.SourceDir ?? AppContext.BaseDirectory);
-		DirectoryReference destinationDir = new(installParams.DestDir ?? Path.Join(Environment.GetFolderPath(Environment.SpecialFolder.ProgramFiles), "Epic Games", "Horde", "Agent"));
+		DirectoryReference destinationDir = new(installParams.DestDir ?? GetDefaultInstallDir());
 		DirectoryReference configDir = new(installParams.ConfigDir ?? Path.Join(Environment.GetFolderPath(Environment.SpecialFolder.CommonApplicationData), "Epic", "Horde", "Agent"));
 		DirectoryReference workingDir = new(installParams.WorkingDir ?? Path.Join("c:", "HordeAgent", "Sandbox"));
 		logger.LogInformation(" Source dir {SourceDir}", sourceDir);
@@ -213,6 +213,11 @@ class SetupInstallCommand : Command
 		return true;
 	}
 	
+	internal static string GetDefaultInstallDir()
+	{
+		return Path.Join(Environment.GetFolderPath(Environment.SpecialFolder.ProgramFiles), "Epic Games", "Horde", "Agent");
+	}
+	
 	internal static void AddFirewallPortRule(string name, string portRange)
 	{
 		RemoveFirewallPortRule(name);
@@ -300,8 +305,13 @@ class SetupInstallCommand : Command
 	}
 	
 	[SupportedOSPlatform("windows")]
-	public static DirectoryReference GetUninstallLocation()
+	public static DirectoryReference GetUninstallLocation(bool useRegistryPath = true)
 	{
+		if (!useRegistryPath)
+		{
+			return new DirectoryReference(GetDefaultInstallDir());
+		}
+		
 		const string RegKeyPath = $"{UninstallKeyPath}\\{UninstallAppName}";
 		using RegistryKey? key = Registry.LocalMachine.OpenSubKey(RegKeyPath, true);
 		if (key?.GetValue("InstallLocation") is string uninstallString)
diff --git a/Engine/Source/Programs/Horde/HordeAgent/Commands/Setup/SetupUninstall.cs b/Engine/Source/Programs/Horde/HordeAgent/Commands/Setup/SetupUninstall.cs
--- a/Engine/Source/Programs/Horde/HordeAgent/Commands/Setup/SetupUninstall.cs
+++ b/Engine/Source/Programs/Horde/HordeAgent/Commands/Setup/SetupUninstall.cs
@@ -64,7 +64,7 @@ class SetupUninstallCommand : Command
 			throw new InstallException("Administrator access is required for uninstalling");
 		}
 		
-		DirectoryReference installDir = useRegistryPath ? SetupInstallCommand.GetUninstallLocation() : new DirectoryReference(AppContext.BaseDirectory);
+		DirectoryReference installDir = SetupInstallCommand.GetUninstallLocation(useRegistryPath);
 		if (!SetupInstallCommand.ContainsHordeAgentExecutable(installDir))
 		{
 			throw new InstallException($"Directory to uninstall does not contain Horde agent executable");
