Horde: Hide agent install console window and skip uninstall registration in Windows Registry

diff --git a/Engine/Source/Programs/Horde/HordeAgent/Toolbox.json b/Engine/Source/Programs/Horde/HordeAgent/Toolbox.json
--- a/Engine/Source/Programs/Horde/HordeAgent/Toolbox.json
+++ b/Engine/Source/Programs/Horde/HordeAgent/Toolbox.json
@@ -1,12 +1,13 @@
 {
 	"manualInstall": true,
 	"requiresElevation": true,
+	"hidden": true,
 	"installCommand": {
 		"fileName": "HordeAgent.exe",
-		"arguments": ["setup", "install"]
+		"arguments": ["setup", "install", "-RegisterForUninstall=false"]
 	},
 	"uninstallCommand": {
 		"fileName": "HordeAgent.exe",
-		"arguments": ["setup", "uninstall", "-UseRegistryPath=true"]
+		"arguments": ["setup", "uninstall", "-UseRegistryPath=false"]
 	}
 }
\ No newline at end of file
