Horde: Use bulk writes to add and remove multiple aliases in a single operation.

#rnx

diff --git a/Engine/Source/Programs/Horde/Plugins/Storage/HordeServer.Storage.Tests/Storage/BlobStoreTests.cs b/Engine/Source/Programs/Horde/Plugins/Storage/HordeServer.Storage.Tests/Storage/BlobStoreTests.cs
--- a/Engine/Source/Programs/Horde/Plugins/Storage/HordeServer.Storage.Tests/Storage/BlobStoreTests.cs
+++ b/Engine/Source/Programs/Horde/Plugins/Storage/HordeServer.Storage.Tests/Storage/BlobStoreTests.cs
@@ -182,6 +182,63 @@ namespace HordeServer.Tests.Storage
 			}
 			return handle.GetRefValue();
 		}
+
+		[TestMethod]
+		public async Task AliasTestAsync()
+		{
+			IStorageNamespace store = GetStorageNamespace();
+
+			BlobType type = new BlobType(Guid.Empty, 0);
+
+			IBlobRef blobRef;
+			await using (IBlobWriter writer = store.CreateBlobWriter())
+			{
+				byte[] data = [1, 2, 3];
+				writer.WriteVariableLengthBytes(data);
+				blobRef = await writer.CompleteAsync(type);
+			}
+
+			await store.AddAliasAsync("hello", blobRef, rank: 4);
+			await store.AddAliasAsync("world", blobRef, rank: 4);
+			await store.AddAliasAsync("hello", blobRef, rank: 3); // Will be ignored because alias already exists
+			await store.AddAliasAsync("world", blobRef, rank: 3); // Will be ignored because alias already exists
+
+			BlobAlias[] aliases1 = await store.FindAliasesAsync("hello", 10);
+			Assert.AreEqual(1, aliases1.Length);
+			Assert.AreEqual(4, aliases1[0].Rank);
+
+			BlobAlias[] aliases2 = await store.FindAliasesAsync("world", 10);
+			Assert.AreEqual(1, aliases2.Length);
+			Assert.AreEqual(4, aliases2[0].Rank);
+		}
+
+		[TestMethod]
+		public async Task MultipleAliasTestsAsync()
+		{
+			IStorageNamespace store = GetStorageNamespace();
+
+			BlobType type = new BlobType(Guid.Empty, 0);
+
+			List<IBlobRef> blobRefs = new List<IBlobRef>();
+			await using (IBlobWriter writer = store.CreateBlobWriter())
+			{
+				byte[] data = [1, 2, 3];
+				for (int idx = 0; idx < 100; idx++)
+				{
+					writer.WriteVariableLengthBytes(data);
+					blobRefs.Add(await writer.CompleteAsync(type));
+				}
+				await writer.FlushAsync();
+			}
+
+			await using (IStorageWriter writer = store.CreateWriter())
+			{
+				for (int idx = 0; idx < blobRefs.Count; idx++)
+				{
+					await writer.AddAliasAsync("test", blobRefs[idx]);
+				}
+			}
+		}
 	}
 }
 
diff --git a/Engine/Source/Programs/Horde/Plugins/Storage/HordeServer.Storage/Storage/StorageService.cs b/Engine/Source/Programs/Horde/Plugins/Storage/HordeServer.Storage/Storage/StorageService.cs
--- a/Engine/Source/Programs/Horde/Plugins/Storage/HordeServer.Storage/Storage/StorageService.cs
+++ b/Engine/Source/Programs/Horde/Plugins/Storage/HordeServer.Storage/Storage/StorageService.cs
@@ -150,14 +150,6 @@ namespace HordeServer.Storage
 
 			#region Aliases
 
-			/// <inheritdoc/>
-			public Task AddAliasAsync(string name, BlobLocator locator, int rank = 0, ReadOnlyMemory<byte> data = default, CancellationToken cancellationToken = default)
-				=> _outer.AddAliasAsync(NamespaceId, name, locator, rank, data, cancellationToken);
-
-			/// <inheritdoc/>
-			public Task RemoveAliasAsync(string name, BlobLocator locator, CancellationToken cancellationToken = default)
-				=> _outer.RemoveAliasAsync(NamespaceId, name, locator, cancellationToken);
-
 			/// <inheritdoc/>
 			public async Task<BlobAliasLocator[]> FindAliasesAsync(string alias, int? maxResults, CancellationToken cancellationToken = default)
 			{
@@ -197,15 +189,16 @@ namespace HordeServer.Storage
 			/// <inheritdoc/>
 			public async Task UpdateMetadataAsync(UpdateMetadataRequest request, CancellationToken cancellationToken = default)
 			{
-				List<Func<CancellationToken, Task>> actions = new List<Func<CancellationToken, Task>>();
-				foreach (AddAliasRequest addAlias in request.AddAliases)
+				if (request.AddAliases.Count > 0)
 				{
-					actions.Add(ctx => _outer.AddAliasAsync(NamespaceId, addAlias.Name, addAlias.Target, addAlias.Rank, addAlias.Data, ctx));
+					await _outer.AddAliasesAsync(NamespaceId, request.AddAliases, cancellationToken);
 				}
-				foreach (RemoveAliasRequest addAlias in request.RemoveAliases)
+				if (request.RemoveAliases.Count > 0)
 				{
-					actions.Add(ctx => _outer.RemoveAliasAsync(NamespaceId, addAlias.Name, addAlias.Target, ctx));
+					await _outer.RemoveAliasesAsync(NamespaceId, request.RemoveAliases, cancellationToken);
 				}
+
+				List<Func<CancellationToken, Task>> actions = new List<Func<CancellationToken, Task>>();
 				foreach (AddRefRequest addRef in request.AddRefs)
 				{
 					actions.Add(ctx => _outer.AddRefAsync(NamespaceId, addRef.RefName, new HashedBlobRefValue(addRef.Hash, addRef.Target), addRef.Options, ctx));
@@ -844,56 +837,57 @@ namespace HordeServer.Storage
 
 		#region Nodes
 
-		/// <summary>
-		/// Adds a node alias
-		/// </summary>
-		/// <param name="namespaceId">Namespace to search</param>
-		/// <param name="name">Alias for the node</param>
-		/// <param name="target">Target node for the alias</param>
-		/// <param name="rank">Rank for the alias. Higher ranked aliases are preferred by default.</param>
-		/// <param name="data">Inline data to store with this alias</param>
-		/// <param name="cancellationToken">Cancellation token for the operation</param>
-		/// <returns>Sequence of thandles</returns>
-		async Task AddAliasAsync(NamespaceId namespaceId, string name, BlobLocator target, int rank, ReadOnlyMemory<byte> data, CancellationToken cancellationToken = default)
+		async Task AddAliasesAsync(NamespaceId namespaceId, IEnumerable<AddAliasRequest> addAliases, CancellationToken cancellationToken)
 		{
-			string blobPath = target.BaseLocator.ToString();
-			string blobFragment = target.Fragment.ToString();
+			List<WriteModel<BlobInfo>> writes = new List<WriteModel<BlobInfo>>();
 
-			BlobInfo? blobInfo = await _blobCollection.Find(x => x.NamespaceId == namespaceId && x.Path == blobPath).FirstOrDefaultAsync(cancellationToken);
-			if (blobInfo == null)
+			foreach (AddAliasRequest addAlias in addAliases)
 			{
-				throw new KeyNotFoundException($"Missing blob {blobPath}");
-			}
-			if (blobInfo.Aliases != null && blobInfo.Aliases.Any(x => x.Name.Equals(name, StringComparison.Ordinal) && x.Fragment.Equals(blobFragment, StringComparison.Ordinal)))
-			{
-				return;
+				string path = addAlias.Target.BaseLocator.ToString();
+				string fragment = addAlias.Target.Fragment.ToString();
+
+				FilterDefinition<BlobInfo> filter =
+					Builders<BlobInfo>.Filter.Expr(x => x.NamespaceId == namespaceId && x.Path == path) &
+					Builders<BlobInfo>.Filter.Not(Builders<BlobInfo>.Filter.ElemMatch(x => x.Aliases, x => x.Name == addAlias.Name && x.Fragment == fragment));
+
+				UpdateDefinition<BlobInfo> update =
+					Builders<BlobInfo>.Update.Push(x => x.Aliases, new AliasInfo(addAlias.Name, fragment, addAlias.Data?.ToArray(), addAlias.Rank));
+
+				writes.Add(new UpdateOneModel<BlobInfo>(filter, update));
 			}
 
-			FilterDefinition<BlobInfo> filter = Builders<BlobInfo>.Filter.Expr(x => x.NamespaceId == blobInfo.NamespaceId && x.Path == blobPath);
-			UpdateDefinition<BlobInfo> update = Builders<BlobInfo>.Update.Push(x => x.Aliases, new AliasInfo(name, blobFragment, data.ToArray(), rank));
-			await _blobCollection.UpdateOneAsync(filter, update, cancellationToken: cancellationToken);
+			await _blobCollection.BulkWriteAsync(writes, new BulkWriteOptions { IsOrdered = false }, cancellationToken);
 
-			_logger.LogDebug("Added alias {Name} to {Target}", name, target);
+			foreach (AddAliasRequest addAlias in addAliases)
+			{
+				_logger.LogDebug("Added alias {Name} to {Target}", addAlias.Name, addAlias.Target);
+			}
 		}
 
-		/// <summary>
-		/// Removes a node alias
-		/// </summary>
-		/// <param name="namespaceId">Namespace to search</param>
-		/// <param name="name">Alias for the node</param>
-		/// <param name="target">Target node for the alias</param>
-		/// <param name="cancellationToken">Cancellation token for the operation</param>
-		/// <returns>Sequence of thandles</returns>
-		async Task RemoveAliasAsync(NamespaceId namespaceId, string name, BlobLocator target, CancellationToken cancellationToken = default)
+		async Task RemoveAliasesAsync(NamespaceId namespaceId, IEnumerable<RemoveAliasRequest> removeAliases, CancellationToken cancellationToken)
 		{
-			string blobPath = target.BaseLocator.ToString();
-			string blobFragment = target.Fragment.ToString();
+			List<WriteModel<BlobInfo>> writes = new List<WriteModel<BlobInfo>>();
 
-			FilterDefinition<BlobInfo> filter = Builders<BlobInfo>.Filter.Expr(x => x.NamespaceId == namespaceId && x.Path == blobPath);
-			UpdateDefinition<BlobInfo> update = Builders<BlobInfo>.Update.PullFilter(x => x.Aliases, Builders<AliasInfo>.Filter.Expr(x => x.Name == name && x.Fragment == blobFragment));
-			await _blobCollection.UpdateOneAsync(filter, update, cancellationToken: cancellationToken);
+			foreach (RemoveAliasRequest removeAlias in removeAliases)
+			{
+				string path = removeAlias.Target.BaseLocator.ToString();
+				string fragment = removeAlias.Target.Fragment.ToString();
+
+				FilterDefinition<BlobInfo> filter = 
+					Builders<BlobInfo>.Filter.Expr(x => x.NamespaceId == namespaceId && x.Path == path);
+
+				UpdateDefinition<BlobInfo> update = 
+					Builders<BlobInfo>.Update.PullFilter(x => x.Aliases, Builders<AliasInfo>.Filter.Expr(x => x.Name == removeAlias.Name && x.Fragment == fragment));
+
+				writes.Add(new UpdateOneModel<BlobInfo>(filter, update));
+			}
+
+			await _blobCollection.BulkWriteAsync(writes, new BulkWriteOptions { IsOrdered = false }, cancellationToken);
 
-			_logger.LogDebug("Removed alias {Name} from {Target}", name, target);
+			foreach (RemoveAliasRequest removeAlias in removeAliases)
+			{
+				_logger.LogDebug("Removed alias {Name} from {Target}", removeAlias.Name, removeAlias.Target);
+			}
 		}
 
 		/// <summary>
