Horde: Add notes to the artifacts doc page about the RetrieveArtifact task.

#rnx

diff --git a/Engine/Source/Programs/Horde/Docs/Config/Artifacts.md b/Engine/Source/Programs/Horde/Docs/Config/Artifacts.md
--- a/Engine/Source/Programs/Horde/Docs/Config/Artifacts.md
+++ b/Engine/Source/Programs/Horde/Docs/Config/Artifacts.md
@@ -144,6 +144,21 @@ faster and more efficient - the following methods are available:
   enumerating and downloading artifacts programmatically. The Horde command-line
   tool serves as an example for using this API to query artifacts.
 
+The `RetrieveArtifact` task can be used to find and download artifacts from
+a build script. By default, the task will search for an artifact in the current
+stream, at the current changelist - though the `Stream` and `Commit` attributes
+can be used to choose a different source. 
+
+  ```xml
+  <RetrieveArtifact
+      Name="build-name"
+      Type="installed-build"
+      OutputDir="My/Output/Dir" />
+  ```
+
+Additionally, specifying the `MaxCommit` parameter will find the most recent artifact at or 
+before the given change.
+
 ## Keys & Metadata
 
 Artifacts are designed to be consumed programmatically and include various
