Horde: Batch all alias updates by target blob. Interior nodes for a directory tree are typically written to the same blob, which makes issuing separate requests to update the same document in the DB an O(n^2) operation.

#rnx

diff --git a/Engine/Source/Programs/Horde/Plugins/Storage/HordeServer.Storage/Storage/StorageService.cs b/Engine/Source/Programs/Horde/Plugins/Storage/HordeServer.Storage/Storage/StorageService.cs
--- a/Engine/Source/Programs/Horde/Plugins/Storage/HordeServer.Storage/Storage/StorageService.cs
+++ b/Engine/Source/Programs/Horde/Plugins/Storage/HordeServer.Storage/Storage/StorageService.cs
@@ -319,6 +319,9 @@ namespace HordeServer.Storage
 			[BsonElement("len"), BsonIgnoreIfDefault]
 			public long Length { get; set; }
 
+			[BsonElement("idx"), BsonIgnoreIfDefault]
+			public int? UpdateIndex { get; set; }
+
 			[BsonIgnore]
 			public BlobLocator Locator => new BlobLocator(Path);
 
@@ -839,54 +842,123 @@ namespace HordeServer.Storage
 
 		async Task AddAliasesAsync(NamespaceId namespaceId, IEnumerable<AddAliasRequest> addAliases, CancellationToken cancellationToken)
 		{
-			List<WriteModel<BlobInfo>> writes = new List<WriteModel<BlobInfo>>();
-
-			foreach (AddAliasRequest addAlias in addAliases)
+			foreach (IGrouping<BlobLocator, AddAliasRequest> group in addAliases.GroupBy(x => x.Target.BaseLocator))
 			{
-				string path = addAlias.Target.BaseLocator.ToString();
-				string fragment = addAlias.Target.Fragment.ToString();
+				string path = group.Key.ToString();
+				for (; ; )
+				{
+					// Get the existing blob definition
+					BlobInfo? blobInfo = await _blobCollection.Find(x => x.NamespaceId == namespaceId && x.Path == path).FirstOrDefaultAsync(cancellationToken);
+					if (blobInfo == null)
+					{
+						break;
+					}
 
-				FilterDefinition<BlobInfo> filter =
-					Builders<BlobInfo>.Filter.Expr(x => x.NamespaceId == namespaceId && x.Path == path) &
-					Builders<BlobInfo>.Filter.Not(Builders<BlobInfo>.Filter.ElemMatch(x => x.Aliases, x => x.Name == addAlias.Name && x.Fragment == fragment));
+					// Create all the new alises
+					List<AliasInfo> aliases = new List<AliasInfo>();
+					foreach (AddAliasRequest addAlias in group)
+					{
+						string fragment = addAlias.Target.Fragment.ToString();
+						aliases.Add(new AliasInfo(addAlias.Name, fragment, addAlias.Data, addAlias.Rank));
+					}
+
+					// Add all the current aliases that haven't been replaced
+					if (blobInfo.Aliases != null)
+					{
+						HashSet<(string Name, string Fragment)> excludeAliases = aliases.Select(x => (x.Name, x.Fragment)).ToHashSet();
+						foreach (AliasInfo aliasInfo in blobInfo.Aliases)
+						{
+							if (!excludeAliases.Contains((aliasInfo.Name, aliasInfo.Fragment)))
+							{
+								aliases.Add(aliasInfo);
+							}
+						}
+					}
 
-				UpdateDefinition<BlobInfo> update =
-					Builders<BlobInfo>.Update.Push(x => x.Aliases, new AliasInfo(addAlias.Name, fragment, addAlias.Data?.ToArray(), addAlias.Rank));
+					// Update the blob info
+					FilterDefinition<BlobInfo> filter = Builders<BlobInfo>.Filter.Eq(x => x.Id, blobInfo.Id);
+					if (blobInfo.UpdateIndex == null)
+					{
+						filter &= Builders<BlobInfo>.Filter.Exists(x => x.UpdateIndex, false);
+					}
+					else
+					{
+						filter &= Builders<BlobInfo>.Filter.Eq(x => x.UpdateIndex, blobInfo.UpdateIndex);
+					}
 
-				writes.Add(new UpdateOneModel<BlobInfo>(filter, update));
-			}
+					UpdateDefinition<BlobInfo> update = Builders<BlobInfo>.Update.Set(x => x.Aliases, aliases).Set(x => x.UpdateIndex, (blobInfo.UpdateIndex ?? 0) + 1);
 
-			await _blobCollection.BulkWriteAsync(writes, new BulkWriteOptions { IsOrdered = false }, cancellationToken);
+					UpdateResult result = await _blobCollection.UpdateOneAsync(filter, update, cancellationToken: cancellationToken);
+					if (result.MatchedCount > 0)
+					{
+						break;
+					}
+				}
 
-			foreach (AddAliasRequest addAlias in addAliases)
-			{
-				_logger.LogDebug("Added alias {Name} to {Target}", addAlias.Name, addAlias.Target);
+				foreach (AddAliasRequest addAlias in group)
+				{
+					_logger.LogDebug("Added alias {Name} to {Target}", addAlias.Name, addAlias.Target);
+				}
 			}
 		}
 
 		async Task RemoveAliasesAsync(NamespaceId namespaceId, IEnumerable<RemoveAliasRequest> removeAliases, CancellationToken cancellationToken)
 		{
-			List<WriteModel<BlobInfo>> writes = new List<WriteModel<BlobInfo>>();
-
-			foreach (RemoveAliasRequest removeAlias in removeAliases)
+			foreach (IGrouping<BlobLocator, RemoveAliasRequest> group in removeAliases.GroupBy(x => x.Target.BaseLocator))
 			{
-				string path = removeAlias.Target.BaseLocator.ToString();
-				string fragment = removeAlias.Target.Fragment.ToString();
+				string path = group.Key.ToString();
+				for (; ; )
+				{
+					// Get the existing blob definition
+					BlobInfo? blobInfo = await _blobCollection.Find(x => x.NamespaceId == namespaceId && x.Path == path).FirstOrDefaultAsync(cancellationToken);
+					if (blobInfo == null || blobInfo.Aliases == null)
+					{
+						break;
+					}
 
-				FilterDefinition<BlobInfo> filter = 
-					Builders<BlobInfo>.Filter.Expr(x => x.NamespaceId == namespaceId && x.Path == path);
+					// Build a set of aliases to remove
+					HashSet<(string Name, string Fragment)> excludeAliases = new();
+					foreach (RemoveAliasRequest removeAlias in group)
+					{
+						string fragment = removeAlias.Target.Fragment.ToString();
+						excludeAliases.Add((removeAlias.Name, fragment));
+					}
 
-				UpdateDefinition<BlobInfo> update = 
-					Builders<BlobInfo>.Update.PullFilter(x => x.Aliases, Builders<AliasInfo>.Filter.Expr(x => x.Name == removeAlias.Name && x.Fragment == fragment));
+					// Remove them from the blob
+					List<AliasInfo> aliases = blobInfo.Aliases.Where(x => !excludeAliases.Contains((x.Name, x.Fragment))).ToList();
 
-				writes.Add(new UpdateOneModel<BlobInfo>(filter, update));
-			}
+					// Update the blob info
+					FilterDefinition<BlobInfo> filter = Builders<BlobInfo>.Filter.Eq(x => x.Id, blobInfo.Id);
+					if (blobInfo.UpdateIndex == null)
+					{
+						filter &= Builders<BlobInfo>.Filter.Exists(x => x.UpdateIndex, false);
+					}
+					else
+					{
+						filter &= Builders<BlobInfo>.Filter.Eq(x => x.UpdateIndex, blobInfo.UpdateIndex);
+					}
 
-			await _blobCollection.BulkWriteAsync(writes, new BulkWriteOptions { IsOrdered = false }, cancellationToken);
+					UpdateDefinition<BlobInfo> update = Builders<BlobInfo>.Update.Inc(x => x.UpdateIndex!.Value, 1);
+					if (aliases.Count == 0)
+					{
+						update = update.Unset(x => x.Aliases);
+					}
+					else
+					{
+						update = update.Set(x => x.Aliases, aliases);
+					}
 
-			foreach (RemoveAliasRequest removeAlias in removeAliases)
-			{
-				_logger.LogDebug("Removed alias {Name} from {Target}", removeAlias.Name, removeAlias.Target);
+					UpdateResult result = await _blobCollection.UpdateOneAsync(filter, update, cancellationToken: cancellationToken);
+					if (result.MatchedCount > 0)
+					{
+						break;
+					}
+				}
+
+				foreach (RemoveAliasRequest removeAlias in group)
+				{
+					_logger.LogDebug("Removed alias {Name} from {Target}", removeAlias.Name, removeAlias.Target);
+				}
 			}
 		}
 
