Horde: Fix test failures.

#rnx

diff --git a/Engine/Source/Programs/Horde/Plugins/Storage/HordeServer.Storage.Tests/Storage/BlobStoreTests.cs b/Engine/Source/Programs/Horde/Plugins/Storage/HordeServer.Storage.Tests/Storage/BlobStoreTests.cs
--- a/Engine/Source/Programs/Horde/Plugins/Storage/HordeServer.Storage.Tests/Storage/BlobStoreTests.cs
+++ b/Engine/Source/Programs/Horde/Plugins/Storage/HordeServer.Storage.Tests/Storage/BlobStoreTests.cs
@@ -200,16 +200,16 @@ namespace HordeServer.Tests.Storage
 
 			await store.AddAliasAsync("hello", blobRef, rank: 4);
 			await store.AddAliasAsync("world", blobRef, rank: 4);
-			await store.AddAliasAsync("hello", blobRef, rank: 3); // Will be ignored because alias already exists
-			await store.AddAliasAsync("world", blobRef, rank: 3); // Will be ignored because alias already exists
+			await store.AddAliasAsync("hello", blobRef, rank: 3);
+			await store.AddAliasAsync("world", blobRef, rank: 3);
 
 			BlobAlias[] aliases1 = await store.FindAliasesAsync("hello", 10);
 			Assert.AreEqual(1, aliases1.Length);
-			Assert.AreEqual(4, aliases1[0].Rank);
+			Assert.AreEqual(3, aliases1[0].Rank);
 
 			BlobAlias[] aliases2 = await store.FindAliasesAsync("world", 10);
 			Assert.AreEqual(1, aliases2.Length);
-			Assert.AreEqual(4, aliases2[0].Rank);
+			Assert.AreEqual(3, aliases2[0].Rank);
 		}
 
 		[TestMethod]
@@ -223,12 +223,11 @@ namespace HordeServer.Tests.Storage
 			await using (IBlobWriter writer = store.CreateBlobWriter())
 			{
 				byte[] data = [1, 2, 3];
-				for (int idx = 0; idx < 100; idx++)
+				for (int idx = 0; idx < 1000; idx++)
 				{
 					writer.WriteVariableLengthBytes(data);
 					blobRefs.Add(await writer.CompleteAsync(type));
 				}
-				await writer.FlushAsync();
 			}
 
 			await using (IStorageWriter writer = store.CreateWriter())
diff --git a/Engine/Source/Programs/Horde/Plugins/Storage/HordeServer.Storage/Storage/StorageService.cs b/Engine/Source/Programs/Horde/Plugins/Storage/HordeServer.Storage/Storage/StorageService.cs
--- a/Engine/Source/Programs/Horde/Plugins/Storage/HordeServer.Storage/Storage/StorageService.cs
+++ b/Engine/Source/Programs/Horde/Plugins/Storage/HordeServer.Storage/Storage/StorageService.cs
@@ -938,7 +938,7 @@ namespace HordeServer.Storage
 						filter &= Builders<BlobInfo>.Filter.Eq(x => x.UpdateIndex, blobInfo.UpdateIndex);
 					}
 
-					UpdateDefinition<BlobInfo> update = Builders<BlobInfo>.Update.Inc(x => x.UpdateIndex!.Value, 1);
+					UpdateDefinition<BlobInfo> update = Builders<BlobInfo>.Update.Set(x => x.UpdateIndex, (blobInfo.UpdateIndex ?? 0) + 1);
 					if (aliases.Count == 0)
 					{
 						update = update.Unset(x => x.Aliases);
