Horde: Fix ObjectDisposedException if session update throws an exception.

#rnx

diff --git a/Engine/Source/Programs/Horde/Plugins/Compute/HordeServer.Compute/Server/RpcService.cs b/Engine/Source/Programs/Horde/Plugins/Compute/HordeServer.Compute/Server/RpcService.cs
--- a/Engine/Source/Programs/Horde/Plugins/Compute/HordeServer.Compute/Server/RpcService.cs
+++ b/Engine/Source/Programs/Horde/Plugins/Compute/HordeServer.Compute/Server/RpcService.cs
@@ -264,72 +264,77 @@ namespace HordeServer.Server
 				// Get a task for moving to the next item. This will only complete once the call has closed.
 				using CancellationTokenSource cancellationSource = CancellationTokenSource.CreateLinkedTokenSource(context.CancellationToken);
 				nextRequestTask = reader.MoveNext();
-				nextRequestTask = nextRequestTask.ContinueWith(task =>
-				{
-					cancellationSource.Cancel();
-					return task.IsCanceled ? false : task.Result;
-				}, TaskScheduler.Current);
-
-				// Get the current agent state
-				IAgent? agent = await _agentService.GetAgentAsync(new AgentId(request.AgentId));
-				if (agent != null)
+				try
 				{
-					SessionId sessionId = SessionId.Parse(request.SessionId);
+					nextRequestTask = nextRequestTask.ContinueWith(task =>
+					{
+						cancellationSource.Cancel();
+						return task.IsCanceled ? false : task.Result;
+					}, TaskScheduler.Current);
 
-					// Check we're authorized to update it
-					if (agent.SessionId != sessionId)
+					// Get the current agent state
+					IAgent? agent = await _agentService.GetAgentAsync(new AgentId(request.AgentId));
+					if (agent != null)
 					{
-						throw new StructuredRpcException(StatusCode.PermissionDenied, "Agent {AgentId} has completed session {SessionId}; now executing session {NewSessionId}. Cannot update state.", request.AgentId, sessionId, agent.SessionId?.ToString() ?? "(None)");
+						SessionId sessionId = SessionId.Parse(request.SessionId);
+
+						// Check we're authorized to update it
+						if (agent.SessionId != sessionId)
+						{
+							throw new StructuredRpcException(StatusCode.PermissionDenied, "Agent {AgentId} has completed session {SessionId}; now executing session {NewSessionId}. Cannot update state.", request.AgentId, sessionId, agent.SessionId?.ToString() ?? "(None)");
+						}
+						if (!_agentService.AuthorizeSession(agent, context.GetHttpContext().User, out string authReason))
+						{
+							throw new StructuredRpcException(StatusCode.PermissionDenied, "Not authenticated for {AgentId}. Reason {Reason}", request.AgentId, authReason);
+						}
+
+						// Update the session
+						try
+						{
+							agent = await _agentService.UpdateSessionWithWaitAsync(agent, sessionId, (AgentStatus)request.Status, request.Capabilities?.MergeDevices(), request.Leases, cancellationSource.Token);
+						}
+						catch (OperationCanceledException)
+						{
+							// Ignore cancellation due to a message having been received
+						}
+						catch (Exception ex)
+						{
+							_logger.LogError(ex, "Swallowed exception while updating session for {AgentId}.", request.AgentId);
+							throw new StructuredRpcException(StatusCode.Internal, "Failed updating session. Reason: {Reason}", ex.Message);
+						}
 					}
-					if (!_agentService.AuthorizeSession(agent, context.GetHttpContext().User, out string authReason))
+
+					// Handle the invalid agent case
+					if (agent == null)
 					{
-						throw new StructuredRpcException(StatusCode.PermissionDenied, "Not authenticated for {AgentId}. Reason {Reason}", request.AgentId, authReason);
+						throw new StructuredRpcException(StatusCode.NotFound, "Invalid agent name '{AgentId}'", request.AgentId);
 					}
 
-					// Update the session
-					try
+					// Create the new session info
+					if (!context.CancellationToken.IsCancellationRequested)
 					{
-						agent = await _agentService.UpdateSessionWithWaitAsync(agent, sessionId, (AgentStatus)request.Status, request.Capabilities?.MergeDevices(), request.Leases, cancellationSource.Token);
+						RpcUpdateSessionResponse response = new RpcUpdateSessionResponse();
+						response.Leases.Add(agent.Leases.Select(x => x.ToRpcMessage()));
+						response.ExpiryTime = (agent.SessionExpiresAt == null) ? new Timestamp() : Timestamp.FromDateTime(agent.SessionExpiresAt.Value);
+						response.Status = (RpcAgentStatus)agent.Status;
+						await writer.WriteAsync(response);
 					}
-					catch (OperationCanceledException)
+				}
+				finally
+				{
+					// Wait for the client to close the stream
+					try
 					{
-						// Ignore cancellation due to a message having been received
+						while (await nextRequestTask)
+						{
+							nextRequestTask = reader.MoveNext();
+						}
 					}
 					catch (Exception ex)
 					{
-						_logger.LogError(ex, "Swallowed exception while updating session for {AgentId}.", request.AgentId);
-						throw new StructuredRpcException(StatusCode.Internal, "Failed updating session. Reason: {Reason}", ex.Message);
-					}
-				}
-
-				// Handle the invalid agent case
-				if (agent == null)
-				{
-					throw new StructuredRpcException(StatusCode.NotFound, "Invalid agent name '{AgentId}'", request.AgentId);
-				}
-
-				// Create the new session info
-				if (!context.CancellationToken.IsCancellationRequested)
-				{
-					RpcUpdateSessionResponse response = new RpcUpdateSessionResponse();
-					response.Leases.Add(agent.Leases.Select(x => x.ToRpcMessage()));
-					response.ExpiryTime = (agent.SessionExpiresAt == null) ? new Timestamp() : Timestamp.FromDateTime(agent.SessionExpiresAt.Value);
-					response.Status = (RpcAgentStatus)agent.Status;
-					await writer.WriteAsync(response);
-				}
-
-				// Wait for the client to close the stream
-				try
-				{
-					while (await nextRequestTask)
-					{
-						nextRequestTask = reader.MoveNext();
+						_logger.LogDebug(ex, "Ignoring exception while finishing UpdateSession request.");
 					}
 				}
-				catch (Exception ex)
-				{
-					_logger.LogDebug(ex, "Ignoring exception while finishing UpdateSession request.");
-				}
 			}
 		}
 
