Horde: Job search improvements, link to related jobs from operations bar, step callout optimizations
#skipci

diff --git a/Engine/Source/Programs/Horde/HordeDashboard/src/components/JobLabelCallout.tsx b/Engine/Source/Programs/Horde/HordeDashboard/src/components/JobLabelCallout.tsx
--- a/Engine/Source/Programs/Horde/HordeDashboard/src/components/JobLabelCallout.tsx
+++ b/Engine/Source/Programs/Horde/HordeDashboard/src/components/JobLabelCallout.tsx
@@ -5,9 +5,10 @@ import { observer } from 'mobx-react-lite';
 import { Callout, DirectionalHint, List, Stack, Text } from "@fluentui/react";
 import React from "react";
 import { Link } from "react-router-dom";
-import { GetLabelStateResponse, StepData } from "../backend/Api";
+import { GetJobResponse, GetLabelStateResponse, StepData } from "../backend/Api";
 import { JobDetails } from "../backend/JobDetails";
 import { StepStatusIcon } from './StatusIcon';
+import backend from 'horde/backend';
 
 type StepItem = {
    step: StepData;
@@ -36,11 +37,11 @@ export class CalloutController {
 
          if (now) {
             this.setStateInternal({});
-            this.details.clear();
+            this.job = undefined;
          } else {
             this.closeTimeout = setTimeout(() => {
                this.setStateInternal({});
-               this.details.clear();
+               this.job = undefined;
             }, 600);
 
          }
@@ -70,14 +71,34 @@ export class CalloutController {
          this.setTimeout = undefined;
          this.pending = undefined;
 
-         if (this.details.id === state.jobId) {
+         if (this.job?.id === state.jobId) {
             this.setStateInternal(state);
             return;
          }
 
-         this.details.set(state.jobId!, undefined, undefined, undefined, 0, () => {
-            this.setStateInternal(state);
-         });
+         if (state.jobId) {
+            let job = this.jobCache.get(state.jobId);
+            const time = this.jobCacheTime.get(state.jobId);
+
+            if (job && time) {
+               if ((Date.now() - time.getTime()) / 1000 > 120) {
+                  job = undefined;                  
+               }
+            }
+
+            if (job) {
+               this.job = job;
+               this.setStateInternal(state);
+            } else {
+               backend.getJob(state.jobId!).then(job => {
+                  this.job = job;
+                  this.jobCache.set(job.id, job);
+                  this.jobCacheTime.set(job.id, new Date());
+                  this.setStateInternal(state);
+               })      
+            }
+         }
+
       }, 800);
 
    }
@@ -90,7 +111,7 @@ export class CalloutController {
 
    @action
    clear() {
-      this.details.clear();
+      this.job = undefined;
       this._state = {};
       this.stateUpdated++;
    }
@@ -107,8 +128,9 @@ export class CalloutController {
    private _state: CalloutState = {};
 
    pending?: CalloutState;
-
-   details = new JobDetails();
+   job?: GetJobResponse;
+   jobCache: Map<string, GetJobResponse> = new Map();
+   jobCacheTime: Map<string, Date> = new Map();
 
    setTimeout: any;
    closeTimeout: any;
@@ -121,13 +143,13 @@ export const JobLabelCallout: React.FC<{ controller: CalloutController }> = obse
    const label = state.label;
    const jobId = state.jobId;
    const target = state.target;
-   const details = controller.details;
+   const job = controller.job;
 
-   if (!label) {
+   if (!label || !job) {
       return <div />;
    }
 
-   const allSteps = details.getSteps();
+   const allSteps = (job.batches ?? []).map(b => b.steps ?? []).flat();
    const steps = allSteps.filter(step => label.steps.indexOf(step.id) !== -1);
    const items = steps.map(step => { return { step: step }; });
 
@@ -141,8 +163,7 @@ export const JobLabelCallout: React.FC<{ controller: CalloutController }> = obse
 
       const stepUrl = `/job/${jobId}?step=${step.id}`;
 
-      const stepId = step.id;
-      const stepName = details.getStepName(stepId);
+      const stepName = step.name;
 
       return <Stack tokens={{ childrenGap: 12 }} >
          <Stack horizontal>
diff --git a/Engine/Source/Programs/Horde/HordeDashboard/src/components/JobSearchSimple.tsx b/Engine/Source/Programs/Horde/HordeDashboard/src/components/JobSearchSimple.tsx
--- a/Engine/Source/Programs/Horde/HordeDashboard/src/components/JobSearchSimple.tsx
+++ b/Engine/Source/Programs/Horde/HordeDashboard/src/components/JobSearchSimple.tsx
@@ -1,10 +1,10 @@
-import { Checkbox, DatePicker, DefaultButton, DetailsList, DetailsRow, IColumn, IconButton, IContextualMenuItem, IContextualMenuProps, IDetailsGroupRenderProps, IDetailsListProps, IGroup, Label, Modal, PrimaryButton, SelectionMode, Spinner, SpinnerSize, Stack, Text, TextField } from "@fluentui/react";
+import { Checkbox, DatePicker, DefaultButton, DetailsList, DetailsRow, IColumn, IconButton, IContextualMenuItem, IContextualMenuItemProps, IContextualMenuProps, IDetailsGroupRenderProps, IDetailsListProps, IGroup, Label, Modal, PrimaryButton, SelectionMode, Spinner, SpinnerSize, Stack, Text, TextField } from "@fluentui/react";
 import { observer } from "mobx-react-lite";
 import moment from "moment";
 import { useState } from "react";
 import { Link } from "react-router-dom";
 import backend from "../backend";
-import { GetJobResponse, GetJobsTabResponse, GetJobStepRefResponse, JobData, JobQuery, TabType, GetTemplateRefResponse } from "../backend/Api";
+import { GetJobResponse, GetJobsTabResponse, GetJobStepRefResponse, JobData, JobQuery, TabType, GetTemplateRefResponse, GetJobsTabLabelColumnResponse, GetLabelStateResponse, LabelOutcome, LabelState, GetJobsTabColumnResponse } from "../backend/Api";
 import dashboard from "../backend/Dashboard";
 import { projectStore } from "../backend/ProjectStore";
 import templateCache from '../backend/TemplateCache';
@@ -14,22 +14,31 @@ import { StepRefStatusIcon } from "./StatusIcon";
 import { UserSelect } from "./UserSelect";
 import { getHordeTheme } from "../styles/theme";
 import { getHordeStyling } from "../styles/Styles";
+import { StreamChooser } from "./projects/StreamChooser";
+import { getLabelColor } from "horde/styles/colors";
+import { CalloutController, JobLabelCallout } from "./JobLabelCallout";
 
 
 const jobToStep = new Map<string, GetJobStepRefResponse>();
 let jobStepName = "";
+let streamChooserId = 0;
 
-export const JobSearchSimpleModal: React.FC<{ streamId: string, onClose: () => void }> = observer(({ streamId, onClose }) => {
+export const JobSearchSimpleModal: React.FC<{ streamIdIn: string, changeListIn?: string, jobTab?: GetJobsTabResponse | undefined, autoloadIn?: boolean, onClose: () => void }> = observer(({ streamIdIn, jobTab, changeListIn, onClose, autoloadIn }) => {
 
    type JobItem = {
       job: GetJobResponse
    };
 
-   const [searchState, setSearchState] = useState<{ items: JobItem[], groups?: IGroup[], templateId?: string, name?: string, userId?: string, global?: boolean, streamIdOverride?: string, preflights: boolean, preflightOnly?: boolean, minDate?: Date, maxDate?: Date, minCL?: string, maxCL?: string, querying: boolean, containsStep?: string, containsParameter?: string }>({ items: [], querying: false, preflights: false });
+   const [searchState, setSearchState] = useState<{ items: JobItem[], groups?: IGroup[], templates: Set<string>, name?: string, userId?: string, global?: boolean, streamIdOverride?: string, preflights: boolean, preflightOnly?: boolean, minDate?: Date, maxDate?: Date, minCL?: string, maxCL?: string, querying: boolean, containsStep?: string, containsParameter?: string }>({ items: [], querying: false, preflights: false, minCL: changeListIn, maxCL: changeListIn, templates: new Set(jobTab?.templates ?? []) });
    const [templateData, setTemplateData] = useState<{ streamId: string; templates: GetTemplateRefResponse[] } | undefined>(undefined);
+   const [controller] = useState(new CalloutController());
+   const [autoload, setAutoload] = useState<boolean | undefined>(undefined);
+
 
    const { hordeClasses, detailClasses } = getHordeStyling();
 
+   const streamId = searchState.streamIdOverride ?? streamIdIn;
+
    const stream = projectStore.streamById(streamId);
 
    if (!stream) {
@@ -98,31 +107,57 @@ export const JobSearchSimpleModal: React.FC<{ streamId: string, onClose: () => v
       }
 
       const subItems = templates.sort((a, b) => a.name < b.name ? -1 : 1).map(t => {
-         return { key: t.id, text: t.name, onClick: () => { setSearchState({ ...searchState, templateId: t.id }) } };
+         return {
+            key: t.id, text: t.name, onClick: (ev) => {
+               ev?.preventDefault();
+               if (searchState.templates.has(t.id)) {
+                  searchState.templates.delete(t.id);
+               } else {
+                  searchState.templates.add(t.id);
+               }
+               setSearchState({ ...searchState })
+               return false;
+            },
+            contextualMenuItemAs: (props: IContextualMenuItemProps) => {
+               return <Stack horizontal verticalFill verticalAlign="center">
+                  <Checkbox checked={searchState.templates.has(t.id)} />
+                  <Text>{props.item.text}</Text></Stack>
+            }
+         };
       })
 
       templateOptions.push({ key: `${cat}_category`, text: cat, subMenuProps: { items: subItems } });
 
    })
 
-   templateOptions.push({ key: `show_all_templates`, text: "All Templates", onClick: (ev) => { ev?.stopPropagation(); ev?.preventDefault(); setSearchState({ ...searchState, templateId: undefined, containsParameter: undefined }) } });
+   templateOptions.push({
+      key: `show_all_templates`, text: "All Templates", onClick: (ev) => {
+         ev?.stopPropagation();
+         ev?.preventDefault();
+         searchState.templates.clear();
+         setSearchState({ ...searchState, containsParameter: undefined })
+         return false;
+      }
+   });
 
    const templateMenuProps: IContextualMenuProps = {
       shouldFocusOnMount: true,
       subMenuHoverDelay: 0,
-      items: templateOptions,
+      items: templateOptions
    };
 
 
-   const template = templates.find(t => t.id === searchState.templateId);
+   const ctemplates = templates.filter(t => searchState.templates.size === 0 || searchState.templates.has(t.id));
 
    // results
+   const labelColumnWidth = 396;
    const columns = [
       { key: 'column1', name: 'Change', minWidth: 60, maxWidth: 60, isResizable: false },
-      { key: 'column2', name: 'Name', minWidth: 340, maxWidth: 340, isResizable: false },
-      { key: 'column3', name: 'Step', minWidth: 200, maxWidth: 200, isResizable: false },
-      { key: 'column4', name: 'StartedBy', minWidth: 120, maxWidth: 120, isResizable: false },
-      { key: 'column5', name: 'Created', minWidth: 100, maxWidth: 100, isResizable: false },
+      { key: 'column2', name: 'Name', minWidth: 140, maxWidth: 140, isResizable: false },
+      { key: 'column3', name: 'Labels', minWidth: labelColumnWidth, maxWidth: labelColumnWidth, isResizable: false },
+      //{ key: 'column4', name: 'Step', minWidth: 200, maxWidth: 200, isResizable: false },
+      { key: 'column5', name: 'StartedBy', minWidth: 120, maxWidth: 120, isResizable: false },
+      { key: 'column6', name: 'Created', minWidth: 100, maxWidth: 100, isResizable: false },
    ];
 
    const onRenderGroupHeader: IDetailsGroupRenderProps['onRenderHeader'] = (props) => {
@@ -141,7 +176,7 @@ export const JobSearchSimpleModal: React.FC<{ streamId: string, onClose: () => v
                      fontSize: "13px",
                      padding: '4px 8px',
                      userSelect: 'none',
-                     color: "#404040",
+                     color: dashboard.darktheme ? "#FFFFFFFF" : "#404040",
                      width: "100%",
                      fontFamily: "Horde Open Sans SemiBold"
                   }}>{`${group.name} (${group.count})`}</div>
@@ -153,6 +188,111 @@ export const JobSearchSimpleModal: React.FC<{ streamId: string, onClose: () => v
       return null;
    };
 
+   const JobLabel: React.FC<{ item: JobItem; label: GetLabelStateResponse }> = ({ item, label }) => {
+
+      const defaultLabel = item.job.defaultLabel;
+
+      const aggregates = item.job.labels;
+
+      // note details may not be loaded here, as only initialized on callout for optimization (details.getLabelIndex(label.Name, label.Category);)
+      const jlabel = aggregates?.find((l, idx) => l.dashboardCategory === label.dashboardCategory && l.dashboardName === label.dashboardName && item.job.labels![idx]?.state !== LabelState.Unspecified);
+      let labelIdx = -1;
+      if (jlabel) {
+         labelIdx = aggregates?.indexOf(jlabel)!;
+      }
+
+      let state: LabelState | undefined;
+      let outcome: LabelOutcome | undefined;
+      if (label === defaultLabel) {
+         state = defaultLabel.state;
+         outcome = defaultLabel.outcome;
+      } else {
+         state = item.job.labels![labelIdx]?.state;
+         outcome = item.job.labels![labelIdx]?.outcome;
+      }
+
+      const color = getLabelColor(state, outcome);
+
+      let url = `/job/${item.job.id}`;
+
+      if (labelIdx >= 0) {
+         url = `/job/${item.job.id}?label=${labelIdx}`;
+      }
+
+      const target = `label_search_${item.job.id}_${label.dashboardName}_${label.dashboardCategory}`.replace(/[^A-Za-z0-9]/g, "");
+
+      return <Stack>
+         <div id={target} className="horde-no-darktheme">
+            <Link to={url}><Stack className={hordeClasses.badgeNoIcon}><DefaultButton key={label.dashboardName} style={{ backgroundColor: color.primaryColor }} text={label.dashboardName}
+               onMouseOver={(ev) => {
+                  ev.stopPropagation();
+                  controller.setState({ target: `#${target}`, label: label, jobId: item.job.id })
+               }}
+               onMouseLeave={() => controller.setState(undefined)}
+               onMouseMove={(ev) => ev.stopPropagation()}>
+               {!!color.secondaryColor && <div style={{
+                  borderLeft: "10px solid transparent",
+                  borderRight: `10px solid ${color.secondaryColor}`,
+                  borderBottom: "10px solid transparent",
+                  height: 0,
+                  width: 0,
+                  position: "absolute",
+                  right: 0,
+                  top: 0,
+                  zIndex: 1
+               }} />}
+
+            </DefaultButton></Stack></Link>
+         </div>
+      </Stack>;
+
+   };
+
+   const renderLabels = (item: JobItem): JSX.Element | null => {
+
+      let labels = (item.job.labels ?? []).filter(l => !!l.dashboardName?.length && !!l.state && l.state !== LabelState.Unspecified);
+
+      if (item.job.defaultLabel) {
+         item.job.defaultLabel.dashboardName = "Other";
+         item.job.defaultLabel.dashboardCategory = "zzt";
+         labels.push(item.job.defaultLabel);
+      }
+
+      if (!labels.length) {
+         return null;
+      }
+
+      labels = labels.sort((a, b) => {
+
+         if (a.dashboardCategory === b.dashboardCategory) {
+            return a.dashboardName!.localeCompare(b.dashboardName!);
+         }
+
+         return (a.dashboardCategory ?? "").localeCompare(b.dashboardCategory ?? "");
+      })
+
+
+      const defaultLabel = item.job.defaultLabel;
+
+      if (!labels.length && !defaultLabel) {
+         return <div />;
+      }
+
+      let key = 0;
+
+      const buttons = labels.map(label => {
+         return <JobLabel key={`label_${item.job.id}_${label.dashboardCategory ?? ""}_${label.dashboardName ?? ""}_${key++}`} item={item} label={label} />;
+      });
+
+      return (<Stack horizontalAlign="center" onMouseMove={() => controller.setState(undefined)}>
+         <Stack wrap horizontal horizontalAlign="center" tokens={{ childrenGap: 4 }} styles={{ root: { paddingTop: 2, width: labelColumnWidth } }}>
+            {buttons}
+         </Stack>
+      </Stack>
+      );
+
+   };
+
 
    const renderItem = (item: JobItem, index?: number, column?: IColumn) => {
 
@@ -169,6 +309,10 @@ export const JobSearchSimpleModal: React.FC<{ streamId: string, onClose: () => v
          return <Stack horizontal verticalFill={true} verticalAlign="center" tokens={{ childrenGap: 0, padding: 0 }} style={{ paddingTop: 4 }} ><ChangeButton job={item.job} /></Stack>;
       }
 
+      if (column.name === "Labels") {
+         return <Stack horizontal verticalFill={true} verticalAlign="center" >{renderLabels(item)}</Stack>;
+      }
+
       if (column.name === "StartedBy") {
          let startedBy = item.job.startedByUserInfo?.name;
          if (!startedBy) {
@@ -232,8 +376,12 @@ export const JobSearchSimpleModal: React.FC<{ streamId: string, onClose: () => v
             maxCL = undefined;
          }
 
+         if (!minCL && maxCL) {
+            minCL = 0;
+         }
+
          if (minCL && !maxCL) {
-            maxCL = minCL;
+            maxCL = 999999999;
          }
 
          let minCreateTime = searchState.minDate;
@@ -241,12 +389,12 @@ export const JobSearchSimpleModal: React.FC<{ streamId: string, onClose: () => v
 
          if (minCreateTime && (!maxCreateTime || minCreateTime === maxCreateTime)) {
 
-            const maxTime = moment(minCreateTime);
-            maxTime.add(23, "hours");
-            maxTime.add(59, "minutes");
-            maxTime.add(59, "seconds");
+            //const maxTime = moment(minCreateTime);
+            //maxTime.add(23, "hours");
+            //maxTime.add(59, "minutes");
+            //maxTime.add(59, "seconds");
 
-            maxCreateTime = maxTime.toDate();
+            //maxCreateTime = maxTime.toDate();
 
          } else if (maxCreateTime) {
 
@@ -278,8 +426,8 @@ export const JobSearchSimpleModal: React.FC<{ streamId: string, onClose: () => v
 
             let stepTemplates: string[] = [];
 
-            if (template?.id) {
-               stepTemplates.push(template.id);
+            if (ctemplates?.length) {
+               stepTemplates.push(...ctemplates.map(t => t.id));
             } else {
                stepTemplates = templates.filter(t => !!t.id).map(t => t.id);
             }
@@ -315,10 +463,10 @@ export const JobSearchSimpleModal: React.FC<{ streamId: string, onClose: () => v
                id: jobIds.length ? jobIds : undefined,
                streamId: searchStreamId?.toLowerCase(),
                name: name,
-               filter: "id,name,change,createTime,streamId,preflightChange,startedByUserInfo,streamId,arguments",
+               filter: "id,labels,defaultLabel,name,change,createTime,streamId,preflightChange,startedByUserInfo,arguments",
                minChange: minCL,
                maxChange: maxCL,
-               template: (!searchState.global && template?.id && !searchState.containsStep) ? [template.id] : undefined,
+               template: (!searchState.global && ctemplates.length && !searchState.containsStep) ? ctemplates.map(t => t.id) : undefined,
                minCreateTime: minCreateTime?.toUTCString(),
                maxCreateTime: maxCreateTime?.toUTCString(),
                includePreflight: searchState.preflights || !!searchState.preflightOnly,
@@ -402,7 +550,7 @@ export const JobSearchSimpleModal: React.FC<{ streamId: string, onClose: () => v
 
          const commonSelectors = { ".ms-DetailsRow-cell": { "overflow": "visible", padding: 0 } };
 
-         props.styles = { ...props.styles, root: { selectors: { ...commonSelectors as any } } };
+         props.styles = { ...props.styles, root: { paddingTop: 8, paddingBottom: 8, selectors: { ...commonSelectors as any } } };
 
          return <Link to={url} onClick={(ev) => { if (!ev.ctrlKey) onClose() }}><div className="job-item"><DetailsRow {...props} /> </div></Link>;
 
@@ -415,11 +563,28 @@ export const JobSearchSimpleModal: React.FC<{ streamId: string, onClose: () => v
       title = `Find Jobs in All Streams`;
    }
    if (searchState.streamIdOverride) {
-      title = `Find Jobs in ${searchState.streamIdOverride}`;
+      let nstream = projectStore.streamById(searchState.streamIdOverride)
+      title = `Find Jobs in ${nstream?.fullname ?? searchState.streamIdOverride}`;
+   }
+
+   let templateText = "All Templates";
+   if (searchState.templates.size) {
+      templateText = ctemplates.map(t => t.name).join(", ");
+      if (templateText.length > 84) {
+         templateText = templateText.slice(0, 84);
+         templateText += "...";
+      }
+   }
+
+   if (autoloadIn && autoload === undefined) {
+      queryJobs();
+      setAutoload(autoloadIn);
+      return null;
    }
 
    return (<Modal isOpen={true} topOffsetFixed={true} styles={{ main: { padding: 8, width: 1400, height: '800px', backgroundColor: hordeTheme.horde.contentBackground, hasBeenOpened: false, top: "80px", position: "absolute" } }} className={hordeClasses.modal} onDismiss={() => { onClose() }}>
       {searchState.querying && <StatusModal text={"Finding Jobs"} />}
+      <JobLabelCallout controller={controller} />
       <Stack styles={{ root: { paddingTop: 8, paddingLeft: 24, paddingRight: 12, paddingBottom: 8 } }}>
          <Stack tokens={{ childrenGap: 12 }}>
             <Stack horizontal styles={{ root: { padding: 0 } }}>
@@ -440,19 +605,11 @@ export const JobSearchSimpleModal: React.FC<{ streamId: string, onClose: () => v
                   <Stack tokens={{ childrenGap: 12 }}>
 
                      <Stack>
-                        <Label>Template</Label>
-                        <DefaultButton disabled={searchState.global || !!searchState.streamIdOverride} style={{ width: 352, textAlign: "left" }} menuProps={templateMenuProps} text={template ? template.name : "All Templates"} />
+                        <Label>{searchState.templates.size > 1 ? "Templates" : "Template"}</Label>
+                        <DefaultButton disabled={searchState.global} style={{ width: 352, textAlign: "left" }} menuProps={templateMenuProps} text={templateText} />
                      </Stack>
 
-                     <Stack>
-                        <TextField value={searchState.name ?? ""} spellCheck={false} autoComplete="off" style={{ width: 350 }} label="Name" onChange={(event, newValue) => {
-                           setSearchState({ ...searchState, name: newValue })
-
-                        }} ></TextField>
-                     </Stack>
-
-
-                     <Stack horizontal tokens={{ childrenGap: 12 }} >
+                     {!searchState.global && <Stack horizontal tokens={{ childrenGap: 12 }} >
                         <TextField value={searchState.minCL ?? ""} style={{ width: 168 }} label="Min Changelist" onChange={(ev, newValue) => {
                            newValue = newValue?.trim();
                            if (!newValue) {
@@ -467,7 +624,7 @@ export const JobSearchSimpleModal: React.FC<{ streamId: string, onClose: () => v
                            }
                            setSearchState({ ...searchState, maxCL: newValue })
                         }} />
-                     </Stack>
+                     </Stack>}
 
 
                      <Stack horizontal tokens={{ childrenGap: 12 }} >
@@ -475,26 +632,15 @@ export const JobSearchSimpleModal: React.FC<{ streamId: string, onClose: () => v
                         <DatePicker value={searchState.maxDate} label="Max Date" onSelectDate={(date) => { setSearchState({ ...searchState, maxDate: date ? date : undefined }) }} style={{ width: 168 }} />
                      </Stack>
 
-                     <Stack>
-                        <TextField value={searchState.streamIdOverride ?? ""} spellCheck={false} autoComplete="off" style={{ width: 350 }} label="Stream Id" onChange={(event, newValue) => {
-                           const templateId = searchState.templateId;
+                     {!searchState.global && <Stack key={`stream_chooser_${streamChooserId}`}>
+                        <Label>Stream</Label>
+                        <StreamChooser defaultStreamId={streamId} allowAll={false} onChange={(streamId) => {
+                           const templates = searchState.templates;
                            const global = searchState.global;
-                           setSearchState({ ...searchState, streamIdOverride: newValue, global: newValue ? false : global, templateId: newValue ? undefined : templateId })
+                           setSearchState({ ...searchState, streamIdOverride: streamId, global: streamId ? false : global, templates: streamId ? new Set() : templates })
 
-                        }} ></TextField>
-                     </Stack>
-
-                     <Stack>
-                        <TextField disabled={!template} value={searchState.containsParameter ?? ""} spellCheck={false} autoComplete="off" style={{ width: 350 }} label="Contains Parameter" placeholder={!template ? "Please select a template to search parameters" : ""} onChange={(event, newValue) => {
-                           setSearchState({ ...searchState, containsParameter: newValue?.trim() ? newValue : undefined })
-                        }} ></TextField>
-                     </Stack>
-
-                     <Stack>
-                        <TextField value={searchState.containsStep ?? ""} spellCheck={false} autoComplete="off" style={{ width: 350 }} label="Contains Step" description="Note: Exact match is required" onChange={(event, newValue) => {
-                           setSearchState({ ...searchState, containsStep: newValue?.trim() ? newValue : undefined })
-                        }} ></TextField>
-                     </Stack>
+                        }} />
+                     </Stack>}
 
                      <Stack tokens={{ childrenGap: 4 }}>
                         <Label>Started By</Label>
@@ -503,6 +649,27 @@ export const JobSearchSimpleModal: React.FC<{ streamId: string, onClose: () => v
                         })} userIdHints={[]} noResultsFoundText="No users found" />
                      </Stack>
 
+                     {!searchState.global && <Stack>
+                        <Stack>
+                           <TextField disabled={!ctemplates.length} value={searchState.containsParameter ?? ""} spellCheck={false} autoComplete="off" style={{ width: 350 }} label="Contains Parameter" placeholder={!ctemplates.length ? "Please select a template to search parameters" : ""} onChange={(event, newValue) => {
+                              setSearchState({ ...searchState, containsParameter: newValue?.trim() ? newValue : undefined })
+                           }} ></TextField>
+                        </Stack>
+
+                        <Stack>
+                           <TextField value={searchState.containsStep ?? ""} spellCheck={false} autoComplete="off" style={{ width: 350 }} label="Contains Step" description="Note: Exact match is required" onChange={(event, newValue) => {
+                              setSearchState({ ...searchState, containsStep: newValue?.trim() ? newValue : undefined })
+                           }} ></TextField>
+                        </Stack>
+                     </Stack>}
+
+                     <Stack>
+                        <TextField value={searchState.name ?? ""} spellCheck={false} autoComplete="off" style={{ width: 350 }} label="Job Name" onChange={(event, newValue) => {
+                           setSearchState({ ...searchState, name: newValue })
+
+                        }} ></TextField>
+                     </Stack>
+
                      <Stack horizontal tokens={{ childrenGap: 18 }}>
                         <Stack verticalFill={true} verticalAlign="center">
                            <Checkbox checked={searchState.preflights} label="Include Preflights" onChange={(ev, checked) => {
@@ -520,9 +687,10 @@ export const JobSearchSimpleModal: React.FC<{ streamId: string, onClose: () => v
 
                      <Stack verticalFill={true} verticalAlign="center">
                         <Checkbox checked={searchState.global} label="Global Search" onChange={(ev, checked) => {
-                           const templateId = searchState.templateId;
+                           const templates = searchState.templates;
                            const streamIdOverride = searchState.streamIdOverride;
-                           setSearchState({ ...searchState, global: checked ? checked : false, templateId: checked ? undefined : templateId, streamIdOverride: checked ? undefined : streamIdOverride })
+                           streamChooserId++;
+                           setSearchState({ ...searchState, global: checked ? checked : false, templates: checked ? new Set() : templates, streamIdOverride: checked ? undefined : streamIdOverride })
                         }} />
                      </Stack>
 
@@ -531,7 +699,7 @@ export const JobSearchSimpleModal: React.FC<{ streamId: string, onClose: () => v
                            <PrimaryButton disabled={searchState.querying} text="Search" onClick={() => { queryJobs() }} />
                         </Stack>
                         <Stack style={{ paddingTop: 8, paddingRight: 24 }}>
-                           <DefaultButton disabled={searchState.querying} text="Clear" onClick={() => { setSearchState({ items: [], groups: [], querying: false, preflights: false, userId: searchState.userId }) }} />
+                           <DefaultButton disabled={searchState.querying} text="Clear" onClick={() => { setSearchState({ items: [], groups: [], querying: false, preflights: false, userId: searchState.userId, templates: new Set() }) }} />
                         </Stack>
                      </Stack>
                   </Stack>
@@ -548,7 +716,7 @@ export const JobSearchSimpleModal: React.FC<{ streamId: string, onClose: () => v
                      </Stack>}
 
                      {!!searchState.items.length && <Stack>
-                        <div style={{ overflowY: 'auto', overflowX: 'hidden', height: "670px" }} data-is-scrollable={true}>
+                        <div style={{ overflowY: 'auto', overflowX: 'hidden', height: "670px" }} data-is-scrollable={true} onScroll={() => { controller.setState(undefined, true) }}>
                            <Stack tokens={{ childrenGap: 12 }} style={{ paddingRight: 12 }}>
                               <DetailsList
                                  compact={true}
@@ -571,7 +739,6 @@ export const JobSearchSimpleModal: React.FC<{ streamId: string, onClose: () => v
                      </Stack>}
                   </Stack>
                </Stack>
-
             </Stack>
          </Stack>
       </Stack>
diff --git a/Engine/Source/Programs/Horde/HordeDashboard/src/components/StreamView.tsx b/Engine/Source/Programs/Horde/HordeDashboard/src/components/StreamView.tsx
--- a/Engine/Source/Programs/Horde/HordeDashboard/src/components/StreamView.tsx
+++ b/Engine/Source/Programs/Horde/HordeDashboard/src/components/StreamView.tsx
@@ -512,7 +512,7 @@ const StreamViewInner: React.FC = observer(() => {
                }
             }
          }} />}
-         {findJobsShown && <JobSearchSimpleModal onClose={() => { setFindJobsShown(false) }} streamId={stream.id} />}
+         {findJobsShown && <JobSearchSimpleModal onClose={() => { setFindJobsShown(false) }} streamIdIn={stream.id} jobTab={currentTab} />}
          <Stack horizontal style={{ backgroundColor: hordeTheme.horde.neutralBackground }}>
             <div key={`windowsize_streamview_${windowSize.width}_${windowSize.height}`} style={{ width: (vw / 2 - (1440 / 2)) - 12, flexShrink: 0, backgroundColor: modeColors.background }} />
             <Stack tokens={{ childrenGap: 0 }} styles={{ root: { width: "100%" } }}>
diff --git a/Engine/Source/Programs/Horde/HordeDashboard/src/components/jobDetailsV2/JobOperationsBar.tsx b/Engine/Source/Programs/Horde/HordeDashboard/src/components/jobDetailsV2/JobOperationsBar.tsx
--- a/Engine/Source/Programs/Horde/HordeDashboard/src/components/jobDetailsV2/JobOperationsBar.tsx
+++ b/Engine/Source/Programs/Horde/HordeDashboard/src/components/jobDetailsV2/JobOperationsBar.tsx
@@ -19,6 +19,7 @@ import { NewBuildV2 } from '../build/NewBuildV2';
 import { AbortJobModal } from './AbortJobModal';
 import { JobDetailsV2 } from './JobDetailsViewCommon';
 import { RetryStepsModal, StepRetryModal, StepRetryType } from './StepRetryModal';
+import { JobSearchSimpleModal } from '../JobSearchSimple';
 
 enum ParameterState {
    Hidden,
@@ -33,6 +34,7 @@ export const JobOperations: React.FC<{ jobDetails: JobDetailsV2 }> = observer(({
    const [abortShown, setAbortShown] = useState(false);
    const [editShown, setEditShown] = useState(false);
    const [retryStepsShown, setRetryStepsShown] = useState(false);
+   const [findJobsShown, setFindJobsShown] = useState(false);
 
    const [parametersState, setParametersState] = useState(query.get("newbuild") ? ParameterState.Clone : ParameterState.Hidden);
 
@@ -117,6 +119,14 @@ export const JobOperations: React.FC<{ jobDetails: JobDetailsV2 }> = observer(({
       onClick: () => { setParametersState(ParameterState.Parameters); }
    });
 
+   if (jobData.change) {
+      opsList.push({
+         key: 'jobops_relatedjobs',
+         text: `Related Jobs`,
+         iconProps: { iconName: "Search" },
+         onClick: () => { setFindJobsShown(true) }
+      });      
+   }
    opsList.push({
       key: 'jobops_edit',
       text: "Edit",
@@ -214,6 +224,7 @@ export const JobOperations: React.FC<{ jobDetails: JobDetailsV2 }> = observer(({
          }} />}
 
       <Stack horizontal>
+         {findJobsShown && <JobSearchSimpleModal changeListIn={jobData.change?.toString()} onClose={() => { setFindJobsShown(false) }} streamIdIn={jobData.streamId} autoloadIn={true} />}
          <Stack grow />
          <Stack horizontal tokens={{ childrenGap: 8 }}>
 
diff --git a/Engine/Source/Programs/Horde/HordeDashboard/src/styles/Styles.tsx b/Engine/Source/Programs/Horde/HordeDashboard/src/styles/Styles.tsx
--- a/Engine/Source/Programs/Horde/HordeDashboard/src/styles/Styles.tsx
+++ b/Engine/Source/Programs/Horde/HordeDashboard/src/styles/Styles.tsx
@@ -333,7 +333,7 @@ export const getHordeStyling = () => {
                filter: "brightness(1.15)"
             },
             '.ms-Button-label': {
-               fontSize: '9px',
+               fontSize: '9px !important',
                fontFamily: "Horde Open Sans SemiBold",
                whiteSpace: "nowrap"
             }
