Horde: Prospective fix for production build only search issue
#skipci

diff --git a/Engine/Source/Programs/Horde/HordeDashboard/src/components/JobSearchSimple.tsx b/Engine/Source/Programs/Horde/HordeDashboard/src/components/JobSearchSimple.tsx
--- a/Engine/Source/Programs/Horde/HordeDashboard/src/components/JobSearchSimple.tsx
+++ b/Engine/Source/Programs/Horde/HordeDashboard/src/components/JobSearchSimple.tsx
@@ -99,6 +99,12 @@ export const JobSearchSimpleModal: React.FC<{ streamIdIn: string, changeListIn?:
       })
    })
 
+   const renderMenuItem = (props: IContextualMenuItemProps) => {
+      return <Stack horizontal verticalFill verticalAlign="center">
+         <Checkbox checked={searchState.templates.has(props.item.data!.id)} />
+         <Text>{props.item.text}</Text></Stack>
+   };
+
    Array.from(sorted.keys()).sort((a, b) => a < b ? -1 : 1).forEach(cat => {
 
       const templates = sorted.get(cat);
@@ -108,7 +114,10 @@ export const JobSearchSimpleModal: React.FC<{ streamIdIn: string, changeListIn?:
 
       const subItems = templates.sort((a, b) => a.name < b.name ? -1 : 1).map(t => {
          return {
-            key: t.id, text: t.name, onClick: (ev) => {
+            key: t.id,
+            text: t.name,
+            data: t,
+            onClick: (ev) => {
                ev?.preventDefault();
                if (searchState.templates.has(t.id)) {
                   searchState.templates.delete(t.id);
@@ -118,11 +127,7 @@ export const JobSearchSimpleModal: React.FC<{ streamIdIn: string, changeListIn?:
                setSearchState({ ...searchState })
                return false;
             },
-            contextualMenuItemAs: (props: IContextualMenuItemProps) => {
-               return <Stack horizontal verticalFill verticalAlign="center">
-                  <Checkbox checked={searchState.templates.has(t.id)} />
-                  <Text>{props.item.text}</Text></Stack>
-            }
+            contextualMenuItemAs: renderMenuItem
          };
       })
 
@@ -604,6 +609,16 @@ export const JobSearchSimpleModal: React.FC<{ streamIdIn: string, changeListIn?:
                <Stack>
                   <Stack tokens={{ childrenGap: 12 }}>
 
+                     {!searchState.global && <Stack key={`stream_chooser_${streamChooserId}`}>
+                        <Label>Stream</Label>
+                        <StreamChooser defaultStreamId={streamId} width={352} allowAll={false} onChange={(streamId) => {
+                           const templates = searchState.templates;
+                           const global = searchState.global;
+                           setSearchState({ ...searchState, streamIdOverride: streamId, global: streamId ? false : global, templates: streamId ? new Set() : templates })
+
+                        }} />
+                     </Stack>}
+
                      <Stack>
                         <Label>{searchState.templates.size > 1 ? "Templates" : "Template"}</Label>
                         <DefaultButton disabled={searchState.global} style={{ width: 352, textAlign: "left" }} menuProps={templateMenuProps} text={templateText} />
@@ -632,16 +647,6 @@ export const JobSearchSimpleModal: React.FC<{ streamIdIn: string, changeListIn?:
                         <DatePicker value={searchState.maxDate} label="Max Date" onSelectDate={(date) => { setSearchState({ ...searchState, maxDate: date ? date : undefined }) }} style={{ width: 168 }} />
                      </Stack>
 
-                     {!searchState.global && <Stack key={`stream_chooser_${streamChooserId}`}>
-                        <Label>Stream</Label>
-                        <StreamChooser defaultStreamId={streamId} allowAll={false} onChange={(streamId) => {
-                           const templates = searchState.templates;
-                           const global = searchState.global;
-                           setSearchState({ ...searchState, streamIdOverride: streamId, global: streamId ? false : global, templates: streamId ? new Set() : templates })
-
-                        }} />
-                     </Stack>}
-
                      <Stack tokens={{ childrenGap: 4 }}>
                         <Label>Started By</Label>
                         <UserSelect handleSelection={(id => {
diff --git a/Engine/Source/Programs/Horde/HordeDashboard/src/components/projects/StreamChooser.tsx b/Engine/Source/Programs/Horde/HordeDashboard/src/components/projects/StreamChooser.tsx
--- a/Engine/Source/Programs/Horde/HordeDashboard/src/components/projects/StreamChooser.tsx
+++ b/Engine/Source/Programs/Horde/HordeDashboard/src/components/projects/StreamChooser.tsx
@@ -11,6 +11,7 @@ export interface IStreamChooser {
 export interface IStreamChooserProps {
    defaultStreamId?: string;
    allowAll?: boolean;
+   width?: number;
    onChange?: (streamId: string | undefined) => void;
 
 }
@@ -110,6 +111,6 @@ export const StreamChooser = React.forwardRef<IStreamChooser, IStreamChooserProp
    }
 
    return <div>
-      <DefaultButton style={{ width: 320, textAlign: "left" }} menuProps={templateMenuProps} text={text} />
+      <DefaultButton style={{ width: props.width ?? 320, textAlign: "left" }} menuProps={templateMenuProps} text={text} />
    </div>
 })
\ No newline at end of file
