Horde: Another fix for the production only search issue
#skipci

diff --git a/Engine/Source/Programs/Horde/HordeDashboard/src/components/JobSearchSimple.tsx b/Engine/Source/Programs/Horde/HordeDashboard/src/components/JobSearchSimple.tsx
--- a/Engine/Source/Programs/Horde/HordeDashboard/src/components/JobSearchSimple.tsx
+++ b/Engine/Source/Programs/Horde/HordeDashboard/src/components/JobSearchSimple.tsx
@@ -1,4 +1,4 @@
-import { Checkbox, DatePicker, DefaultButton, DetailsList, DetailsRow, IColumn, IconButton, IContextualMenuItem, IContextualMenuItemProps, IContextualMenuProps, IDetailsGroupRenderProps, IDetailsListProps, IGroup, Label, Modal, PrimaryButton, SelectionMode, Spinner, SpinnerSize, Stack, Text, TextField } from "@fluentui/react";
+import { Checkbox, ContextualMenuItemType, DatePicker, DefaultButton, DetailsList, DetailsRow, IColumn, IconButton, IContextualMenuItem, IContextualMenuItemProps, IContextualMenuProps, IDetailsGroupRenderProps, IDetailsListProps, IGroup, Label, Modal, PrimaryButton, SelectionMode, Spinner, SpinnerSize, Stack, Text, TextField } from "@fluentui/react";
 import { observer } from "mobx-react-lite";
 import moment from "moment";
 import { useState } from "react";
@@ -99,12 +99,6 @@ export const JobSearchSimpleModal: React.FC<{ streamIdIn: string, changeListIn?:
       })
    })
 
-   const renderMenuItem = (props: IContextualMenuItemProps) => {
-      return <Stack horizontal verticalFill verticalAlign="center">
-         <Checkbox checked={searchState.templates.has(props.item.data!.id)} />
-         <Text>{props.item.text}</Text></Stack>
-   };
-
    Array.from(sorted.keys()).sort((a, b) => a < b ? -1 : 1).forEach(cat => {
 
       const templates = sorted.get(cat);
@@ -112,22 +106,27 @@ export const JobSearchSimpleModal: React.FC<{ streamIdIn: string, changeListIn?:
          return;
       }
 
-      const subItems = templates.sort((a, b) => a.name < b.name ? -1 : 1).map(t => {
+      const subItems: IContextualMenuItem[] = templates.sort((a, b) => a.name < b.name ? -1 : 1).map(t => {
          return {
+            itemType: ContextualMenuItemType.Normal,
             key: t.id,
             text: t.name,
             data: t,
-            onClick: (ev) => {
-               ev?.preventDefault();
-               if (searchState.templates.has(t.id)) {
-                  searchState.templates.delete(t.id);
-               } else {
-                  searchState.templates.add(t.id);
-               }
-               setSearchState({ ...searchState })
-               return false;
-            },
-            contextualMenuItemAs: renderMenuItem
+            onRender: (props: IContextualMenuItemProps) => {
+               return <Stack horizontal verticalFill verticalAlign="center" style={{ paddingTop: 4, paddingBottom: 4, cursor: "pointer" }} onClick={(ev) => {
+                  ev?.preventDefault();
+                  ev?.stopPropagation();
+                  if (searchState.templates.has(t.id)) {
+                     searchState.templates.delete(t.id);
+                  } else {
+                     searchState.templates.add(t.id);
+                  }
+                  setSearchState({ ...searchState })
+
+               }} >
+                  <Checkbox checked={searchState.templates.has(t.id)} />
+                  <Text>{t.name}</Text></Stack>
+            }
          };
       })
 
