Horde: Address padding on context menu
#skipci

diff --git a/Engine/Source/Programs/Horde/HordeDashboard/src/components/JobSearchSimple.tsx b/Engine/Source/Programs/Horde/HordeDashboard/src/components/JobSearchSimple.tsx
--- a/Engine/Source/Programs/Horde/HordeDashboard/src/components/JobSearchSimple.tsx
+++ b/Engine/Source/Programs/Horde/HordeDashboard/src/components/JobSearchSimple.tsx
@@ -113,7 +113,7 @@ export const JobSearchSimpleModal: React.FC<{ streamIdIn: string, changeListIn?:
             text: t.name,
             data: t,
             onRender: (props: IContextualMenuItemProps) => {
-               return <Stack horizontal verticalFill verticalAlign="center" style={{ paddingTop: 4, paddingBottom: 4, cursor: "pointer" }} onClick={(ev) => {
+               return <Stack horizontal verticalFill verticalAlign="center" style={{ padding: "4px 12px 4px 12px", cursor: "pointer" }} onClick={(ev) => {
                   ev?.preventDefault();
                   ev?.stopPropagation();
                   if (searchState.templates.has(t.id)) {
