Horde: Fix deadlock in RPC session test.

#rnx

diff --git a/Engine/Source/Programs/Horde/Plugins/Compute/HordeServer.Compute/Server/RpcService.cs b/Engine/Source/Programs/Horde/Plugins/Compute/HordeServer.Compute/Server/RpcService.cs
--- a/Engine/Source/Programs/Horde/Plugins/Compute/HordeServer.Compute/Server/RpcService.cs
+++ b/Engine/Source/Programs/Horde/Plugins/Compute/HordeServer.Compute/Server/RpcService.cs
@@ -263,7 +263,7 @@ namespace HordeServer.Server
 
 				// Get a task for moving to the next item. This will only complete once the call has closed.
 				using CancellationTokenSource cancellationSource = CancellationTokenSource.CreateLinkedTokenSource(context.CancellationToken);
-				nextRequestTask = reader.MoveNext();
+				nextRequestTask = reader.MoveNext(cancellationSource.Token);
 				try
 				{
 					nextRequestTask = nextRequestTask.ContinueWith(task =>
@@ -319,15 +319,13 @@ namespace HordeServer.Server
 						response.Status = (RpcAgentStatus)agent.Status;
 						await writer.WriteAsync(response);
 					}
-				}
-				finally
-				{
+
 					// Wait for the client to close the stream
 					try
 					{
 						while (await nextRequestTask)
 						{
-							nextRequestTask = reader.MoveNext();
+							nextRequestTask = reader.MoveNext(cancellationSource.Token);
 						}
 					}
 					catch (Exception ex)
@@ -335,6 +333,18 @@ namespace HordeServer.Server
 						_logger.LogDebug(ex, "Ignoring exception while finishing UpdateSession request.");
 					}
 				}
+				finally
+				{
+					cancellationSource.Cancel();
+					try
+					{
+						await nextRequestTask;
+					}
+					catch (OperationCanceledException)
+					{
+						// Ignore
+					}
+				}
 			}
 		}
 
