Horde: Fix static analysis warning.

#rnx

diff --git a/Engine/Source/Programs/Horde/Plugins/Compute/HordeServer.Compute/Server/RpcService.cs b/Engine/Source/Programs/Horde/Plugins/Compute/HordeServer.Compute/Server/RpcService.cs
--- a/Engine/Source/Programs/Horde/Plugins/Compute/HordeServer.Compute/Server/RpcService.cs
+++ b/Engine/Source/Programs/Horde/Plugins/Compute/HordeServer.Compute/Server/RpcService.cs
@@ -335,7 +335,8 @@ namespace HordeServer.Server
 				}
 				finally
 				{
-					cancellationSource.Cancel();
+					await cancellationSource.CancelAsync();
+
 					try
 					{
 						await nextRequestTask;
