Updating device views
- Add tooltip to history icon
- Allow to add notes while changing device status
- Add failure rate explict number at the top of the device reservation history
- Auto fill telemetry dropdown based on filters from the main device view
- Fix telemetry graph zoom
- Use step curve instead of straight line to represent reservation
- Add average load to telemetry graph
- Filter telemetry summary based on graph selection
- Filter reservable failing devices in summary
- Add link to device for quicker action

#rnx

diff --git a/Engine/Source/Programs/Horde/HordeDashboard/src/components/DeviceEditor.tsx b/Engine/Source/Programs/Horde/HordeDashboard/src/components/DeviceEditor.tsx
--- a/Engine/Source/Programs/Horde/HordeDashboard/src/components/DeviceEditor.tsx
+++ b/Engine/Source/Programs/Horde/HordeDashboard/src/components/DeviceEditor.tsx
@@ -483,9 +483,9 @@ export const DeviceEditor: React.FC<{ handler: DeviceHandler, deviceIn?: GetDevi
 
             }} />}
 
-            {!!editNote && <TextField label="Notes" defaultValue={device.notes} multiline rows={5} resizable={false} onChange={(ev, newValue) => {
+            <TextField label="Notes" defaultValue={device.notes} multiline rows={5} resizable={false} onChange={(ev, newValue) => {
                device.notes = newValue;
-            }} />}
+            }} />
 
             {!!reservation && !editNote && <Stack style={{ paddingTop: 24 }}><Checkbox label="Disable device for maintenance once reservation finishes"
                checked={!!device.maintenance}
diff --git a/Engine/Source/Programs/Horde/HordeDashboard/src/components/DeviceInfoView.tsx b/Engine/Source/Programs/Horde/HordeDashboard/src/components/DeviceInfoView.tsx
--- a/Engine/Source/Programs/Horde/HordeDashboard/src/components/DeviceInfoView.tsx
+++ b/Engine/Source/Programs/Horde/HordeDashboard/src/components/DeviceInfoView.tsx
@@ -5,14 +5,31 @@ import { useState } from "react";
 import { Link } from "react-router-dom";
 import backend from "../backend";
 import { DeviceStatus, DeviceTelemetryQuery, GetDeviceResponse, GetDeviceTelemetryResponse, GetDeviceUtilizationResponse, GetTelemetryInfoResponse, JobData } from "../backend/Api";
-import dashboard from "../backend/Dashboard";
+import dashboard, { StatusColor } from "../backend/Dashboard";
 import { projectStore } from "../backend/ProjectStore";
 import { displayTimeZone } from "../base/utilities/timeUtils";
 import { getHordeStyling } from "../styles/Styles";
 import { DeviceHandler } from "./DeviceEditor";
 import { DeviceStatusIcon } from "./StatusIcon";
 
-const streamIdToFullname = new Map<string, string>();;
+const streamIdToFullname = new Map<string, string>();
+
+const RateBar = (rate: number, width: number, height: number): JSX.Element => {
+   const colors = dashboard.getStatusColors();
+   return (
+      <Stack horizontal>
+         <div style={{ backgroundColor: colors.get(StatusColor.Success), width: width, height: height, verticalAlign: 'middle', borderRadius: '6px', marginRight: '3px'}}>
+            <span
+               style={{
+                  width: `${rate}%`, height: '100%',
+                  backgroundColor: colors.get(StatusColor.Failure),
+                  display: 'inline-block', borderRadius: 'inherit'
+               }} />
+         </div>
+         <Label>{rate}% Failure Rate</Label>
+      </Stack>
+   );
+}
 
 export const DeviceInfoModal: React.FC<{ handler: DeviceHandler, deviceIn?: GetDeviceResponse | undefined, onEdit: (device: GetDeviceResponse) => void, onClose: () => void }> = observer(({ handler, deviceIn, onEdit, onClose }) => {
 
@@ -232,6 +249,7 @@ export const DeviceInfoModal: React.FC<{ handler: DeviceHandler, deviceIn?: GetD
       return <Stack />;
    }
 
+   let failureRate = jobState.telemetryItems.length == 0? 0 : Math.floor(jobState.telemetryItems.filter(item => !!item.job.problemTimeUtc).length / jobState.telemetryItems.length * 100);
 
    return (<Modal isOpen={true} styles={{ main: { padding: 8, width: 1140, height: '800px', backgroundColor: '#FFFFFF' } }} className={hordeClasses.modal} onDismiss={() => { onClose() }}>
       <Stack styles={{ root: { paddingTop: 8, paddingLeft: 24, paddingRight: 12, paddingBottom: 8 } }}>
@@ -284,7 +302,8 @@ export const DeviceInfoModal: React.FC<{ handler: DeviceHandler, deviceIn?: GetD
                <Stack>
                   <Stack styles={{ root: { paddingLeft: 4, paddingRight: 0, paddingBottom: 4 } }}>
                      <Stack>
-                        <Label>Device Issues</Label>
+                        <Label>Device Reservations History</Label>
+                        {!!failureRate && RateBar(failureRate, 70, 15)}
                      </Stack>
 
                      {!jobState.telemetryItems.length && <Stack>
diff --git a/Engine/Source/Programs/Horde/HordeDashboard/src/components/DevicePoolTelemetry.tsx b/Engine/Source/Programs/Horde/HordeDashboard/src/components/DevicePoolTelemetry.tsx
--- a/Engine/Source/Programs/Horde/HordeDashboard/src/components/DevicePoolTelemetry.tsx
+++ b/Engine/Source/Programs/Horde/HordeDashboard/src/components/DevicePoolTelemetry.tsx
@@ -71,7 +71,7 @@ type ProblemStep = {
    stepName?: string;
 }
 
-type DeviceProblem = { deviceId: string, deviceName: string, problems: number, problemsDesc: string, latestProblem?: DeviceTelemetry };
+type DeviceProblem = { deviceId: string, deviceName: string, poolName: string, problems: number, problemsDesc: string, latestProblem?: DeviceTelemetry, reservable: boolean };
 
 class PoolTelemetryHandler {
 
@@ -214,7 +214,11 @@ class PoolTelemetryHandler {
             platforms.set(platform.id, []);
          }
 
-         platforms.get(platform.id)!.push({ deviceId: deviceId, deviceName: `${device.name} / ${device.address}`, problems: count, problemsDesc: `${count} / ${telemetry.length}`, latestProblem: latest });
+         const poolName = this.pools.get(device.poolId)!.name;
+         const ratePercent = Math.floor(count / telemetry.length * 100);
+         const reservable = device.enabled && !device.maintenanceTime;
+
+         platforms.get(platform.id)!.push({ deviceId: deviceId, deviceName: device.name, poolName: poolName,problems: count, problemsDesc: `${count} / ${telemetry.length} (${ratePercent}%)`, latestProblem: latest, reservable: reservable });
 
       });
 
@@ -560,6 +564,7 @@ class PoolTelemetryGraph {
       const { modeColors } = getHordeStyling();
 
       const xPlot = this.xPlot = this.xPlotMethod();
+      const xPlotAxis = this.xPlotMethod();
 
       const totalDevices = data.deviceIds.length;
 
@@ -597,42 +602,42 @@ class PoolTelemetryGraph {
          .attr("width", width - margin.left - margin.right - 240)
          .attr("height", height - margin.top - margin.bottom);
 
-      const updateChart = (event: any) => {
+      const zoom = d3.zoom()
+         .scaleExtent([1, 12])
+         .extent([[margin.left, 0], [width - margin.right - 240, height]])
+         .translateExtent([[margin.left, -Infinity], [width - margin.right - 240, Infinity]])
+         .on("zoom", zoomed);
 
-         if (this.freezeTooltip || (event && !event.sourceEvent)) {
-            return;
-         }
+      const graph = this;
 
-         const extent = event?.selection
-         this.extent = extent;
-
-         if (!extent) {
-            xPlot!.domain(d3.extent(this.data!.telemetry, function (d: any) { return d.date; }) as any)
-         } else {
-            xPlot!.domain([xPlot!.invert(extent[0]), xPlot!.invert(extent[1])])
-            areaChart.select(".brush").call((brush as any).move, null)
+      function zoomed(event: any) {
+         if (graph.freezeTooltip) {
+            return;
          }
 
-         this.xAxis!.transition().duration(500).call(this.xAxisMethod as any, xPlot)
+         xAxis!.call(graph.xAxisMethod as any, event.transform.rescaleX(xPlotAxis));
+         xPlot!.range([margin.left, width - margin.right - 240].map(d => event.transform.applyX(d)));
          areaChart
             .selectAll("path")
-            .transition().duration(500)
             .attr("d", area as any)
       }
 
+      svg.call(zoom as any);
+
       // left axis
       svg.append("g")
          .attr("transform", `translate(${this.margin.left},${margin.bottom - 8})`)
          .call(d3.axisLeft(y).ticks(totalDevices))
 
       // bottom axis      
-      this.xAxis = svg.append("g")
+      const xAxis = this.xAxis = svg.append("g")
          .call(this.xAxisMethod as any, xPlot)
 
       const area = d3.area()
          .x((d: any) => xPlot!(d.data.date))
          .y0((d) => y(d[0]) + 32)
          .y1((d) => y(d[1]) + 32)
+         .curve(d3.curveStepAfter);
 
       const keys = ["Disabled", "Maintenance", "Problem"];
 
@@ -653,18 +658,9 @@ class PoolTelemetryGraph {
          d3.selectAll(".telemetryArea").style("opacity", 1)
       }
 
-      const brush = d3.brushX()
-         .extent([[0, 0], [width - 270, height]])
-         .on("end", updateChart)
-
       const areaChart = svg.append('g')
          .attr("clip-path", `url(#${this.clipId})`);
 
-      areaChart
-         .append("g")
-         .attr("class", "brush")
-         .call(brush);
-
       areaChart
          .selectAll("telemetrylayers")
          .data(stackedData)
@@ -723,15 +719,9 @@ class PoolTelemetryGraph {
          .style("position", "absolute")
       //.style("pointer-events", "none")
 
-      svg.on("mousemove", (event) => this.handleMouseMove(event))
-      svg.on("wheel", (event:any) => {
-         if (event.wheelDelta < 0) {
-            updateChart(undefined);
-         }
-      })
+      svg.on("mousemove", (event) => this.handleMouseMove(event));
       // dblclick for double click events
       svg.on("click", (event) => {
-
          if (this.freezeTooltip) {
             this.freezeTooltip = false;
             this.updateTooltip(false);
@@ -744,7 +734,36 @@ class PoolTelemetryGraph {
 
       });
 
-
+      // Average load
+      const avgLoad = svg.append("g")
+         .attr("transform", "translate(8,6)");
+      const colorScale = d3.scaleLinear<string>()
+         .domain([10, 50, 80])
+         .range([scolors.get(StatusColor.Success)!, scolors.get(StatusColor.Warnings)!, scolors.get(StatusColor.Failure)!]);
+      const avgLoadWidth = 100;
+      const load = totalDevices > 0 ? Math.floor(data.telemetry.reduce((acru, item) => acru + totalDevices - (item.available?.length ?? 0), 0) / (data.telemetry.length * totalDevices) * 100) : 0;
+      avgLoad.append("rect")
+         .attr("width", avgLoadWidth)
+         .attr("height", 20)
+         .style("fill", "rgba(0,0,0,0)")
+         .style("stroke", function (d: any) { return modeColors.text })
+         .style("stroke-width", 1);
+      avgLoad.append("text")
+         .attr("x", avgLoadWidth + 6)
+         .attr("y", 10)
+         .text(`${load}% Avg Load`)
+         .style("fill", function (d: any) { return modeColors.text })
+         .attr("text-anchor", "left")
+         .style("alignment-baseline", "middle")
+         .style("font-family", "Horde Open Sans Regular")
+         .style("font-size", "13px");
+      avgLoad.selectAll("averageLoad")
+         .data([load])
+         .enter()
+         .append("rect")
+         .attr("width", (d) => Math.floor(d / 100 * (avgLoadWidth-1)))
+         .attr("height", 19)
+         .style("fill", (d) => colorScale(d));
    }
 
    private clear() {
@@ -911,69 +930,24 @@ class PoolTelemetryGraph {
 
    private get xAxisMethod() {
       return (g: SelectionType, x: ScalarTime) => {
+         const tickSet: Set<string> = new Set();
+
+         g.attr("transform", `translate(0,${this.height - this.margin.bottom})`)
+            .style("font-family", "Horde Open Sans Regular")
+            .style("font-size", "11px")
+            .call(d3.axisBottom(x)
+               .ticks(8)
+               .tickFormat(d => {
+                  const v = (d as Date).toDateString();
+                  if (!tickSet.has(v)) {
+                     tickSet.add(v)
+                     return getHumanTime(d as Date)
+                  }
 
-         const data = this.data!;
-
-         const xPlot = this.xPlot;
-
-         const tickValues: { x: number, date: Date }[] = [];
-         const d = x.domain();
-
-         if (d.length === 2) {
-
-            const start = d[0];
-            const end = d[1];
-            const tickSet: Set<string> = new Set();
-
-            data.telemetry.forEach(d => {
-               if (d.date < start || d.date > end) {
-                  return;
-               }
-               const v = d.dateString;
-               if (tickSet.has(v)) {
-                  return;
-               }
-               tickSet.add(v);
-               const date = new Date(v);
-               const x = xPlot!(date);
-               if (true) { //x >= 0 ) {
-                  tickValues.push({ x: x, date: date });
-               }
-            });
-         }
-
-         //const min = d[0];
-         //const max = d[1];
-         //const span = moment.duration(moment(max).diff(moment(min))).asDays();
-         if (true) {//span <= 6) {
-
-            const tickSet: Set<string> = new Set();
-
-            g.attr("transform", `translate(0,${this.height - this.margin.bottom})`)
-               .style("font-family", "Horde Open Sans Regular")
-               .style("font-size", "11px")
-               .call(d3.axisBottom(x)
-                  .ticks(8)
-                  .tickFormat(d => {
-
-                     const v = (d as Date).toDateString();
-                     if (!tickSet.has(v)) {
-                        tickSet.add(v)
-                        return getHumanTime(d as Date)
-                     }
-
-                     return (d as Date).toLocaleString(undefined, { hour: 'numeric', minute: "numeric", hour12: true })
-                  })
-                  .tickSizeOuter(0))
-
-            return;
-
-         }
-
+                  return (d as Date).toLocaleString(undefined, { hour: 'numeric', minute: "numeric", hour12: true })
+               })
+               .tickSizeOuter(0))
       }
-
-
-
    }
 
    private xPlot?: ScalarTime;
@@ -1017,15 +991,7 @@ const SummaryModal: React.FC<{ handler: PoolTelemetryHandler, selectedPlatform?:
 
    const { hordeClasses, modeColors } = getHordeStyling();
 
-   let platforms = Array.from(handler.platforms.values()).sort((a, b) => {
-      if (a.id === selectedPlatform?.id) {
-         return -1;
-      }
-      if (b.id === selectedPlatform?.id) {
-         return 1;
-      }
-      return a.name.localeCompare(b.name)
-   });
+   let platforms = Array.from(handler.platforms.values()).filter((p) => p.id === selectedPlatform?.id);
 
    const problems = handler.getProblemDevices();
 
@@ -1033,16 +999,16 @@ const SummaryModal: React.FC<{ handler: PoolTelemetryHandler, selectedPlatform?:
 
       const psteps = handler.getSteps(platform.id);
       let platformProblems = problems.get(platform.id) ?? [];
-      platformProblems = platformProblems.filter(p => p.problems > 8).sort((a, b) => b.latestProblem!.problemTime!.getTime() - a.latestProblem!.problemTime!.getTime());
+      platformProblems = platformProblems.filter(p => p.problems > 8 && p.reservable).sort((a, b) => b.problems > a.problems ? 1 : -1);
 
       const problemColumns: IColumn[] = [
-         { key: 'column1', name: 'Problem Devices', fieldName: 'deviceName', minWidth: 320, maxWidth: 320 },
-         { key: 'column2', name: 'Issues/Reservations', fieldName: 'problemsDesc', minWidth: 150, maxWidth: 150 },
-         { key: 'column3', name: 'Latest Issue', minWidth: 640, maxWidth: 640 },
+         { key: 'column1', name: 'Problem Devices', minWidth: 60, maxWidth: 150 },
+         { key: 'column2', name: 'Pool', fieldName: 'poolName', minWidth: 60, maxWidth: 80 },
+         { key: 'column3', name: 'Issues/Reservations', fieldName: 'problemsDesc', minWidth: 100, maxWidth: 150 },
+         { key: 'column4', name: 'Latest Issue', minWidth: 300, maxWidth: 640 },
       ];
 
       let column = problemColumns.find(c => c.name === "Latest Issue")!;
-
       column.onRender = ((item: DeviceProblem) => {
 
          const latest = item.latestProblem;
@@ -1066,6 +1032,15 @@ const SummaryModal: React.FC<{ handler: PoolTelemetryHandler, selectedPlatform?:
          </Stack>
       });
 
+      column = problemColumns.find(c => c.name === "Problem Devices")!;
+      column.onRender = ((item: DeviceProblem) => {
+         const deviceName = item.deviceName;
+         const url = `/devices?pivotKey=pivot-key-automation&filter=${deviceName}`;
+         return <Stack verticalFill={true} verticalAlign="center">
+            <Link to={url} target="_blank"><Text variant="small">{deviceName}</Text></Link>
+         </Stack>
+      });
+
       const steps: Map<string, { stepName: string, reservations: number, problems: number, durationMS: number }> = new Map();
 
       psteps.forEach(s => {
@@ -1164,7 +1139,7 @@ const SummaryModal: React.FC<{ handler: PoolTelemetryHandler, selectedPlatform?:
 
 }
 
-export const DevicePoolGraph: React.FC = observer(() => {
+export const DevicePoolGraph: React.FC<{ poolIn?: string, platformIn?: string }> = observer(({ poolIn, platformIn }) => {
 
    const [container, setContainer] = useState<HTMLDivElement | null>(null);
    const [showSummary, setShowSummary] = useState(false);
@@ -1200,8 +1175,8 @@ export const DevicePoolGraph: React.FC = observer(() => {
       }
 
 
-      const pool = pools[0];
-      const platform = platforms[0];
+      const pool = pools.find((e) => e.id == poolIn) ?? pools[0];
+      const platform = platforms.find((e) => e.id == platformIn) ?? platforms[0];
 
       setState({ ...state, pool: pool, platform: platform })
       return null;
@@ -1234,24 +1209,24 @@ export const DevicePoolGraph: React.FC = observer(() => {
             <Stack grow />
             <DefaultButton text="Summary" onClick={() => setShowSummary(true)} />
             <Dropdown
-               defaultSelectedKey={`pool_key_${state.pool!.id}`}
-               options={poolOptions}
+               defaultSelectedKey={`platform_key_${state.platform!.id}`}
+               options={platformOptions}
                onChange={(event, option) => {
                   if (state.graph) {
                      state.graph.forceRender = true;
                   }
-                  setState({ ...state, pool: option!.data, platform: platforms[0] });
+                  setState({ ...state, platform: option!.data, })
                }}
                style={{ width: 120 }}
             />
             <Dropdown
-               defaultSelectedKey={`platform_key_${state.platform!.id}`}
-               options={platformOptions}
+               defaultSelectedKey={`pool_key_${state.pool!.id}`}
+               options={poolOptions}
                onChange={(event, option) => {
                   if (state.graph) {
                      state.graph.forceRender = true;
                   }
-                  setState({ ...state, platform: option!.data, })
+                  setState({ ...state, pool: option!.data, });
                }}
                style={{ width: 120 }}
             />
@@ -1265,7 +1240,7 @@ export const DevicePoolGraph: React.FC = observer(() => {
    </Stack>
 });
 
-export const DevicePoolTelemetryModal: React.FC<{ onClose?: () => void }> = ({ onClose }) => {
+export const DevicePoolTelemetryModal: React.FC<{ poolsIn?: Array<string>, platformsIn?: Array<string>, onClose?: () => void }> = ({ poolsIn, platformsIn, onClose }) => {
 
    const { hordeClasses, modeColors } = getHordeStyling();
 
@@ -1283,7 +1258,7 @@ export const DevicePoolTelemetryModal: React.FC<{ onClose?: () => void }> = ({ o
                </Stack>
             </Stack>
             <Stack style={{ paddingLeft: 24 }}>
-               <DevicePoolGraph />
+               <DevicePoolGraph poolIn={poolsIn?.find((e) => true)} platformIn={platformsIn?.find((e) => true)} />
             </Stack>
          </Stack>
       </Modal>
diff --git a/Engine/Source/Programs/Horde/HordeDashboard/src/components/DeviceView.tsx b/Engine/Source/Programs/Horde/HordeDashboard/src/components/DeviceView.tsx
--- a/Engine/Source/Programs/Horde/HordeDashboard/src/components/DeviceView.tsx
+++ b/Engine/Source/Programs/Horde/HordeDashboard/src/components/DeviceView.tsx
@@ -1,4 +1,4 @@
-import { DefaultButton, DetailsList, DetailsListLayoutMode, Dropdown, FocusZone, FocusZoneDirection, IColumn, Icon, IconButton, IGroup, Label, Link as FluentLink, mergeStyleSets, Modal, Pivot, PivotItem, PrimaryButton, ScrollablePane, ScrollbarVisibility, SelectionMode, Spinner, SpinnerSize, Stack, Text, TextField, themeRulesStandardCreator, } from "@fluentui/react";
+import { DefaultButton, DetailsList, DetailsListLayoutMode, Dropdown, FocusZone, FocusZoneDirection, IColumn, Icon, IconButton, IGroup, Label, Link as FluentLink, mergeStyleSets, Modal, Pivot, PivotItem, PrimaryButton, ScrollablePane, ScrollbarVisibility, SelectionMode, Spinner, SpinnerSize, Stack, Text, TextField, themeRulesStandardCreator, TooltipHost, DirectionalHint, } from "@fluentui/react";
 import { observer } from "mobx-react-lite";
 import { observable, action, makeObservable } from "mobx";
 import moment from "moment";
@@ -524,17 +524,20 @@ const DevicePanel: React.FC = observer(() => {
       return <Stack verticalFill={true} verticalAlign="center">
          <Stack horizontal tokens={{ childrenGap: 12 }}>
             {automationTab && <Stack verticalFill={true} verticalAlign="center">
-               <div style={{ cursor: "pointer" }} onClick={(ev) => {
-                  ev.stopPropagation(); ev.preventDefault();
-                  setEditState({
-                     infoShown: true, device: item
-                  })
-                  localState.setHistoryItem(item.device.id);
-               }}>
-                  <Stack horizontal tokens={{ childrenGap: 18 }}>
-                     <Icon style={{ paddingTop: 2 }} iconName="History" />
-                  </Stack>
-               </div>
+               <TooltipHost content="History" delay={0} closeDelay={250}
+                  styles={{root: {display: 'inline-block'}}} directionalHint={DirectionalHint.leftCenter}>
+                  <div style={{ cursor: "pointer" }} onClick={(ev) => {
+                     ev.stopPropagation(); ev.preventDefault();
+                     setEditState({
+                        infoShown: true, device: item
+                     })
+                     localState.setHistoryItem(item.device.id);
+                  }}>
+                     <Stack horizontal tokens={{ childrenGap: 18 }}>
+                        <Icon style={{ padding: 2 }} iconName="History"/>
+                     </Stack>
+                  </div>
+               </TooltipHost>
 
             </Stack>}
             {!!r && <Stack>
@@ -771,7 +774,7 @@ const DevicePanel: React.FC = observer(() => {
    const hasDevices = !!handler.getDevices().length
 
    return (<Stack>
-      {!!telemState && <DevicePoolTelemetryModal onClose={() => setTelemState(false)} />}
+      {!!telemState && <DevicePoolTelemetryModal poolsIn={localState.searchState.filterPools} platformsIn={localState.searchState.filterPlatforms} onClose={() => setTelemState(false)} />}
       {(!!checkoutState.checkinId || !!checkoutState.checkoutId) && <InfoModal />}
       {!!checkoutState.showConfirm && checkedOut.length > 0 && <CheckoutConfirmModal check={checkoutState.showConfirm} devices={checkedOut} onClose={() => setCheckoutState({})} />}
       {editState.infoShown && <DeviceInfoModal handler={handler} deviceIn={editState.device?.device} onEdit={(device) => { setEditState({ infoShown: true, shown: true, device: { device: device } }) }} onClose={() => { setEditState({}); localState.resetHistoryItem(); }} />}
