Horde: Add interfaces for template parameters rather than using annotated MongoDB model classes. Will allow moving into EpicGames.Horde in future.

#rnx

diff --git a/Engine/Source/Programs/Horde/Plugins/Build/HordeServer.Build.Tests/Fixture.cs b/Engine/Source/Programs/Horde/Plugins/Build/HordeServer.Build.Tests/Fixture.cs
--- a/Engine/Source/Programs/Horde/Plugins/Build/HordeServer.Build.Tests/Fixture.cs
+++ b/Engine/Source/Programs/Horde/Plugins/Build/HordeServer.Build.Tests/Fixture.cs
@@ -138,6 +138,6 @@ namespace HordeServer.Tests
 		public string? SubmitNewChange { get; init; } = null;
 		public string? SubmitDescription { get; init; } = null;
 		public IReadOnlyList<string> Arguments { get; init; } = [];
-		public IReadOnlyList<Parameter> Parameters { get; init; } = [];
+		public IReadOnlyList<ITemplateParameter> Parameters { get; init; } = [];
 	}
 }
\ No newline at end of file
diff --git a/Engine/Source/Programs/Horde/Plugins/Build/HordeServer.Build.Tests/Jobs/Schedules/SchedulerTests.cs b/Engine/Source/Programs/Horde/Plugins/Build/HordeServer.Build.Tests/Jobs/Schedules/SchedulerTests.cs
--- a/Engine/Source/Programs/Horde/Plugins/Build/HordeServer.Build.Tests/Jobs/Schedules/SchedulerTests.cs
+++ b/Engine/Source/Programs/Horde/Plugins/Build/HordeServer.Build.Tests/Jobs/Schedules/SchedulerTests.cs
@@ -651,18 +651,18 @@ namespace HordeServer.Tests.Jobs.Schedules
 			templateConfig.Id = TemplateId;
 			templateConfig.Name = "Test";
 			templateConfig.Parameters.Add(
-				new TextParameterData
+				new TemplateTextParameterConfig
 				{
 					Argument = "-Text=",
 					Default = "Default",
 					ScheduleOverride = "Scheduled"
 				});
 			templateConfig.Parameters.Add(
-				new ListParameterData
+				new TemplateListParameterConfig
 				{
-					Items = new List<ListParameterItemData>
+					Items = new List<TemplateListParameterItemConfig>
 					{
-						new ListParameterItemData
+						new TemplateListParameterItemConfig
 						{
 							ArgumentIfEnabled = "-List=Default",
 							ArgumentIfDisabled = "-List=Scheduled",
@@ -672,7 +672,7 @@ namespace HordeServer.Tests.Jobs.Schedules
 					}
 				});
 			templateConfig.Parameters.Add(
-				new BoolParameterData
+				new TemplateBoolParameterConfig
 				{
 					ArgumentIfEnabled = "-Bool=Default",
 					ArgumentIfDisabled = "-Bool=Scheduled",
diff --git a/Engine/Source/Programs/Horde/Plugins/Build/HordeServer.Build/Jobs/Templates/ITemplate.cs b/Engine/Source/Programs/Horde/Plugins/Build/HordeServer.Build/Jobs/Templates/ITemplate.cs
--- a/Engine/Source/Programs/Horde/Plugins/Build/HordeServer.Build/Jobs/Templates/ITemplate.cs
+++ b/Engine/Source/Programs/Horde/Plugins/Build/HordeServer.Build/Jobs/Templates/ITemplate.cs
@@ -3,584 +3,296 @@
 using EpicGames.Core;
 using EpicGames.Horde.Jobs;
 using EpicGames.Horde.Jobs.Templates;
-using MongoDB.Bson.Serialization.Attributes;
 
 #pragma warning disable CA2227 // Change x to be read-only by removing the property setter
 
 namespace HordeServer.Jobs.Templates
 {
 	/// <summary>
-	/// Base class for parameters used to configure templates via the new build dialog
+	/// Document describing a job template. These objects are considered immutable once created and uniquely referenced by hash, in order to de-duplicate across all job runs.
 	/// </summary>
-	[BsonDiscriminator(RootClass = true)]
-	[BsonKnownTypes(typeof(TextParameter), typeof(ListParameter), typeof(BoolParameter))]
-	public abstract class Parameter
+	public interface ITemplate
 	{
 		/// <summary>
-		/// Gets the arguments for a job given a set of parameters
-		/// </summary>
-		/// <param name="parameters">Map of parameter id to value</param>
-		/// <param name="scheduledBuild">Whether this is a scheduled build</param>
-		/// <param name="arguments">Receives command line arguments for the job</param>
-		public abstract void GetArguments(IReadOnlyDictionary<ParameterId, string> parameters, bool scheduledBuild, List<string> arguments);
-
-		/// <summary>
-		/// Gets the default arguments for this parameter and its children
+		/// Hash of this template
 		/// </summary>
-		/// <param name="parameters">List of default parameters</param>
-		/// <param name="scheduledBuild">Whether the arguments are being queried for a scheduled build</param>
-		public abstract void GetDefaultParameters(Dictionary<ParameterId, string> parameters, bool scheduledBuild);
+		public ContentHash Hash { get; }
 
 		/// <summary>
-		/// Convert this parameter to data for serialization
+		/// Name of the template.
 		/// </summary>
-		/// <returns>Serializable parameter data</returns>
-		public abstract ParameterData ToData();
-	}
+		public string Name { get; }
 
-	/// <summary>
-	/// Free-form text entry parameter
-	/// </summary>
-	public sealed class TextParameter : Parameter
-	{
 		/// <summary>
-		/// Identifier for this parameter
+		/// Description for the template
 		/// </summary>
-		public ParameterId Id { get; set; }
+		public string? Description { get; }
 
 		/// <summary>
-		/// Label to display next to this parameter. Should default to the parameter name.
+		/// Priority of this job
 		/// </summary>
-		public string Label { get; set; }
+		public Priority? Priority { get; }
 
 		/// <summary>
-		/// Argument to add (will have the value of this field appended)
+		/// Whether to allow preflights for this job type
 		/// </summary>
-		public string Argument { get; set; }
+		public bool AllowPreflights { get; }
 
 		/// <summary>
-		/// Default value for this argument
+		/// Whether to always issues for jobs using this template
 		/// </summary>
-		public string Default { get; set; }
+		public bool UpdateIssues { get; }
 
 		/// <summary>
-		/// Override for this argument in scheduled builds. 
+		/// Whether to promote issues by default for jobs using this template
 		/// </summary>
-		public string? ScheduleOverride { get; set; }
+		public bool PromoteIssuesByDefault { get; }
 
 		/// <summary>
-		/// Hint text to display when the field is empty
+		/// Agent type to use for parsing the job state
 		/// </summary>
-		public string? Hint { get; set; }
+		public string? InitialAgentType { get; }
 
 		/// <summary>
-		/// Regex used to validate values entered into this text field.
+		/// Path to a file within the stream to submit to generate a new changelist for jobs
 		/// </summary>
-		[BsonIgnoreIfNull]
-		public string? Validation { get; set; }
+		public string? SubmitNewChange { get; }
 
 		/// <summary>
-		/// Message displayed to explain valid values if validation fails.
+		/// Description for new changelists
 		/// </summary>
-		[BsonIgnoreIfNull]
-		public string? Description { get; set; }
+		public string? SubmitDescription { get; }
 
 		/// <summary>
-		/// Tool tip text to display
+		/// Optional predefined user-defined properties for this job
 		/// </summary>
-		public string? ToolTip { get; set; }
+		public IReadOnlyList<string> Arguments { get; }
 
 		/// <summary>
-		/// Private constructor for serialization
+		/// Parameters for this template
 		/// </summary>
-		public TextParameter()
-		{
-			Label = null!;
-			Argument = null!;
-			Default = String.Empty;
-		}
+		public IReadOnlyList<ITemplateParameter> Parameters { get; }
+	}
 
+	/// <summary>
+	/// Base class for parameters used to configure templates via the new build dialog
+	/// </summary>
+	public interface ITemplateParameter
+	{
 		/// <summary>
-		/// Constructor
+		/// Gets the arguments for a job given a set of parameters
 		/// </summary>
-		/// <param name="id">Identifier for this parameter</param>
-		/// <param name="label">Parameter to pass this value to the BuildGraph script with</param>
-		/// <param name="argument">Label to show next to the parameter</param>
-		/// <param name="defaultValue">Default value for this argument</param>
-		/// <param name="scheduleOverride">Override for this argument in scheduled builds</param>
-		/// <param name="hint">Hint text to display</param>
-		/// <param name="validation">Regex used to validate entries</param>
-		/// <param name="description">Message displayed for invalid values</param>
-		/// <param name="toolTip">Tool tip text to display</param>
-		public TextParameter(ParameterId id, string label, string argument, string defaultValue, string? scheduleOverride, string? hint, string? validation, string? description, string? toolTip)
-		{
-			Id = id;
-			Label = label;
-			Argument = argument;
-			Default = defaultValue;
-			ScheduleOverride = scheduleOverride;
-			Hint = hint;
-			Validation = validation;
-			Description = description;
-			ToolTip = toolTip;
-		}
-
-		/// <inheritdoc/>
-		public override void GetArguments(IReadOnlyDictionary<ParameterId, string> parameters, bool scheduledBuild, List<string> arguments)
-		{
-			if (parameters.TryGetValue(Id, out string? value))
-			{
-				arguments.Add(Argument + value);
-			}
-			else
-			{
-				arguments.Add(Argument + (scheduledBuild ? (ScheduleOverride ?? Default) : Default));
-			}
-		}
-
-		/// <inheritdoc/>
-		public override void GetDefaultParameters(Dictionary<ParameterId, string> parameters, bool scheduledBuild)
-		{
-			parameters[Id] = scheduledBuild ? (ScheduleOverride ?? Default) : Default;
-		}
+		/// <param name="parameters">Map of parameter id to value</param>
+		/// <param name="scheduledBuild">Whether this is a scheduled build</param>
+		/// <param name="arguments">Receives command line arguments for the job</param>
+		void GetArguments(IReadOnlyDictionary<ParameterId, string> parameters, bool scheduledBuild, List<string> arguments);
 
 		/// <summary>
-		/// Convert this parameter to data for serialization
+		/// Gets the default arguments for this parameter and its children
 		/// </summary>
-		/// <returns>Serializable parameter data</returns>
-		public override ParameterData ToData()
-		{
-			return new TextParameterData(Id, Label, Argument, Default, ScheduleOverride, Hint, Validation, Description, ToolTip);
-		}
+		/// <param name="parameters">List of default parameters</param>
+		/// <param name="scheduledBuild">Whether the arguments are being queried for a scheduled build</param>
+		void GetDefaultParameters(Dictionary<ParameterId, string> parameters, bool scheduledBuild);
 	}
 
 	/// <summary>
-	/// Possible option for a list parameter
+	/// Allows the user to toggle an option on or off
 	/// </summary>
-	public class ListParameterItem
+	public interface ITemplateBoolParameter : ITemplateParameter
 	{
 		/// <summary>
 		/// Identifier for this parameter
 		/// </summary>
-		public ParameterId Id { get; set; }
-
-		/// <summary>
-		/// Group to display this entry in
-		/// </summary>
-		[BsonIgnoreIfNull]
-		public string? Group { get; set; }
-
-		/// <summary>
-		/// Text to display for this option.
-		/// </summary>
-		public string Text { get; set; }
-
-		/// <summary>
-		/// Argument to add if this parameter is enabled.
-		/// </summary>
-		[BsonIgnoreIfNull]
-		public string? ArgumentIfEnabled { get; set; }
+		ParameterId Id { get; }
 
 		/// <summary>
-		/// Arguments to add if this parameter is enabled.
+		/// Label to display next to this parameter.
 		/// </summary>
-		[BsonIgnoreIfNull]
-		public List<string>? ArgumentsIfEnabled { get; set; }
+		string Label { get; }
 
 		/// <summary>
-		/// Argument to add if this parameter is disabled.
+		/// Argument to add if this parameter is enabled
 		/// </summary>
-		[BsonIgnoreIfNull]
-		public string? ArgumentIfDisabled { get; set; }
+		string? ArgumentIfEnabled { get; }
 
 		/// <summary>
-		/// Arguments to add if this parameter is disabled.
+		/// Arguments to add if this parameter is enabled
 		/// </summary>
-		[BsonIgnoreIfNull]
-		public List<string>? ArgumentsIfDisabled { get; set; }
+		IReadOnlyList<string>? ArgumentsIfEnabled { get; }
 
 		/// <summary>
-		/// Whether this item is selected by default
+		/// Argument to add if this parameter is disabled
 		/// </summary>
-		[BsonIgnoreIfDefault, BsonDefaultValue(false)]
-		public bool Default { get; set; }
+		string? ArgumentIfDisabled { get; }
 
 		/// <summary>
-		/// Whether this item is selected by default
+		/// Arguments to add if this parameter is disabled
 		/// </summary>
-		public bool? ScheduleOverride { get; set; }
+		IReadOnlyList<string>? ArgumentsIfDisabled { get; }
 
 		/// <summary>
-		/// Default constructor for JSON serializer
+		/// Whether this option should be enabled by default
 		/// </summary>
-		public ListParameterItem()
-		{
-			Text = String.Empty;
-		}
+		bool Default { get; }
 
 		/// <summary>
-		/// Constructor
+		/// Whether this option should be enabled by default
 		/// </summary>
-		/// <param name="id">Identifier for this parameter</param>
-		/// <param name="group">The group to put this parameter in</param>
-		/// <param name="text">Text to display for this option</param>
-		/// <param name="argumentIfEnabled">Argument to add if this option is enabled</param>
-		/// <param name="argumentsIfEnabled">Arguments to add if this option is enabled</param>
-		/// <param name="argumentIfDisabled">Argument to add if this option is disabled</param>
-		/// <param name="argumentsIfDisabled">Arguments to add if this option is disabled </param>
-		/// <param name="defaultValue">Whether this item is selected by default</param>
-		/// <param name="scheduleOverride">Override for this value in scheduled builds</param>
-		public ListParameterItem(ParameterId id, string? group, string text, string? argumentIfEnabled, List<string>? argumentsIfEnabled, string? argumentIfDisabled, List<string>? argumentsIfDisabled, bool defaultValue, bool? scheduleOverride)
-		{
-			Id = id;
-			Group = group;
-			Text = text;
-			ArgumentIfEnabled = argumentIfEnabled;
-			ArgumentsIfEnabled = argumentsIfEnabled;
-			ArgumentIfDisabled = argumentIfDisabled;
-			ArgumentsIfDisabled = argumentsIfDisabled;
-			Default = defaultValue;
-			ScheduleOverride = scheduleOverride;
-		}
+		bool? ScheduleOverride { get; }
 
 		/// <summary>
-		/// Convert this parameter to data for serialization
+		/// Tool tip text to display
 		/// </summary>
-		/// <returns>Serializable parameter data</returns>
-		public ListParameterItemData ToData()
-		{
-			return new ListParameterItemData(Id, Group, Text, ArgumentIfEnabled, ArgumentsIfEnabled, ArgumentIfDisabled, ArgumentsIfDisabled, Default, ScheduleOverride);
-		}
+		string? ToolTip { get; }
 	}
 
 	/// <summary>
-	/// Allows the user to select a value from a constrained list of choices
+	/// Free-form text entry parameter
 	/// </summary>
-	public class ListParameter : Parameter
+	public interface ITemplateTextParameter : ITemplateParameter
 	{
 		/// <summary>
-		/// Label to display next to this parameter.
+		/// Identifier for this parameter
 		/// </summary>
-		public string Label { get; set; }
+		ParameterId Id { get; }
 
 		/// <summary>
-		/// Style of picker parameter to use
+		/// Label to display next to this parameter. Should default to the parameter name.
 		/// </summary>
-		public ListParameterStyle Style { get; set; }
+		string Label { get; }
 
 		/// <summary>
-		/// List of values to display in the list
+		/// Argument to add (will have the value of this field appended)
 		/// </summary>
-		public List<ListParameterItem> Items { get; set; }
+		string Argument { get; }
 
 		/// <summary>
-		/// Tool tip text to display
+		/// Default value for this argument
 		/// </summary>
-		public string? ToolTip { get; set; }
+		string Default { get; }
 
 		/// <summary>
-		/// Private constructor for serialization
+		/// Override for this argument in scheduled builds. 
 		/// </summary>
-		public ListParameter()
-		{
-			Label = null!;
-			Items = null!;
-		}
+		string? ScheduleOverride { get; }
 
 		/// <summary>
-		/// List of possible values
+		/// Hint text to display when the field is empty
 		/// </summary>
-		/// <param name="label">Label to show next to this parameter</param>
-		/// <param name="style">Style of list parameter to use</param>
-		/// <param name="entries">Entries for this list</param>
-		/// <param name="toolTip">Tool tip text to display</param>
-		public ListParameter(string label, ListParameterStyle style, List<ListParameterItem> entries, string? toolTip)
-		{
-			Label = label;
-			Style = style;
-			Items = entries;
-			ToolTip = toolTip;
-		}
-
-		/// <inheritdoc/>
-		public override void GetArguments(IReadOnlyDictionary<ParameterId, string> parameters, bool scheduledBuild, List<string> arguments)
-		{
-			foreach (ListParameterItem item in Items)
-			{
-				bool value = BoolParameter.GetValue(item.Id, parameters) ?? (scheduledBuild ? (item.ScheduleOverride ?? item.Default) : item.Default);
-				GetCommandLineArgumentsForItem(item, value, arguments);
-			}
-		}
-
-		/// <inheritdoc/>
-		public override void GetDefaultParameters(Dictionary<ParameterId, string> parameters, bool scheduledBuild)
-		{
-			foreach (ListParameterItem item in Items)
-			{
-				bool value = scheduledBuild ? (item.ScheduleOverride ?? item.Default) : item.Default;
-				parameters[item.Id] = value.ToString();
-			}
-		}
-
-		static void GetCommandLineArgumentsForItem(ListParameterItem item, bool value, List<string> arguments)
-		{
-			if (value)
-			{
-				if (item.ArgumentIfEnabled != null)
-				{
-					arguments.Add(item.ArgumentIfEnabled);
-				}
-				if (item.ArgumentsIfEnabled != null)
-				{
-					arguments.AddRange(item.ArgumentsIfEnabled);
-				}
-			}
-			else
-			{
-				if (item.ArgumentIfDisabled != null)
-				{
-					arguments.Add(item.ArgumentIfDisabled);
-				}
-				if (item.ArgumentsIfDisabled != null)
-				{
-					arguments.AddRange(item.ArgumentsIfDisabled);
-				}
-			}
-		}
+		string? Hint { get; }
 
 		/// <summary>
-		/// Convert this parameter to data for serialization
+		/// Regex used to validate values entered into this text field.
 		/// </summary>
-		/// <returns>Serializable parameter data</returns>
-		public override ParameterData ToData()
-		{
-			return new ListParameterData(Label, Style, Items.ConvertAll(x => x.ToData()), ToolTip);
-		}
-	}
+		string? Validation { get; }
 
-	/// <summary>
-	/// Allows the user to toggle an option on or off
-	/// </summary>
-	public class BoolParameter : Parameter
-	{
 		/// <summary>
-		/// Identifier for this parameter
+		/// Message displayed to explain valid values if validation fails.
 		/// </summary>
-		public ParameterId Id { get; set; }
+		string? ValidationError { get; }
 
 		/// <summary>
-		/// Label to display next to this parameter.
+		/// Tool tip text to display
 		/// </summary>
-		public string Label { get; set; }
+		string? ToolTip { get; }
+	}
 
+	/// <summary>
+	/// Style of list parameter
+	/// </summary>
+	public enum TemplateListParameterStyle
+	{
 		/// <summary>
-		/// Argument to add if this parameter is enabled
+		/// Regular drop-down list. One item is always selected.
 		/// </summary>
-		[BsonIgnoreIfNull]
-		public string? ArgumentIfEnabled { get; set; }
+		List,
 
 		/// <summary>
-		/// Arguments to add if this parameter is enabled
+		/// Drop-down list with checkboxes
 		/// </summary>
-		[BsonIgnoreIfNull]
-		public List<string>? ArgumentsIfEnabled { get; set; }
+		MultiList,
 
 		/// <summary>
-		/// Argument to add if this parameter is disabled
+		/// Tag picker from list of options
 		/// </summary>
-		[BsonIgnoreIfNull]
-		public string? ArgumentIfDisabled { get; set; }
+		TagPicker,
+	}
 
+	/// <summary>
+	/// Allows the user to select a value from a constrained list of choices
+	/// </summary>
+	public interface ITemplateListParameter : ITemplateParameter
+	{
 		/// <summary>
-		/// Arguments to add if this parameter is disabled
+		/// Label to display next to this parameter.
 		/// </summary>
-		[BsonIgnoreIfNull]
-		public List<string>? ArgumentsIfDisabled { get; set; }
+		string Label { get; }
 
 		/// <summary>
-		/// Whether this option should be enabled by default
+		/// Style of picker parameter to use
 		/// </summary>
-		[BsonIgnoreIfDefault, BsonDefaultValue(false)]
-		public bool Default { get; set; }
+		TemplateListParameterStyle Style { get; }
 
 		/// <summary>
-		/// Whether this option should be enabled by default
+		/// List of values to display in the list
 		/// </summary>
-		[BsonIgnoreIfNull]
-		public bool? ScheduleOverride { get; set; }
+		IReadOnlyList<ITemplateListParameterItem> Items { get; }
 
 		/// <summary>
 		/// Tool tip text to display
 		/// </summary>
-		public string? ToolTip { get; set; }
-
-		/// <summary>
-		/// Private constructor for serialization
-		/// </summary>
-		public BoolParameter()
-		{
-			Label = String.Empty;
-		}
-
-		/// <summary>
-		/// Constructor
-		/// </summary>
-		/// <param name="id">Identifier for this parameter</param>
-		/// <param name="label">Label to display for this parameter</param>
-		/// <param name="argumentIfEnabled">Argument to add if this parameter is enabled</param>
-		/// <param name="argumentsIfEnabled">Arguments to add if this parameter is enabled</param>
-		/// <param name="argumentIfDisabled">Argument to add if this parameter is disabled</param>
-		/// <param name="argumentsIfDisabled">Arguments to add if this parameter is disabled</param>
-		/// <param name="defaultValue">Default value for this argument</param>
-		/// <param name="scheduleOverride">Override for this argument in scheduled builds</param>
-		/// <param name="toolTip">Tool tip text to display</param>
-		public BoolParameter(ParameterId id, string label, string? argumentIfEnabled, List<string>? argumentsIfEnabled, string? argumentIfDisabled, List<string>? argumentsIfDisabled, bool defaultValue, bool? scheduleOverride, string? toolTip)
-		{
-			Id = id;
-			Label = label;
-			ArgumentIfEnabled = argumentIfEnabled;
-			ArgumentsIfEnabled = argumentsIfEnabled;
-			ArgumentIfDisabled = argumentIfDisabled;
-			ArgumentsIfDisabled = argumentsIfDisabled;
-			Default = defaultValue;
-			ScheduleOverride = scheduleOverride;
-			ToolTip = toolTip;
-		}
-
-		/// <summary>
-		/// Parses a bool parameter value from a dictionary
-		/// </summary>
-		/// <param name="parameterId">The parameter id</param>
-		/// <param name="parameters">Map from parameter id to value</param>
-		public static bool? GetValue(ParameterId parameterId, IReadOnlyDictionary<ParameterId, string> parameters)
-		{
-			if (parameters.TryGetValue(parameterId, out string? stringValue))
-			{
-				if (stringValue.Equals("0", StringComparison.Ordinal) || stringValue.Equals("false", StringComparison.OrdinalIgnoreCase))
-				{
-					return false;
-				}
-				if (stringValue.Equals("1", StringComparison.Ordinal) || stringValue.Equals("true", StringComparison.OrdinalIgnoreCase))
-				{
-					return true;
-				}
-			}
-			return null;
-		}
-
-		/// <inheritdoc/>
-		public override void GetArguments(IReadOnlyDictionary<ParameterId, string> parameters, bool scheduledBuild, List<string> arguments)
-		{
-			bool value = GetValue(Id, parameters) ?? (scheduledBuild ? (ScheduleOverride ?? Default) : Default);
-			GetCommandLineArgumentsForItem(value, arguments);
-		}
-
-		/// <inheritdoc/>
-		public override void GetDefaultParameters(Dictionary<ParameterId, string> parameters, bool scheduledBuild)
-		{
-			bool value = scheduledBuild ? (ScheduleOverride ?? Default) : Default;
-			parameters[Id] = value.ToString();
-		}
-
-		void GetCommandLineArgumentsForItem(bool value, List<string> arguments)
-		{
-			if (value)
-			{
-				if (!String.IsNullOrEmpty(ArgumentIfEnabled))
-				{
-					arguments.Add(ArgumentIfEnabled);
-				}
-				if (ArgumentsIfEnabled != null)
-				{
-					arguments.AddRange(ArgumentsIfEnabled);
-				}
-			}
-			else
-			{
-				if (!String.IsNullOrEmpty(ArgumentIfDisabled))
-				{
-					arguments.Add(ArgumentIfDisabled);
-				}
-				if (ArgumentsIfDisabled != null)
-				{
-					arguments.AddRange(ArgumentsIfDisabled);
-				}
-			}
-		}
-		/// <summary>
-		/// Convert this parameter to data for serialization
-		/// </summary>
-		/// <returns><see cref="BoolParameterData"/> instance</returns>
-		public override ParameterData ToData()
-		{
-			return new BoolParameterData(Id, Label, ArgumentIfEnabled, ArgumentsIfEnabled, ArgumentIfDisabled, ArgumentsIfDisabled, Default, ScheduleOverride, ToolTip);
-		}
+		string? ToolTip { get; }
 	}
 
 	/// <summary>
-	/// Document describing a job template. These objects are considered immutable once created and uniquely referenced by hash, in order to de-duplicate across all job runs.
+	/// Possible option for a list parameter
 	/// </summary>
-	public interface ITemplate
+	public interface ITemplateListParameterItem
 	{
 		/// <summary>
-		/// Hash of this template
-		/// </summary>
-		public ContentHash Hash { get; }
-
-		/// <summary>
-		/// Name of the template.
-		/// </summary>
-		public string Name { get; }
-
-		/// <summary>
-		/// Description for the template
+		/// Identifier for this parameter
 		/// </summary>
-		public string? Description { get; }
+		ParameterId Id { get; }
 
 		/// <summary>
-		/// Priority of this job
+		/// Group to display this entry in
 		/// </summary>
-		public Priority? Priority { get; }
+		string? Group { get; }
 
 		/// <summary>
-		/// Whether to allow preflights for this job type
+		/// Text to display for this option.
 		/// </summary>
-		public bool AllowPreflights { get; }
+		string Text { get; }
 
 		/// <summary>
-		/// Whether to always issues for jobs using this template
+		/// Argument to add if this parameter is enabled.
 		/// </summary>
-		public bool UpdateIssues { get; }
+		string? ArgumentIfEnabled { get; }
 
 		/// <summary>
-		/// Whether to promote issues by default for jobs using this template
-		/// </summary>
-		public bool PromoteIssuesByDefault { get; }
-
-		/// <summary>
-		/// Agent type to use for parsing the job state
+		/// Arguments to add if this parameter is enabled.
 		/// </summary>
-		public string? InitialAgentType { get; }
+		IReadOnlyList<string>? ArgumentsIfEnabled { get; }
 
 		/// <summary>
-		/// Path to a file within the stream to submit to generate a new changelist for jobs
+		/// Argument to add if this parameter is disabled.
 		/// </summary>
-		public string? SubmitNewChange { get; }
+		string? ArgumentIfDisabled { get; }
 
 		/// <summary>
-		/// Description for new changelists
+		/// Arguments to add if this parameter is disabled.
 		/// </summary>
-		public string? SubmitDescription { get; }
+		IReadOnlyList<string>? ArgumentsIfDisabled { get; }
 
 		/// <summary>
-		/// Optional predefined user-defined properties for this job
+		/// Whether this item is selected by default
 		/// </summary>
-		public IReadOnlyList<string> Arguments { get; }
+		bool Default { get; }
 
 		/// <summary>
-		/// Parameters for this template
+		/// Whether this item is selected by default
 		/// </summary>
-		public IReadOnlyList<Parameter> Parameters { get; }
+		bool? ScheduleOverride { get; }
 	}
 
 	/// <summary>
@@ -595,7 +307,7 @@ namespace HordeServer.Jobs.Templates
 		{
 			arguments.AddRange(template.Arguments);
 
-			foreach (Parameter parameter in template.Parameters)
+			foreach (ITemplateParameter parameter in template.Parameters)
 			{
 				parameter.GetArguments(parameters, false, arguments);
 			}
@@ -607,7 +319,7 @@ namespace HordeServer.Jobs.Templates
 		/// <returns>List of default arguments</returns>
 		public static void GetDefaultParameters(this ITemplate template, Dictionary<ParameterId, string> parameters, bool scheduledBuild)
 		{
-			foreach (Parameter parameter in template.Parameters)
+			foreach (ITemplateParameter parameter in template.Parameters)
 			{
 				parameter.GetDefaultParameters(parameters, scheduledBuild);
 			}
diff --git a/Engine/Source/Programs/Horde/Plugins/Build/HordeServer.Build/Jobs/Templates/TemplateCollection.cs b/Engine/Source/Programs/Horde/Plugins/Build/HordeServer.Build/Jobs/Templates/TemplateCollection.cs
--- a/Engine/Source/Programs/Horde/Plugins/Build/HordeServer.Build/Jobs/Templates/TemplateCollection.cs
+++ b/Engine/Source/Programs/Horde/Plugins/Build/HordeServer.Build/Jobs/Templates/TemplateCollection.cs
@@ -2,6 +2,7 @@
 
 using EpicGames.Core;
 using EpicGames.Horde.Jobs;
+using EpicGames.Horde.Jobs.Templates;
 using HordeServer.Server;
 using Microsoft.Extensions.Caching.Memory;
 using MongoDB.Bson;
@@ -42,11 +43,11 @@ namespace HordeServer.Jobs.Templates
 			public string? SubmitDescription { get; set; }
 
 			public List<string> Arguments { get; private set; } = new List<string>();
-			public List<Parameter> Parameters { get; private set; } = new List<Parameter>();
+			public List<ParameterDocument> Parameters { get; private set; } = new List<ParameterDocument>();
 
 			ContentHash ITemplate.Hash => Id;
 			IReadOnlyList<string> ITemplate.Arguments => Arguments;
-			IReadOnlyList<Parameter> ITemplate.Parameters => Parameters;
+			IReadOnlyList<ITemplateParameter> ITemplate.Parameters => Parameters;
 
 			[BsonConstructor]
 			private TemplateDocument()
@@ -66,11 +67,317 @@ namespace HordeServer.Jobs.Templates
 				SubmitNewChange = config.SubmitNewChange;
 				SubmitDescription = config.SubmitDescription;
 				Arguments = config.Arguments ?? new List<string>();
-				Parameters = config.Parameters.ConvertAll(x => x.ToModel());
+				Parameters = config.Parameters.ConvertAll(x => CreateParameterDocument(x));
 
 				// Compute the hash once all other fields have been set
 				Id = ContentHash.SHA1(BsonExtensionMethods.ToBson(this));
 			}
+
+			static ParameterDocument CreateParameterDocument(TemplateParameterConfig config)
+			{
+				return config switch
+				{
+					TemplateBoolParameterConfig boolParameter
+						=> new BoolParameterDocument(boolParameter),
+
+					TemplateTextParameterConfig textParameter
+						=> new TextParameterDocument(textParameter),
+
+					TemplateListParameterConfig listParameter
+						=> new ListParameterDocument(listParameter),
+
+					_ => throw new NotImplementedException()
+				};
+			}
+		}
+
+		[BsonDiscriminator(RootClass = true)]
+		[BsonKnownTypes(typeof(BoolParameterDocument), typeof(TextParameterDocument), typeof(ListParameterDocument))]
+		abstract class ParameterDocument : ITemplateParameter
+		{
+			public abstract void GetArguments(IReadOnlyDictionary<ParameterId, string> parameters, bool scheduledBuild, List<string> arguments);
+			public abstract void GetDefaultParameters(Dictionary<ParameterId, string> parameters, bool scheduledBuild);
+		}
+
+		sealed class BoolParameterDocument : ParameterDocument, ITemplateBoolParameter
+		{
+			public ParameterId Id { get; set; }
+			public string Label { get; set; }
+
+			[BsonIgnoreIfNull]
+			public string? ArgumentIfEnabled { get; set; }
+
+			[BsonIgnoreIfNull]
+			public List<string>? ArgumentsIfEnabled { get; set; }
+
+			[BsonIgnoreIfNull]
+			public string? ArgumentIfDisabled { get; set; }
+
+			[BsonIgnoreIfNull]
+			public List<string>? ArgumentsIfDisabled { get; set; }
+
+			[BsonIgnoreIfDefault, BsonDefaultValue(false)]
+			public bool Default { get; set; }
+
+			[BsonIgnoreIfNull]
+			public bool? ScheduleOverride { get; set; }
+
+			public string? ToolTip { get; set; }
+
+			IReadOnlyList<string>? ITemplateBoolParameter.ArgumentsIfEnabled => ArgumentsIfEnabled;
+			IReadOnlyList<string>? ITemplateBoolParameter.ArgumentsIfDisabled => ArgumentsIfDisabled;
+
+			public BoolParameterDocument()
+			{
+				Label = String.Empty;
+			}
+
+			public BoolParameterDocument(TemplateBoolParameterConfig config)
+			{
+				Id = config.Id;
+				Label = config.Label;
+				ArgumentIfEnabled = config.ArgumentIfEnabled;
+				ArgumentsIfEnabled = config.ArgumentsIfEnabled;
+				ArgumentIfDisabled = config.ArgumentIfDisabled;
+				ArgumentsIfDisabled = config.ArgumentsIfDisabled;
+				Default = config.Default;
+				ScheduleOverride = config.ScheduleOverride;
+				ToolTip = config.ToolTip;
+			}
+
+			public static bool? GetValue(ParameterId parameterId, IReadOnlyDictionary<ParameterId, string> parameters)
+			{
+				if (parameters.TryGetValue(parameterId, out string? stringValue))
+				{
+					if (stringValue.Equals("0", StringComparison.Ordinal) || stringValue.Equals("false", StringComparison.OrdinalIgnoreCase))
+					{
+						return false;
+					}
+					if (stringValue.Equals("1", StringComparison.Ordinal) || stringValue.Equals("true", StringComparison.OrdinalIgnoreCase))
+					{
+						return true;
+					}
+				}
+				return null;
+			}
+
+			/// <inheritdoc/>
+			public override void GetArguments(IReadOnlyDictionary<ParameterId, string> parameters, bool scheduledBuild, List<string> arguments)
+			{
+				bool value = GetValue(Id, parameters) ?? (scheduledBuild ? (ScheduleOverride ?? Default) : Default);
+				GetCommandLineArgumentsForItem(value, arguments);
+			}
+
+			/// <inheritdoc/>
+			public override void GetDefaultParameters(Dictionary<ParameterId, string> parameters, bool scheduledBuild)
+			{
+				bool value = scheduledBuild ? (ScheduleOverride ?? Default) : Default;
+				parameters[Id] = value.ToString();
+			}
+
+			void GetCommandLineArgumentsForItem(bool value, List<string> arguments)
+			{
+				if (value)
+				{
+					if (!String.IsNullOrEmpty(ArgumentIfEnabled))
+					{
+						arguments.Add(ArgumentIfEnabled);
+					}
+					if (ArgumentsIfEnabled != null)
+					{
+						arguments.AddRange(ArgumentsIfEnabled);
+					}
+				}
+				else
+				{
+					if (!String.IsNullOrEmpty(ArgumentIfDisabled))
+					{
+						arguments.Add(ArgumentIfDisabled);
+					}
+					if (ArgumentsIfDisabled != null)
+					{
+						arguments.AddRange(ArgumentsIfDisabled);
+					}
+				}
+			}
+		}
+
+		sealed class TextParameterDocument : ParameterDocument, ITemplateTextParameter
+		{
+			public ParameterId Id { get; set; }
+			public string Label { get; set; }
+			public string Argument { get; set; }
+			public string Default { get; set; }
+			public string? ScheduleOverride { get; set; }
+			public string? Hint { get; set; }
+
+			[BsonIgnoreIfNull]
+			public string? Validation { get; set; }
+
+			[BsonIgnoreIfNull]
+			public string? ValidationError { get; set; }
+
+			public string? ToolTip { get; set; }
+
+			public TextParameterDocument()
+			{
+				Label = null!;
+				Argument = null!;
+				Default = String.Empty;
+			}
+
+			public TextParameterDocument(TemplateTextParameterConfig config)
+			{
+				Id = config.Id;
+				Label = config.Label;
+				Argument = config.Argument;
+				Default = config.Default;
+				ScheduleOverride = config.ScheduleOverride;
+				Hint = config.Hint;
+				Validation = config.Validation;
+				ValidationError = config.ValidationError;
+				ToolTip = config.ToolTip;
+			}
+
+			/// <inheritdoc/>
+			public override void GetArguments(IReadOnlyDictionary<ParameterId, string> parameters, bool scheduledBuild, List<string> arguments)
+			{
+				if (parameters.TryGetValue(Id, out string? value))
+				{
+					arguments.Add(Argument + value);
+				}
+				else
+				{
+					arguments.Add(Argument + (scheduledBuild ? (ScheduleOverride ?? Default) : Default));
+				}
+			}
+
+			/// <inheritdoc/>
+			public override void GetDefaultParameters(Dictionary<ParameterId, string> parameters, bool scheduledBuild)
+			{
+				parameters[Id] = scheduledBuild ? (ScheduleOverride ?? Default) : Default;
+			}
+		}
+
+		sealed class ListParameterDocument : ParameterDocument, ITemplateListParameter
+		{
+			public string Label { get; set; }
+			public TemplateListParameterStyle Style { get; set; }
+			public List<ListParameterItemDocument> Items { get; set; }
+			public string? ToolTip { get; set; }
+
+			IReadOnlyList<ITemplateListParameterItem> ITemplateListParameter.Items => Items;
+
+			/// <summary>
+			/// Private constructor for serialization
+			/// </summary>
+			public ListParameterDocument()
+			{
+				Label = String.Empty;
+				Items = new List<ListParameterItemDocument>();
+			}
+
+			public ListParameterDocument(TemplateListParameterConfig config)
+			{
+				Label = config.Label;
+				Style = config.Style;
+				Items = config.Items.ConvertAll(x => new ListParameterItemDocument(x));
+				ToolTip = config.ToolTip;
+			}
+
+			/// <inheritdoc/>
+			public override void GetArguments(IReadOnlyDictionary<ParameterId, string> parameters, bool scheduledBuild, List<string> arguments)
+			{
+				foreach (ListParameterItemDocument item in Items)
+				{
+					bool value = BoolParameterDocument.GetValue(item.Id, parameters) ?? (scheduledBuild ? (item.ScheduleOverride ?? item.Default) : item.Default);
+					GetCommandLineArgumentsForItem(item, value, arguments);
+				}
+			}
+
+			/// <inheritdoc/>
+			public override void GetDefaultParameters(Dictionary<ParameterId, string> parameters, bool scheduledBuild)
+			{
+				foreach (ListParameterItemDocument item in Items)
+				{
+					bool value = scheduledBuild ? (item.ScheduleOverride ?? item.Default) : item.Default;
+					parameters[item.Id] = value.ToString();
+				}
+			}
+
+			static void GetCommandLineArgumentsForItem(ListParameterItemDocument item, bool value, List<string> arguments)
+			{
+				if (value)
+				{
+					if (item.ArgumentIfEnabled != null)
+					{
+						arguments.Add(item.ArgumentIfEnabled);
+					}
+					if (item.ArgumentsIfEnabled != null)
+					{
+						arguments.AddRange(item.ArgumentsIfEnabled);
+					}
+				}
+				else
+				{
+					if (item.ArgumentIfDisabled != null)
+					{
+						arguments.Add(item.ArgumentIfDisabled);
+					}
+					if (item.ArgumentsIfDisabled != null)
+					{
+						arguments.AddRange(item.ArgumentsIfDisabled);
+					}
+				}
+			}
+		}
+
+		class ListParameterItemDocument : ITemplateListParameterItem
+		{
+			public ParameterId Id { get; set; }
+
+			[BsonIgnoreIfNull]
+			public string? Group { get; set; }
+
+			public string Text { get; set; }
+
+			[BsonIgnoreIfNull]
+			public string? ArgumentIfEnabled { get; set; }
+
+			[BsonIgnoreIfNull]
+			public List<string>? ArgumentsIfEnabled { get; set; }
+
+			[BsonIgnoreIfNull]
+			public string? ArgumentIfDisabled { get; set; }
+
+			[BsonIgnoreIfNull]
+			public List<string>? ArgumentsIfDisabled { get; set; }
+
+			[BsonIgnoreIfDefault, BsonDefaultValue(false)]
+			public bool Default { get; set; }
+
+			public bool? ScheduleOverride { get; set; }
+
+			IReadOnlyList<string>? ITemplateListParameterItem.ArgumentsIfEnabled => ArgumentsIfEnabled;
+			IReadOnlyList<string>? ITemplateListParameterItem.ArgumentsIfDisabled => ArgumentsIfDisabled;
+
+			public ListParameterItemDocument()
+			{
+				Text = String.Empty;
+			}
+
+			public ListParameterItemDocument(TemplateListParameterItemConfig config)
+			{
+				Id = config.Id;
+				Group = config.Group;
+				Text = config.Text;
+				ArgumentIfEnabled = config.ArgumentIfEnabled;
+				ArgumentsIfEnabled = config.ArgumentsIfEnabled;
+				ArgumentIfDisabled = config.ArgumentIfDisabled;
+				ArgumentsIfDisabled = config.ArgumentsIfDisabled;
+				Default = config.Default;
+				ScheduleOverride = config.ScheduleOverride;
+			}
 		}
 
 		/// <summary>
diff --git a/Engine/Source/Programs/Horde/Plugins/Build/HordeServer.Build/Jobs/Templates/TemplateConfig.cs b/Engine/Source/Programs/Horde/Plugins/Build/HordeServer.Build/Jobs/Templates/TemplateConfig.cs
--- a/Engine/Source/Programs/Horde/Plugins/Build/HordeServer.Build/Jobs/Templates/TemplateConfig.cs
+++ b/Engine/Source/Programs/Horde/Plugins/Build/HordeServer.Build/Jobs/Templates/TemplateConfig.cs
@@ -79,7 +79,7 @@
 		/// Parameters for this template
 		/// </summary>
 		[ConfigMergeStrategy(ConfigMergeStrategy.Append)]
-		public List<ParameterData> Parameters { get; set; } = new List<ParameterData>();
+		public List<TemplateParameterConfig> Parameters { get; set; } = new List<TemplateParameterConfig>();
 
 		/// <summary>
 		/// Default settings for jobs
@@ -97,8 +97,8 @@
 	/// <summary>
 	/// Base class for template parameters
 	/// </summary>
-	[JsonKnownTypes(typeof(TextParameterData), typeof(ListParameterData), typeof(BoolParameterData))]
-	public abstract class ParameterData
+	[JsonKnownTypes(typeof(TemplateTextParameterConfig), typeof(TemplateListParameterConfig), typeof(TemplateBoolParameterConfig))]
+	public abstract class TemplateParameterConfig
 	{
 		/// <summary>
 		/// Callback after a parameter has been read.
@@ -106,11 +106,5 @@
 		public abstract void PostLoad(HashSet<ParameterId> parameterIds);
 
 		/// <summary>
-		/// Convert to a parameter object
-		/// </summary>
-		/// <returns><see cref="Parameter"/> object</returns>
-		public abstract Parameter ToModel();
-
-		/// <summary>
 		/// Adds the given parameter id to the set of known ids, or create a new one
 		/// </summary>
@@ -164,10 +158,10 @@
 	}
 
 	/// <summary>
-	/// Free-form text entry parameter
+	/// Allows the user to toggle an option on or off
 	/// </summary>
-	[JsonDiscriminator("Text")]
-	public class TextParameterData : ParameterData
+	[JsonDiscriminator("Bool")]
+	public class TemplateBoolParameterConfig : TemplateParameterConfig
 	{
 		/// <summary>
 		/// Identifier for this parameter
@@ -180,72 +174,70 @@
 		public string Label { get; set; }
 
 		/// <summary>
-		/// Argument to pass to the executor
+		/// Argument to add if this parameter is enabled
 		/// </summary>
-		public string Argument { get; set; }
+		public string? ArgumentIfEnabled { get; set; }
 
 		/// <summary>
-		/// Default value for this argument
+		/// Argument to add if this parameter is enabled
 		/// </summary>
-		public string Default { get; set; }
+		public List<string>? ArgumentsIfEnabled { get; set; }
 
 		/// <summary>
-		/// Override for the default value for this parameter when running a scheduled build
+		/// Argument to add if this parameter is enabled
 		/// </summary>
-		public string? ScheduleOverride { get; set; }
+		public string? ArgumentIfDisabled { get; set; }
 
 		/// <summary>
-		/// Hint text for this parameter
+		/// Arguments to add if this parameter is disabled
 		/// </summary>
-		public string? Hint { get; set; }
+		public List<string>? ArgumentsIfDisabled { get; set; }
 
 		/// <summary>
-		/// Regex used to validate this parameter
+		/// Whether this argument is enabled by default
 		/// </summary>
-		public string? Validation { get; set; }
+		public bool Default { get; set; }
 
 		/// <summary>
-		/// Message displayed if validation fails, informing user of valid values.
+		/// Override for this parameter in scheduled builds
 		/// </summary>
-		public string? ValidationError { get; set; }
+		public bool? ScheduleOverride { get; set; }
 
 		/// <summary>
-		/// Tool-tip text to display
+		/// Tool tip text to display
 		/// </summary>
 		public string? ToolTip { get; set; }
 
 		/// <summary>
 		/// Private constructor for serialization
 		/// </summary>
-		public TextParameterData()
+		public TemplateBoolParameterConfig()
 		{
 			Label = null!;
-			Argument = null!;
-			Default = String.Empty;
 		}
 
 		/// <summary>
 		/// Constructor
 		/// </summary>
 		/// <param name="id">Identifier for this parameter</param>
-		/// <param name="label">Label to show next to the parameter</param>
-		/// <param name="argument">Argument to pass this value with</param>
-		/// <param name="defaultValue">Default value for this parameter</param>
-		/// <param name="scheduleOverride">Default value for scheduled builds</param>
-		/// <param name="hint">Hint text to display for this parameter</param>
-		/// <param name="validation">Regex used to validate entries</param>
-		/// <param name="validationError">Message displayed to explain validation issues</param>
-		/// <param name="toolTip">Tool tip text to display</param>
-		public TextParameterData(ParameterId id, string label, string argument, string defaultValue, string? scheduleOverride, string? hint, string? validation, string? validationError, string? toolTip)
+		/// <param name="label">Label to show next to this parameter</param>
+		/// <param name="argumentIfEnabled">Argument to add if this parameter is enabled</param>
+		/// <param name="argumentsIfEnabled">Arguments to add if this parameter is enabled</param>
+		/// <param name="argumentIfDisabled">Argument to add if this parameter is disabled</param>
+		/// <param name="argumentsIfDisabled">Arguments to add if this parameter is disabled</param>
+		/// <param name="defaultValue">Whether this option is enabled by default</param>
+		/// <param name="scheduleOverride">Override for scheduled builds</param>
+		/// <param name="toolTip">The tool tip text to display</param>
+		public TemplateBoolParameterConfig(ParameterId id, string label, string? argumentIfEnabled, List<string>? argumentsIfEnabled, string? argumentIfDisabled, List<string>? argumentsIfDisabled, bool defaultValue, bool? scheduleOverride, string? toolTip)
 		{
 			Id = id;
 			Label = label;
-			Argument = argument;
+			ArgumentIfEnabled = argumentIfEnabled;
+			ArgumentsIfEnabled = argumentsIfEnabled;
+			ArgumentIfDisabled = argumentIfDisabled;
+			ArgumentsIfDisabled = argumentsIfDisabled;
 			Default = defaultValue;
 			ScheduleOverride = scheduleOverride;
-			Hint = hint;
-			Validation = validation;
-			ValidationError = validationError;
 			ToolTip = toolTip;
 		}
 
@@ -253,40 +245,23 @@
 		public override void PostLoad(HashSet<ParameterId> parameterIds)
 		{
 			Id = AddParameterId(Id, Label, parameterIds);
-		}
 
-		/// <inheritdoc/>
-		public override Parameter ToModel()
-		{
-			return new TextParameter(Id, Label, Argument, Default, ScheduleOverride, Hint, Validation, ValidationError, ToolTip);
+			if (!String.IsNullOrEmpty(ArgumentIfEnabled) && ArgumentsIfEnabled != null && ArgumentsIfEnabled.Count > 0)
+			{
+				throw new InvalidDataException("Cannot specify both 'ArgumentIfEnabled' and 'ArgumentsIfEnabled'");
+			}
+			if (!String.IsNullOrEmpty(ArgumentIfDisabled) && ArgumentsIfDisabled != null && ArgumentsIfDisabled.Count > 0)
+			{
+				throw new InvalidDataException("Cannot specify both 'ArgumentIfDisabled' and 'ArgumentsIfDisabled'");
+			}
 		}
 	}
 
 	/// <summary>
-	/// Style of list parameter
+	/// Free-form text entry parameter
 	/// </summary>
-	public enum ListParameterStyle
+	[JsonDiscriminator("Text")]
+	public class TemplateTextParameterConfig : TemplateParameterConfig
 	{
 		/// <summary>
-		/// Regular drop-down list. One item is always selected.
-		/// </summary>
-		List,
-
-		/// <summary>
-		/// Drop-down list with checkboxes
-		/// </summary>
-		MultiList,
-
-		/// <summary>
-		/// Tag picker from list of options
-		/// </summary>
-		TagPicker,
-	}
-
-	/// <summary>
-	/// Possible option for a list parameter
-	/// </summary>
-	public class ListParameterItemData
-	{
-		/// <summary>
 		/// Identifier for this parameter
@@ -294,85 +269,84 @@
 		public ParameterId Id { get; set; }
 
 		/// <summary>
-		/// Optional group heading to display this entry under, if the picker style supports it.
+		/// Name of the parameter associated with this parameter.
 		/// </summary>
-		public string? Group { get; set; }
+		public string Label { get; set; }
 
 		/// <summary>
-		/// Name of the parameter associated with this list.
+		/// Argument to pass to the executor
 		/// </summary>
-		public string Text { get; set; }
+		public string Argument { get; set; }
 
 		/// <summary>
-		/// Argument to pass with this parameter.
+		/// Default value for this argument
 		/// </summary>
-		public string? ArgumentIfEnabled { get; set; }
+		public string Default { get; set; }
 
 		/// <summary>
-		/// Arguments to pass with this parameter.
+		/// Override for the default value for this parameter when running a scheduled build
 		/// </summary>
-		public List<string>? ArgumentsIfEnabled { get; set; }
+		public string? ScheduleOverride { get; set; }
 
 		/// <summary>
-		/// Argument to pass with this parameter.
+		/// Hint text for this parameter
 		/// </summary>
-		public string? ArgumentIfDisabled { get; set; }
+		public string? Hint { get; set; }
 
 		/// <summary>
-		/// Arguments to pass if this parameter is disabled.
+		/// Regex used to validate this parameter
 		/// </summary>
-		public List<string>? ArgumentsIfDisabled { get; set; }
+		public string? Validation { get; set; }
 
 		/// <summary>
-		/// Whether this item is selected by default
+		/// Message displayed if validation fails, informing user of valid values.
 		/// </summary>
-		public bool Default { get; set; }
+		public string? ValidationError { get; set; }
 
 		/// <summary>
-		/// Overridden value for this property in schedule builds
+		/// Tool-tip text to display
 		/// </summary>
-		public bool? ScheduleOverride { get; set; }
+		public string? ToolTip { get; set; }
 
 		/// <summary>
 		/// Private constructor for serialization
 		/// </summary>
-		public ListParameterItemData()
+		public TemplateTextParameterConfig()
 		{
-			Text = null!;
+			Label = null!;
+			Argument = null!;
+			Default = String.Empty;
 		}
 
 		/// <summary>
 		/// Constructor
 		/// </summary>
 		/// <param name="id">Identifier for this parameter</param>
-		/// <param name="group">The group to put this parameter in</param>
-		/// <param name="text">Text to display for this option</param>
-		/// <param name="argumentIfEnabled">Argument to pass for this item if it's enabled</param>
-		/// <param name="argumentsIfEnabled">Argument to pass for this item if it's enabled</param>
-		/// <param name="argumentIfDisabled">Argument to pass for this item if it's enabled</param>
-		/// <param name="argumentsIfDisabled">Argument to pass for this item if it's enabled</param>
-		/// <param name="defaultValue">Whether this item is selected by default</param>
-		/// <param name="scheduleOverride">Overridden value for this item for scheduled builds</param>
-		public ListParameterItemData(ParameterId id, string? group, string text, string? argumentIfEnabled, List<string>? argumentsIfEnabled, string? argumentIfDisabled, List<string>? argumentsIfDisabled, bool defaultValue, bool? scheduleOverride)
+		/// <param name="label">Label to show next to the parameter</param>
+		/// <param name="argument">Argument to pass this value with</param>
+		/// <param name="defaultValue">Default value for this parameter</param>
+		/// <param name="scheduleOverride">Default value for scheduled builds</param>
+		/// <param name="hint">Hint text to display for this parameter</param>
+		/// <param name="validation">Regex used to validate entries</param>
+		/// <param name="validationError">Message displayed to explain validation issues</param>
+		/// <param name="toolTip">Tool tip text to display</param>
+		public TemplateTextParameterConfig(ParameterId id, string label, string argument, string defaultValue, string? scheduleOverride, string? hint, string? validation, string? validationError, string? toolTip)
 		{
 			Id = id;
-			Group = group;
-			Text = text;
-			ArgumentIfEnabled = argumentIfEnabled;
-			ArgumentsIfEnabled = argumentsIfEnabled;
-			ArgumentIfDisabled = argumentIfDisabled;
-			ArgumentsIfDisabled = argumentsIfDisabled;
+			Label = label;
+			Argument = argument;
 			Default = defaultValue;
 			ScheduleOverride = scheduleOverride;
+			Hint = hint;
+			Validation = validation;
+			ValidationError = validationError;
+			ToolTip = toolTip;
 		}
 
-		/// <summary>
-		/// Converts this data to a model object
-		/// </summary>
-		/// <returns>New <see cref="ListParameterItem"/> object</returns>
-		public ListParameterItem ToModel()
+		/// <inheritdoc/>
+		public override void PostLoad(HashSet<ParameterId> parameterIds)
 		{
-			return new ListParameterItem(Id, Group, Text, ArgumentIfEnabled, ArgumentsIfEnabled, ArgumentIfDisabled, ArgumentsIfDisabled, Default, ScheduleOverride);
+			Id = AddParameterId(Id, Label, parameterIds);
 		}
 	}
 
@@ -380,7 +354,7 @@
 	/// Allows the user to select a value from a constrained list of choices
 	/// </summary>
 	[JsonDiscriminator("List")]
-	public class ListParameterData : ParameterData
+	public class TemplateListParameterConfig : TemplateParameterConfig
 	{
 		/// <summary>
 		/// Label to display next to this parameter. Defaults to the parameter name.
@@ -390,12 +364,12 @@
 		/// <summary>
 		/// The type of list parameter
 		/// </summary>
-		public ListParameterStyle Style { get; set; }
+		public TemplateListParameterStyle Style { get; set; }
 
 		/// <summary>
 		/// List of values to display in the list
 		/// </summary>
-		public List<ListParameterItemData> Items { get; set; }
+		public List<TemplateListParameterItemConfig> Items { get; set; }
 
 		/// <summary>
 		/// Tool tip text to display
@@ -405,7 +379,7 @@
 		/// <summary>
 		/// Private constructor
 		/// </summary>
-		public ListParameterData()
+		public TemplateListParameterConfig()
 		{
 			Label = null!;
 			Items = null!;
@@ -418,7 +392,7 @@
 		/// <param name="style">Type of picker to show</param>
 		/// <param name="items">Entries for this list</param>
 		/// <param name="toolTip">Tool tip text to display</param>
-		public ListParameterData(string label, ListParameterStyle style, List<ListParameterItemData> items, string? toolTip)
+		public TemplateListParameterConfig(string label, TemplateListParameterStyle style, List<TemplateListParameterItemConfig> items, string? toolTip)
 		{
 			Label = label;
 			Style = style;
@@ -429,7 +403,7 @@
 		/// <inheritdoc/>
 		public override void PostLoad(HashSet<ParameterId> parameterIds)
 		{
-			foreach (ListParameterItemData item in Items)
+			foreach (TemplateListParameterItemConfig item in Items)
 			{
 				item.Id = AddParameterId(item.Id, $"{Label}.{item.Text}", parameterIds);
 
@@ -443,19 +417,12 @@
 				}
 			}
 		}
-
-		/// <inheritdoc/>
-		public override Parameter ToModel()
-		{
-			return new ListParameter(Label, Style, Items.ConvertAll(x => x.ToModel()), ToolTip);
-		}
 	}
 
 	/// <summary>
-	/// Allows the user to toggle an option on or off
+	/// Possible option for a list parameter
 	/// </summary>
-	[JsonDiscriminator("Bool")]
-	public class BoolParameterData : ParameterData
+	public class TemplateListParameterItemConfig
 	{
 		/// <summary>
 		/// Identifier for this parameter
@@ -463,97 +430,77 @@
 		public ParameterId Id { get; set; }
 
 		/// <summary>
-		/// Name of the parameter associated with this parameter.
+		/// Optional group heading to display this entry under, if the picker style supports it.
 		/// </summary>
-		public string Label { get; set; }
+		public string? Group { get; set; }
 
 		/// <summary>
-		/// Argument to add if this parameter is enabled
+		/// Name of the parameter associated with this list.
 		/// </summary>
+		public string Text { get; set; }
+
+		/// <summary>
+		/// Argument to pass with this parameter.
+		/// </summary>
 		public string? ArgumentIfEnabled { get; set; }
 
 		/// <summary>
-		/// Argument to add if this parameter is enabled
+		/// Arguments to pass with this parameter.
 		/// </summary>
 		public List<string>? ArgumentsIfEnabled { get; set; }
 
 		/// <summary>
-		/// Argument to add if this parameter is enabled
+		/// Argument to pass with this parameter.
 		/// </summary>
 		public string? ArgumentIfDisabled { get; set; }
 
 		/// <summary>
-		/// Arguments to add if this parameter is disabled
+		/// Arguments to pass if this parameter is disabled.
 		/// </summary>
 		public List<string>? ArgumentsIfDisabled { get; set; }
 
 		/// <summary>
-		/// Whether this argument is enabled by default
+		/// Whether this item is selected by default
 		/// </summary>
 		public bool Default { get; set; }
 
 		/// <summary>
-		/// Override for this parameter in scheduled builds
+		/// Overridden value for this property in schedule builds
 		/// </summary>
 		public bool? ScheduleOverride { get; set; }
 
 		/// <summary>
-		/// Tool tip text to display
-		/// </summary>
-		public string? ToolTip { get; set; }
-
-		/// <summary>
 		/// Private constructor for serialization
 		/// </summary>
-		public BoolParameterData()
+		public TemplateListParameterItemConfig()
 		{
-			Label = null!;
+			Text = null!;
 		}
 
 		/// <summary>
 		/// Constructor
 		/// </summary>
 		/// <param name="id">Identifier for this parameter</param>
-		/// <param name="label">Label to show next to this parameter</param>
-		/// <param name="argumentIfEnabled">Argument to add if this parameter is enabled</param>
-		/// <param name="argumentsIfEnabled">Arguments to add if this parameter is enabled</param>
-		/// <param name="argumentIfDisabled">Argument to add if this parameter is disabled</param>
-		/// <param name="argumentsIfDisabled">Arguments to add if this parameter is disabled</param>
-		/// <param name="defaultValue">Whether this option is enabled by default</param>
-		/// <param name="scheduleOverride">Override for scheduled builds</param>
-		/// <param name="toolTip">The tool tip text to display</param>
-		public BoolParameterData(ParameterId id, string label, string? argumentIfEnabled, List<string>? argumentsIfEnabled, string? argumentIfDisabled, List<string>? argumentsIfDisabled, bool defaultValue, bool? scheduleOverride, string? toolTip)
+		/// <param name="group">The group to put this parameter in</param>
+		/// <param name="text">Text to display for this option</param>
+		/// <param name="argumentIfEnabled">Argument to pass for this item if it's enabled</param>
+		/// <param name="argumentsIfEnabled">Argument to pass for this item if it's enabled</param>
+		/// <param name="argumentIfDisabled">Argument to pass for this item if it's enabled</param>
+		/// <param name="argumentsIfDisabled">Argument to pass for this item if it's enabled</param>
+		/// <param name="defaultValue">Whether this item is selected by default</param>
+		/// <param name="scheduleOverride">Overridden value for this item for scheduled builds</param>
+		public TemplateListParameterItemConfig(ParameterId id, string? group, string text, string? argumentIfEnabled, List<string>? argumentsIfEnabled, string? argumentIfDisabled, List<string>? argumentsIfDisabled, bool defaultValue, bool? scheduleOverride)
 		{
 			Id = id;
-			Label = label;
+			Group = group;
+			Text = text;
 			ArgumentIfEnabled = argumentIfEnabled;
 			ArgumentsIfEnabled = argumentsIfEnabled;
 			ArgumentIfDisabled = argumentIfDisabled;
 			ArgumentsIfDisabled = argumentsIfDisabled;
 			Default = defaultValue;
 			ScheduleOverride = scheduleOverride;
-			ToolTip = toolTip;
 		}
-
-		/// <inheritdoc/>
-		public override void PostLoad(HashSet<ParameterId> parameterIds)
-		{
-			Id = AddParameterId(Id, Label, parameterIds);
-
-			if (!String.IsNullOrEmpty(ArgumentIfEnabled) && ArgumentsIfEnabled != null && ArgumentsIfEnabled.Count > 0)
-			{
-				throw new InvalidDataException("Cannot specify both 'ArgumentIfEnabled' and 'ArgumentsIfEnabled'");
-			}
-			if (!String.IsNullOrEmpty(ArgumentIfDisabled) && ArgumentsIfDisabled != null && ArgumentsIfDisabled.Count > 0)
-			{
-				throw new InvalidDataException("Cannot specify both 'ArgumentIfDisabled' and 'ArgumentsIfDisabled'");
-			}
-		}
-
-		/// <inheritdoc/>
-		public override Parameter ToModel()
-		{
-			return new BoolParameter(Id, Label, ArgumentIfEnabled, ArgumentsIfEnabled, ArgumentIfDisabled, ArgumentsIfDisabled, Default, ScheduleOverride, ToolTip);
-		}
 	}
 }
+
diff --git a/Engine/Source/Programs/Horde/Plugins/Build/HordeServer.Build/Jobs/Templates/TemplateMessages.cs b/Engine/Source/Programs/Horde/Plugins/Build/HordeServer.Build/Jobs/Templates/TemplateMessages.cs
--- a/Engine/Source/Programs/Horde/Plugins/Build/HordeServer.Build/Jobs/Templates/TemplateMessages.cs
+++ b/Engine/Source/Programs/Horde/Plugins/Build/HordeServer.Build/Jobs/Templates/TemplateMessages.cs
@@ -2,6 +2,7 @@
 
 using EpicGames.Core;
 using EpicGames.Horde.Jobs;
+using EpicGames.Horde.Jobs.Templates;
 
 #pragma warning disable CA2227 // Change x to be read-only by removing the property setter
 
@@ -55,7 +56,7 @@ namespace HordeServer.Jobs.Templates
 		/// <summary>
 		/// List of parameters for this template
 		/// </summary>
-		public List<ParameterData> Parameters { get; set; }
+		public List<GetTemplateParameterResponse> Parameters { get; set; }
 
 		/// <summary>
 		/// Parameterless constructor for serialization
@@ -65,7 +66,7 @@ namespace HordeServer.Jobs.Templates
 			Name = null!;
 			AllowPreflights = true;
 			Arguments = new List<string>();
-			Parameters = new List<ParameterData>();
+			Parameters = new List<GetTemplateParameterResponse>();
 		}
 
 		/// <summary>
@@ -82,7 +83,24 @@ namespace HordeServer.Jobs.Templates
 			InitialAgentType = template.InitialAgentType;
 			SubmitNewChange = template.SubmitNewChange;
 			Arguments = new List<string>(template.Arguments);
-			Parameters = template.Parameters.ConvertAll(x => x.ToData());
+			Parameters = template.Parameters.ConvertAll(x => CreateParameterResponse(x));
+		}
+
+		static GetTemplateParameterResponse CreateParameterResponse(ITemplateParameter parameter)
+		{
+			return parameter switch
+			{
+				ITemplateBoolParameter boolParameter
+					=> new GetTemplateBoolParameterResponse(boolParameter),
+
+				ITemplateTextParameter textParameter
+					=> new GetTemplateTextParameterResponse(textParameter),
+
+				ITemplateListParameter listParameter
+					=> new GetTemplateListParameterResponse(listParameter),
+
+				_ => throw new NotImplementedException()
+			};
 		}
 	}
 
@@ -115,4 +133,238 @@ namespace HordeServer.Jobs.Templates
 			Id = template.Hash.ToString();
 		}
 	}
+
+	/// <summary>
+	/// Base class for template parameters
+	/// </summary>
+	[JsonKnownTypes(typeof(GetTemplateBoolParameterResponse), typeof(GetTemplateTextParameterResponse), typeof(GetTemplateListParameterResponse))]
+	public abstract class GetTemplateParameterResponse : ITemplateParameter
+	{
+		void ITemplateParameter.GetArguments(IReadOnlyDictionary<ParameterId, string> parameters, bool scheduledBuild, List<string> arguments)
+			=> throw new NotSupportedException();
+
+		void ITemplateParameter.GetDefaultParameters(Dictionary<ParameterId, string> parameters, bool scheduledBuild)
+			=> throw new NotSupportedException();
+	}
+
+	/// <summary>
+	/// Allows the user to toggle an option on or off
+	/// </summary>
+	[JsonDiscriminator("Bool")]
+	public class GetTemplateBoolParameterResponse : GetTemplateParameterResponse, ITemplateBoolParameter
+	{
+		/// <inheritdoc/>
+		public ParameterId Id { get; set; }
+
+		/// <inheritdoc/>
+		public string Label { get; set; }
+
+		/// <inheritdoc/>
+		public string? ArgumentIfEnabled { get; set; }
+
+		/// <inheritdoc/>
+		public List<string>? ArgumentsIfEnabled { get; set; }
+
+		/// <inheritdoc/>
+		public string? ArgumentIfDisabled { get; set; }
+
+		/// <inheritdoc/>
+		public List<string>? ArgumentsIfDisabled { get; set; }
+
+		/// <inheritdoc/>
+		public bool Default { get; set; }
+
+		/// <inheritdoc/>
+		public bool? ScheduleOverride { get; set; }
+
+		/// <inheritdoc/>
+		public string? ToolTip { get; set; }
+
+		IReadOnlyList<string>? ITemplateBoolParameter.ArgumentsIfEnabled => ArgumentsIfEnabled;
+		IReadOnlyList<string>? ITemplateBoolParameter.ArgumentsIfDisabled => ArgumentsIfDisabled;
+
+		/// <summary>
+		/// Default constructor
+		/// </summary>
+		public GetTemplateBoolParameterResponse()
+		{
+			Label = String.Empty;
+		}
+
+		/// <summary>
+		/// Constructor
+		/// </summary>
+		public GetTemplateBoolParameterResponse(ITemplateBoolParameter parameter)
+		{
+			Id = parameter.Id;
+			Label = parameter.Label;
+			ArgumentIfEnabled = parameter.ArgumentIfEnabled;
+			ArgumentsIfEnabled = parameter.ArgumentsIfEnabled?.ToList();
+			ArgumentIfDisabled = parameter.ArgumentIfDisabled;
+			ArgumentsIfDisabled = parameter.ArgumentsIfDisabled?.ToList();
+			Default = parameter.Default;
+			ScheduleOverride = parameter.ScheduleOverride;
+			ToolTip = parameter.ToolTip;
+		}
+	}
+
+	/// <summary>
+	/// Free-form text entry parameter
+	/// </summary>
+	[JsonDiscriminator("Text")]
+	public class GetTemplateTextParameterResponse : GetTemplateParameterResponse, ITemplateTextParameter
+	{
+		/// <inheritdoc/>
+		public ParameterId Id { get; set; }
+
+		/// <inheritdoc/>
+		public string Label { get; set; }
+
+		/// <inheritdoc/>
+		public string Argument { get; set; }
+
+		/// <inheritdoc/>
+		public string Default { get; set; }
+
+		/// <inheritdoc/>
+		public string? ScheduleOverride { get; set; }
+
+		/// <inheritdoc/>
+		public string? Hint { get; set; }
+
+		/// <inheritdoc/>
+		public string? Validation { get; set; }
+
+		/// <inheritdoc/>
+		public string? ValidationError { get; set; }
+
+		/// <inheritdoc/>
+		public string? ToolTip { get; set; }
+
+		/// <summary>
+		/// Default constructor
+		/// </summary>
+		public GetTemplateTextParameterResponse()
+		{
+			Label = String.Empty;
+			Argument = String.Empty;
+			Default = String.Empty;
+		}
+
+		/// <summary>
+		/// Constructor
+		/// </summary>
+		public GetTemplateTextParameterResponse(ITemplateTextParameter parameter)
+		{
+			Id = parameter.Id;
+			Label = parameter.Label;
+			Argument = parameter.Argument;
+			Default = parameter.Default;
+			ScheduleOverride = parameter.ScheduleOverride;
+			Hint = parameter.Hint;
+			Validation = parameter.Validation;
+			ValidationError = parameter.ValidationError;
+			ToolTip = parameter.ToolTip;
+		}
+	}
+
+	/// <summary>
+	/// Allows the user to select a value from a constrained list of choices
+	/// </summary>
+	[JsonDiscriminator("List")]
+	public class GetTemplateListParameterResponse : GetTemplateParameterResponse, ITemplateListParameter
+	{
+		/// <inheritdoc/>
+		public string Label { get; }
+
+		/// <inheritdoc/>
+		public TemplateListParameterStyle Style { get; }
+
+		/// <inheritdoc cref="ITemplateListParameter.Items"/>
+		public List<GetTemplateListParameterItemResponse> Items { get; set; }
+
+		/// <inheritdoc/>
+		public string? ToolTip { get; }
+
+		IReadOnlyList<ITemplateListParameterItem> ITemplateListParameter.Items => Items;
+
+		/// <summary>
+		/// Default constructor
+		/// </summary>
+		public GetTemplateListParameterResponse()
+		{
+			Label = String.Empty;
+			Items = new List<GetTemplateListParameterItemResponse>();
+		}
+
+		/// <summary>
+		/// Constructor
+		/// </summary>
+		public GetTemplateListParameterResponse(ITemplateListParameter parameter)
+		{
+			Label = parameter.Label;
+			Style = parameter.Style;
+			Items = parameter.Items.ConvertAll(x => new GetTemplateListParameterItemResponse(x));
+		}
+	}
+
+	/// <summary>
+	/// Possible option for a list parameter
+	/// </summary>
+	public class GetTemplateListParameterItemResponse : ITemplateListParameterItem
+	{
+		/// <inheritdoc/>
+		public ParameterId Id { get; set; }
+
+		/// <inheritdoc/>
+		public string? Group { get; set; }
+
+		/// <inheritdoc/>
+		public string Text { get; set; }
+
+		/// <inheritdoc/>
+		public string? ArgumentIfEnabled { get; set; }
+
+		/// <inheritdoc/>
+		public List<string>? ArgumentsIfEnabled { get; set; }
+
+		/// <inheritdoc/>
+		public string? ArgumentIfDisabled { get; set; }
+
+		/// <inheritdoc/>
+		public List<string>? ArgumentsIfDisabled { get; set; }
+
+		/// <inheritdoc/>
+		public bool Default { get; set; }
+
+		/// <inheritdoc/>
+		public bool? ScheduleOverride { get; set; }
+
+		IReadOnlyList<string>? ITemplateListParameterItem.ArgumentsIfEnabled => ArgumentsIfEnabled;
+		IReadOnlyList<string>? ITemplateListParameterItem.ArgumentsIfDisabled => ArgumentsIfDisabled;
+
+		/// <summary>
+		/// Default constructor
+		/// </summary>
+		public GetTemplateListParameterItemResponse()
+		{
+			Text = String.Empty;
+		}
+
+		/// <summary>
+		/// Constructor
+		/// </summary>
+		public GetTemplateListParameterItemResponse(ITemplateListParameterItem item)
+		{
+			Id = item.Id;
+			Group = item.Group;
+			Text = item.Text;
+			ArgumentIfEnabled = item.ArgumentIfEnabled;
+			ArgumentsIfEnabled = item.ArgumentsIfEnabled?.ToList();
+			ArgumentIfDisabled = item.ArgumentIfDisabled;
+			ArgumentsIfDisabled = item.ArgumentsIfDisabled?.ToList();
+			Default = item.Default;
+			ScheduleOverride = item.ScheduleOverride;
+		}
+	}
 }
diff --git a/Engine/Source/Programs/Horde/Plugins/Build/HordeServer.Build/Streams/StreamConfig.cs b/Engine/Source/Programs/Horde/Plugins/Build/HordeServer.Build/Streams/StreamConfig.cs
--- a/Engine/Source/Programs/Horde/Plugins/Build/HordeServer.Build/Streams/StreamConfig.cs
+++ b/Engine/Source/Programs/Horde/Plugins/Build/HordeServer.Build/Streams/StreamConfig.cs
@@ -313,7 +313,7 @@ namespace HordeServer.Streams
 			foreach (TemplateRefConfig template in Templates)
 			{
 				HashSet<ParameterId> parameterIds = new HashSet<ParameterId>();
-				foreach (ParameterData parameter in template.Parameters)
+				foreach (TemplateParameterConfig parameter in template.Parameters)
 				{
 					parameter.PostLoad(parameterIds);
 				}
