Horde: Move configuration for schedules into a separate file.

#rnx

diff --git a/Engine/Source/Programs/Horde/Plugins/Build/HordeServer.Build/Jobs/Schedules/ScheduleConfig.cs b/Engine/Source/Programs/Horde/Plugins/Build/HordeServer.Build/Jobs/Schedules/ScheduleConfig.cs
new file mode 100644
--- /dev/null
+++ b/Engine/Source/Programs/Horde/Plugins/Build/HordeServer.Build/Jobs/Schedules/ScheduleConfig.cs
@@ -0,0 +1,308 @@
+﻿// Copyright Epic Games, Inc. All Rights Reserved.
+
+using System.ComponentModel.DataAnnotations;
+using System.Diagnostics.CodeAnalysis;
+using System.Text.Json;
+using System.Text.Json.Serialization;
+using EpicGames.Core;
+using EpicGames.Horde.Commits;
+using EpicGames.Horde.Jobs.Templates;
+using HordeServer.Acls;
+using HordeServer.Streams;
+using HordeServer.Utilities;
+
+namespace HordeServer.Jobs.Schedules
+{
+
+	/// <summary>
+	/// Parameters to create a new schedule
+	/// </summary>
+	public class ScheduleConfig
+	{
+		/// <summary>
+		/// Roles to impersonate for this schedule
+		/// </summary>
+		public List<AclClaimConfig>? Claims { get; set; }
+
+		/// <summary>
+		/// Whether the schedule should be enabled
+		/// </summary>
+		public bool Enabled { get; set; } = true;
+
+		/// <summary>
+		/// Maximum number of builds that can be active at once
+		/// </summary>
+		public int MaxActive { get; set; }
+
+		/// <summary>
+		/// Maximum number of changes the schedule can fall behind head revision. If greater than zero, builds will be triggered for every submitted changelist until the backlog is this size.
+		/// </summary>
+		public int MaxChanges { get; set; }
+
+		/// <summary>
+		/// Whether the build requires a change to be submitted
+		/// </summary>
+		public bool RequireSubmittedChange { get; set; } = true;
+
+		/// <summary>
+		/// Gate allowing the schedule to trigger
+		/// </summary>
+		public ScheduleGateConfig? Gate { get; set; }
+
+		/// <summary>
+		/// Commit tags for this schedule
+		/// </summary>
+		public List<CommitTag> Commits { get; set; } = new List<CommitTag>();
+
+		/// <summary>
+		/// The types of changes to run for
+		/// </summary>
+		[Obsolete("Use Commits instead")]
+		public List<ChangeContentFlags>? Filter { get; set; }
+
+		/// <summary>
+		/// Merges the list of commit tags with the legacy filter field
+		/// </summary>
+#pragma warning disable CS0618 // Type or member is obsolete
+		internal List<CommitTag> GetAllCommitTags()
+		{
+			List<CommitTag> commits = Commits;
+			if (Filter != null && Filter.Count > 0)
+			{
+				commits = new List<CommitTag>(commits);
+				if (Filter.Contains(ChangeContentFlags.ContainsCode))
+				{
+					commits.Add(CommitTag.Code);
+				}
+				if (Filter.Contains(ChangeContentFlags.ContainsContent))
+				{
+					commits.Add(CommitTag.Content);
+				}
+			}
+			return commits;
+		}
+#pragma warning restore CS0618 // Type or member is obsolete
+
+		/// <summary>
+		/// Files that should cause the job to trigger
+		/// </summary>
+		public List<string>? Files { get; set; }
+
+		/// <summary>
+		/// Parameters for the template
+		/// </summary>
+		public Dictionary<string, string> TemplateParameters { get; set; } = new Dictionary<string, string>();
+
+		/// <summary>
+		/// New patterns for the schedule
+		/// </summary>
+		public List<SchedulePatternConfig> Patterns { get; set; } = new List<SchedulePatternConfig>();
+
+		/// <summary>
+		/// Get the next time that the schedule will trigger
+		/// </summary>
+		/// <param name="lastTimeUtc">Last time at which the schedule triggered</param>
+		/// <param name="timeZone">Timezone to evaluate the trigger</param>
+		/// <returns>Next time at which the schedule will trigger</returns>
+		public DateTime? GetNextTriggerTimeUtc(DateTime lastTimeUtc, TimeZoneInfo timeZone)
+		{
+			DateTime? nextTriggerTimeUtc = null;
+			foreach (SchedulePatternConfig pattern in Patterns)
+			{
+				DateTime patternTriggerTime = pattern.GetNextTriggerTimeUtc(lastTimeUtc, timeZone);
+				if (nextTriggerTimeUtc == null || patternTriggerTime < nextTriggerTimeUtc)
+				{
+					nextTriggerTimeUtc = patternTriggerTime;
+				}
+			}
+			return nextTriggerTimeUtc;
+		}
+	}
+
+	/// <summary>
+	/// Gate allowing a schedule to trigger.
+	/// </summary>
+	public class ScheduleGateConfig
+	{
+		/// <summary>
+		/// The template containing the dependency
+		/// </summary>
+		[Required]
+		public TemplateId TemplateId { get; set; }
+
+		/// <summary>
+		/// Target to wait for
+		/// </summary>
+		[Required]
+		public string Target { get; set; } = String.Empty;
+	}
+
+	/// <summary>
+	/// Parameters to create a new schedule
+	/// </summary>
+	public class SchedulePatternConfig
+	{
+		/// <summary>
+		/// Days of the week to run this schedule on. If null, the schedule will run every day.
+		/// </summary>
+		public List<DayOfWeek>? DaysOfWeek { get; set; }
+
+		/// <summary>
+		/// Time during the day for the first schedule to trigger. Measured in minutes from midnight.
+		/// </summary>
+		public ScheduleTimeOfDay MinTime { get; set; } = new ScheduleTimeOfDay(0);
+
+		/// <summary>
+		/// Time during the day for the last schedule to trigger. Measured in minutes from midnight.
+		/// </summary>
+		public ScheduleTimeOfDay? MaxTime { get; set; }
+
+		/// <summary>
+		/// Interval between each schedule triggering
+		/// </summary>
+		public ScheduleInterval? Interval { get; set; }
+
+		/// <summary>
+		/// Constructor
+		/// </summary>
+		public SchedulePatternConfig()
+		{
+		}
+
+		/// <summary>
+		/// Constructor
+		/// </summary>
+		/// <param name="daysOfWeek">Which days of the week the schedule should run</param>
+		/// <param name="minTime">Time during the day for the first schedule to trigger. Measured in minutes from midnight.</param>
+		/// <param name="maxTime">Time during the day for the last schedule to trigger. Measured in minutes from midnight.</param>
+		/// <param name="interval">Interval between each schedule triggering</param>
+		public SchedulePatternConfig(List<DayOfWeek>? daysOfWeek, int minTime, int? maxTime, int? interval)
+		{
+			DaysOfWeek = daysOfWeek;
+			MinTime = new ScheduleTimeOfDay(minTime);
+			MaxTime = (maxTime != null) ? new ScheduleTimeOfDay(maxTime.Value) : null;
+			Interval = (interval != null) ? new ScheduleInterval(interval.Value) : null;
+		}
+
+		/// <summary>
+		/// Calculates the trigger index based on the given time in minutes
+		/// </summary>
+		/// <param name="lastTimeUtc">Time for the last trigger</param>
+		/// <param name="timeZone">The timezone for running the schedule</param>
+		/// <returns>Index of the trigger</returns>
+		public DateTime GetNextTriggerTimeUtc(DateTime lastTimeUtc, TimeZoneInfo timeZone)
+		{
+			// Convert last time into the correct timezone for running the scheule
+			DateTimeOffset lastTime = TimeZoneInfo.ConvertTime((DateTimeOffset)lastTimeUtc, timeZone);
+
+			// Get the base time (ie. the start of this day) for anchoring the schedule
+			DateTimeOffset baseTime = new DateTimeOffset(lastTime.Year, lastTime.Month, lastTime.Day, 0, 0, 0, lastTime.Offset);
+			for (; ; )
+			{
+				if (DaysOfWeek == null || DaysOfWeek.Contains(baseTime.DayOfWeek))
+				{
+					// Get the last time in minutes from the start of this day
+					int lastTimeMinutes = (int)(lastTime - baseTime).TotalMinutes;
+
+					// Get the time of the first trigger of this day. If the last time is less than this, this is the next trigger.
+					if (lastTimeMinutes < MinTime.Minutes)
+					{
+						return baseTime.AddMinutes(MinTime.Minutes).UtcDateTime;
+					}
+
+					// Otherwise, get the time for the last trigger in the day.
+					if (Interval != null && Interval.Minutes > 0)
+					{
+						int actualMaxTime = MaxTime?.Minutes ?? ((24 * 60) - 1);
+						if (lastTimeMinutes < actualMaxTime)
+						{
+							int lastIndex = (lastTimeMinutes - MinTime.Minutes) / Interval.Minutes;
+							int nextIndex = lastIndex + 1;
+
+							int nextTimeMinutes = MinTime.Minutes + (nextIndex * Interval.Minutes);
+							if (nextTimeMinutes <= actualMaxTime)
+							{
+								return baseTime.AddMinutes(nextTimeMinutes).UtcDateTime;
+							}
+						}
+					}
+				}
+				baseTime = baseTime.AddDays(1.0);
+			}
+		}
+	}
+
+	/// <summary>
+	/// Time of day value for a schedule
+	/// </summary>
+	[JsonSchemaString]
+	[JsonConverter(typeof(ScheduleTimeOfDayJsonConverter))]
+	public record class ScheduleTimeOfDay(int Minutes)
+	{
+		/// <summary>
+		/// Parse a string as a time of day
+		/// </summary>
+		[return: NotNullIfNotNull("text")]
+		public static ScheduleTimeOfDay? Parse(string? text)
+			=> (text != null) ? new ScheduleTimeOfDay((int)TimeOfDayJsonConverter.Parse(text).TotalMinutes) : null;
+	}
+
+	class ScheduleTimeOfDayJsonConverter : JsonConverter<ScheduleTimeOfDay>
+	{
+		/// <inheritdoc/>
+		public override ScheduleTimeOfDay? Read(ref Utf8JsonReader reader, Type typeToConvert, JsonSerializerOptions options)
+		{
+			if (reader.TokenType == JsonTokenType.Number)
+			{
+				return new ScheduleTimeOfDay(reader.GetInt32());
+			}
+			else
+			{
+				return ScheduleTimeOfDay.Parse(reader.GetString());
+			}
+		}
+
+		/// <inheritdoc/>
+		public override void Write(Utf8JsonWriter writer, ScheduleTimeOfDay value, JsonSerializerOptions options)
+		{
+			writer.WriteNumberValue(value.Minutes);
+		}
+	}
+
+	/// <summary>
+	/// Time of day value for a schedule
+	/// </summary>
+	[JsonSchemaString]
+	[JsonConverter(typeof(ScheduleIntervalJsonConverter))]
+	public record class ScheduleInterval(int Minutes)
+	{
+		/// <summary>
+		/// Parse a string as a time of day
+		/// </summary>
+		[return: NotNullIfNotNull("text")]
+		public static ScheduleInterval? Parse(string? text)
+			=> (text != null) ? new ScheduleInterval((int)IntervalJsonConverter.Parse(text).TotalMinutes) : null;
+	}
+
+	class ScheduleIntervalJsonConverter : JsonConverter<ScheduleInterval>
+	{
+		/// <inheritdoc/>
+		public override ScheduleInterval? Read(ref Utf8JsonReader reader, Type typeToConvert, JsonSerializerOptions options)
+		{
+			if (reader.TokenType == JsonTokenType.Number)
+			{
+				return new ScheduleInterval(reader.GetInt32());
+			}
+			else
+			{
+				return ScheduleInterval.Parse(reader.GetString());
+			}
+		}
+
+		/// <inheritdoc/>
+		public override void Write(Utf8JsonWriter writer, ScheduleInterval value, JsonSerializerOptions options)
+		{
+			writer.WriteNumberValue(value.Minutes);
+		}
+	}
+}
diff --git a/Engine/Source/Programs/Horde/Plugins/Build/HordeServer.Build/Jobs/Schedules/ScheduleMessages.cs b/Engine/Source/Programs/Horde/Plugins/Build/HordeServer.Build/Jobs/Schedules/ScheduleMessages.cs
new file mode 100644
--- /dev/null
+++ b/Engine/Source/Programs/Horde/Plugins/Build/HordeServer.Build/Jobs/Schedules/ScheduleMessages.cs
@@ -0,0 +1,153 @@
+// Copyright Epic Games, Inc. All Rights Reserved.
+
+using EpicGames.Horde.Commits;
+using EpicGames.Horde.Jobs;
+using HordeServer.Streams;
+
+namespace HordeServer.Jobs.Schedules
+{
+	/// <summary>
+	/// Response describing a schedule
+	/// </summary>
+	public class GetScheduleResponse
+	{
+		/// <summary>
+		/// Whether the schedule is currently enabled
+		/// </summary>
+		public bool Enabled { get; set; }
+
+		/// <summary>
+		/// Maximum number of scheduled jobs at once
+		/// </summary>
+		public int MaxActive { get; set; }
+
+		/// <summary>
+		/// Maximum number of changes the schedule can fall behind head revision. If greater than zero, builds will be triggered for every submitted changelist until the backlog is this size.
+		/// </summary>
+		public int MaxChanges { get; set; }
+
+		/// <summary>
+		/// Whether the build requires a change to be submitted
+		/// </summary>
+		public bool RequireSubmittedChange { get; set; }
+
+		/// <summary>
+		/// Gate for this schedule to trigger
+		/// </summary>
+		public ScheduleGateConfig? Gate { get; set; }
+
+		/// <summary>
+		/// Which commits to run this job for
+		/// </summary>
+		public List<CommitTag>? Commits { get; set; }
+
+		/// <summary>
+		/// The types of changes to run for
+		/// </summary>
+		[Obsolete("Use Commits instead")]
+		public List<ChangeContentFlags>? Filter { get; set; }
+
+		/// <summary>
+		/// Parameters for the template
+		/// </summary>
+		public Dictionary<string, string> TemplateParameters { get; set; }
+
+		/// <summary>
+		/// New patterns for the schedule
+		/// </summary>
+		public List<SchedulePatternConfig> Patterns { get; set; }
+
+		/// <summary>
+		/// Last changelist number that this was triggered for
+		/// </summary>
+		[Obsolete("Use LastTriggerCommitId instead")]
+		public int LastTriggerChange
+		{
+			get => _lastTriggerChange ?? _lastTriggerCommitId?.TryGetPerforceChange() ?? 0;
+			set => _lastTriggerChange = value;
+		}
+		int? _lastTriggerChange;
+
+		/// <summary>
+		/// Last changelist number that this was triggered for
+		/// </summary>
+		public CommitIdWithOrder? LastTriggerCommitId
+		{
+			get => _lastTriggerCommitId ?? CommitIdWithOrder.FromPerforceChange(_lastTriggerChange);
+			set => _lastTriggerCommitId = value;
+		}
+		CommitIdWithOrder? _lastTriggerCommitId;
+
+		/// <summary>
+		/// Last time that the schedule was triggered
+		/// </summary>
+		public DateTimeOffset LastTriggerTime { get; set; }
+
+		/// <summary>
+		/// Next trigger times for schedule
+		/// </summary>
+		public List<DateTime> NextTriggerTimesUTC { get; set; }
+
+		/// <summary>
+		/// List of active jobs
+		/// </summary>
+		public List<JobId> ActiveJobs { get; set; }
+
+		/// <summary>
+		/// Constructor
+		/// </summary>
+		/// <param name="schedule">Schedule to construct from</param>
+		/// <param name="schedulerTimeZone">The scheduler time zone</param>
+		public GetScheduleResponse(ITemplateSchedule schedule, TimeZoneInfo schedulerTimeZone)
+		{
+			Enabled = schedule.Config.Enabled;
+			MaxActive = schedule.Config.MaxActive;
+			MaxChanges = schedule.Config.MaxChanges;
+			RequireSubmittedChange = schedule.Config.RequireSubmittedChange;
+			Gate = schedule.Config.Gate;
+			Commits = schedule.Config.Commits;
+#pragma warning disable CS0618 // Type or member is obsolete
+			Filter = schedule.Config.Filter;
+#pragma warning restore CS0618 // Type or member is obsolete
+			TemplateParameters = schedule.Config.TemplateParameters;
+			Patterns = schedule.Config.Patterns;
+			LastTriggerCommitId = schedule.LastTriggerCommitId;
+			LastTriggerTime = schedule.LastTriggerTimeUtc;
+			ActiveJobs = new List<JobId>(schedule.ActiveJobs);
+
+			DateTime curTime = schedule.LastTriggerTimeUtc;
+			NextTriggerTimesUTC = new List<DateTime>();
+			for (int i = 0; i < 16; i++)
+			{
+				DateTime? nextTime = schedule.Config.GetNextTriggerTimeUtc(curTime, schedulerTimeZone);
+				if (nextTime == null)
+				{
+					break;
+				}
+
+				curTime = nextTime.Value;
+				NextTriggerTimesUTC.Add(curTime);
+			}
+		}
+	}
+
+	/// <summary>
+	/// Response describing when a schedule is expected to trigger
+	/// </summary>
+	public class GetScheduleForecastResponse
+	{
+		/// <summary>
+		/// Next trigger times
+		/// </summary>
+		public List<DateTime> Times { get; set; }
+
+		/// <summary>
+		/// Constructor
+		/// </summary>
+		/// <param name="times">List of trigger times</param>
+		public GetScheduleForecastResponse(List<DateTime> times)
+		{
+			Times = times;
+		}
+	}
+}
diff --git a/Engine/Source/Programs/Horde/Plugins/Build/HordeServer.Build/Jobs/Schedules/Schedules.cs b/Engine/Source/Programs/Horde/Plugins/Build/HordeServer.Build/Jobs/Schedules/Schedules.cs
deleted file mode 100644
diff --git a/Engine/Source/Programs/Horde/Plugins/Build/HordeServer.Build/Streams/IStream.cs b/Engine/Source/Programs/Horde/Plugins/Build/HordeServer.Build/Streams/IStream.cs
--- a/Engine/Source/Programs/Horde/Plugins/Build/HordeServer.Build/Streams/IStream.cs
+++ b/Engine/Source/Programs/Horde/Plugins/Build/HordeServer.Build/Streams/IStream.cs
@@ -6,6 +6,7 @@ using EpicGames.Horde.Jobs.Templates;
 using EpicGames.Horde.Streams;
 using EpicGames.Horde.Users;
 using Horde.Common.Rpc;
+using HordeServer.Jobs.Schedules;
 
 namespace HordeServer.Streams
 {
diff --git a/Engine/Source/Programs/Horde/Plugins/Build/HordeServer.Build/Streams/StreamCollection.cs b/Engine/Source/Programs/Horde/Plugins/Build/HordeServer.Build/Streams/StreamCollection.cs
--- a/Engine/Source/Programs/Horde/Plugins/Build/HordeServer.Build/Streams/StreamCollection.cs
+++ b/Engine/Source/Programs/Horde/Plugins/Build/HordeServer.Build/Streams/StreamCollection.cs
@@ -6,6 +6,7 @@ using EpicGames.Horde.Jobs;
 using EpicGames.Horde.Jobs.Templates;
 using EpicGames.Horde.Streams;
 using EpicGames.Horde.Users;
+using HordeServer.Jobs.Schedules;
 using HordeServer.Server;
 using HordeServer.Utilities;
 using MongoDB.Bson.Serialization.Attributes;
diff --git a/Engine/Source/Programs/Horde/Plugins/Build/HordeServer.Build/Streams/StreamConfig.cs b/Engine/Source/Programs/Horde/Plugins/Build/HordeServer.Build/Streams/StreamConfig.cs
--- a/Engine/Source/Programs/Horde/Plugins/Build/HordeServer.Build/Streams/StreamConfig.cs
+++ b/Engine/Source/Programs/Horde/Plugins/Build/HordeServer.Build/Streams/StreamConfig.cs
@@ -28,6 +28,7 @@ using HordeServer.Issues;
 using HordeServer.Jobs;
 using HordeServer.Jobs.Bisect;
 using HordeServer.Jobs.Graphs;
+using HordeServer.Jobs.Schedules;
 using HordeServer.Jobs.Templates;
 using HordeServer.Notifications;
 using HordeServer.Projects;
@@ -942,298 +943,6 @@ namespace HordeServer.Streams
 		}
 	}
 
-	/// <summary>
-	/// Parameters to create a new schedule
-	/// </summary>
-	public class SchedulePatternConfig
-	{
-		/// <summary>
-		/// Days of the week to run this schedule on. If null, the schedule will run every day.
-		/// </summary>
-		public List<DayOfWeek>? DaysOfWeek { get; set; }
-
-		/// <summary>
-		/// Time during the day for the first schedule to trigger. Measured in minutes from midnight.
-		/// </summary>
-		public ScheduleTimeOfDay MinTime { get; set; } = new ScheduleTimeOfDay(0);
-
-		/// <summary>
-		/// Time during the day for the last schedule to trigger. Measured in minutes from midnight.
-		/// </summary>
-		public ScheduleTimeOfDay? MaxTime { get; set; }
-
-		/// <summary>
-		/// Interval between each schedule triggering
-		/// </summary>
-		public ScheduleInterval? Interval { get; set; }
-
-		/// <summary>
-		/// Constructor
-		/// </summary>
-		public SchedulePatternConfig()
-		{
-		}
-
-		/// <summary>
-		/// Constructor
-		/// </summary>
-		/// <param name="daysOfWeek">Which days of the week the schedule should run</param>
-		/// <param name="minTime">Time during the day for the first schedule to trigger. Measured in minutes from midnight.</param>
-		/// <param name="maxTime">Time during the day for the last schedule to trigger. Measured in minutes from midnight.</param>
-		/// <param name="interval">Interval between each schedule triggering</param>
-		public SchedulePatternConfig(List<DayOfWeek>? daysOfWeek, int minTime, int? maxTime, int? interval)
-		{
-			DaysOfWeek = daysOfWeek;
-			MinTime = new ScheduleTimeOfDay(minTime);
-			MaxTime = (maxTime != null) ? new ScheduleTimeOfDay(maxTime.Value) : null;
-			Interval = (interval != null) ? new ScheduleInterval(interval.Value) : null;
-		}
-
-		/// <summary>
-		/// Calculates the trigger index based on the given time in minutes
-		/// </summary>
-		/// <param name="lastTimeUtc">Time for the last trigger</param>
-		/// <param name="timeZone">The timezone for running the schedule</param>
-		/// <returns>Index of the trigger</returns>
-		public DateTime GetNextTriggerTimeUtc(DateTime lastTimeUtc, TimeZoneInfo timeZone)
-		{
-			// Convert last time into the correct timezone for running the scheule
-			DateTimeOffset lastTime = TimeZoneInfo.ConvertTime((DateTimeOffset)lastTimeUtc, timeZone);
-
-			// Get the base time (ie. the start of this day) for anchoring the schedule
-			DateTimeOffset baseTime = new DateTimeOffset(lastTime.Year, lastTime.Month, lastTime.Day, 0, 0, 0, lastTime.Offset);
-			for (; ; )
-			{
-				if (DaysOfWeek == null || DaysOfWeek.Contains(baseTime.DayOfWeek))
-				{
-					// Get the last time in minutes from the start of this day
-					int lastTimeMinutes = (int)(lastTime - baseTime).TotalMinutes;
-
-					// Get the time of the first trigger of this day. If the last time is less than this, this is the next trigger.
-					if (lastTimeMinutes < MinTime.Minutes)
-					{
-						return baseTime.AddMinutes(MinTime.Minutes).UtcDateTime;
-					}
-
-					// Otherwise, get the time for the last trigger in the day.
-					if (Interval != null && Interval.Minutes > 0)
-					{
-						int actualMaxTime = MaxTime?.Minutes ?? ((24 * 60) - 1);
-						if (lastTimeMinutes < actualMaxTime)
-						{
-							int lastIndex = (lastTimeMinutes - MinTime.Minutes) / Interval.Minutes;
-							int nextIndex = lastIndex + 1;
-
-							int nextTimeMinutes = MinTime.Minutes + (nextIndex * Interval.Minutes);
-							if (nextTimeMinutes <= actualMaxTime)
-							{
-								return baseTime.AddMinutes(nextTimeMinutes).UtcDateTime;
-							}
-						}
-					}
-				}
-				baseTime = baseTime.AddDays(1.0);
-			}
-		}
-	}
-
-	/// <summary>
-	/// Time of day value for a schedule
-	/// </summary>
-	[JsonSchemaString]
-	[JsonConverter(typeof(ScheduleTimeOfDayJsonConverter))]
-	public record class ScheduleTimeOfDay(int Minutes)
-	{
-		/// <summary>
-		/// Parse a string as a time of day
-		/// </summary>
-		[return: NotNullIfNotNull("text")]
-		public static ScheduleTimeOfDay? Parse(string? text)
-			=> (text != null) ? new ScheduleTimeOfDay((int)TimeOfDayJsonConverter.Parse(text).TotalMinutes) : null;
-	}
-
-	class ScheduleTimeOfDayJsonConverter : JsonConverter<ScheduleTimeOfDay>
-	{
-		/// <inheritdoc/>
-		public override ScheduleTimeOfDay? Read(ref Utf8JsonReader reader, Type typeToConvert, JsonSerializerOptions options)
-		{
-			if (reader.TokenType == JsonTokenType.Number)
-			{
-				return new ScheduleTimeOfDay(reader.GetInt32());
-			}
-			else
-			{
-				return ScheduleTimeOfDay.Parse(reader.GetString());
-			}
-		}
-
-		/// <inheritdoc/>
-		public override void Write(Utf8JsonWriter writer, ScheduleTimeOfDay value, JsonSerializerOptions options)
-		{
-			writer.WriteNumberValue(value.Minutes);
-		}
-	}
-
-	/// <summary>
-	/// Time of day value for a schedule
-	/// </summary>
-	[JsonSchemaString]
-	[JsonConverter(typeof(ScheduleIntervalJsonConverter))]
-	public record class ScheduleInterval(int Minutes)
-	{
-		/// <summary>
-		/// Parse a string as a time of day
-		/// </summary>
-		[return: NotNullIfNotNull("text")]
-		public static ScheduleInterval? Parse(string? text)
-			=> (text != null) ? new ScheduleInterval((int)IntervalJsonConverter.Parse(text).TotalMinutes) : null;
-	}
-
-	class ScheduleIntervalJsonConverter : JsonConverter<ScheduleInterval>
-	{
-		/// <inheritdoc/>
-		public override ScheduleInterval? Read(ref Utf8JsonReader reader, Type typeToConvert, JsonSerializerOptions options)
-		{
-			if (reader.TokenType == JsonTokenType.Number)
-			{
-				return new ScheduleInterval(reader.GetInt32());
-			}
-			else
-			{
-				return ScheduleInterval.Parse(reader.GetString());
-			}
-		}
-
-		/// <inheritdoc/>
-		public override void Write(Utf8JsonWriter writer, ScheduleInterval value, JsonSerializerOptions options)
-		{
-			writer.WriteNumberValue(value.Minutes);
-		}
-	}
-
-	/// <summary>
-	/// Gate allowing a schedule to trigger.
-	/// </summary>
-	public class ScheduleGateConfig
-	{
-		/// <summary>
-		/// The template containing the dependency
-		/// </summary>
-		[Required]
-		public TemplateId TemplateId { get; set; }
-
-		/// <summary>
-		/// Target to wait for
-		/// </summary>
-		[Required]
-		public string Target { get; set; } = String.Empty;
-	}
-
-	/// <summary>
-	/// Parameters to create a new schedule
-	/// </summary>
-	public class ScheduleConfig
-	{
-		/// <summary>
-		/// Roles to impersonate for this schedule
-		/// </summary>
-		public List<AclClaimConfig>? Claims { get; set; }
-
-		/// <summary>
-		/// Whether the schedule should be enabled
-		/// </summary>
-		public bool Enabled { get; set; } = true;
-
-		/// <summary>
-		/// Maximum number of builds that can be active at once
-		/// </summary>
-		public int MaxActive { get; set; }
-
-		/// <summary>
-		/// Maximum number of changes the schedule can fall behind head revision. If greater than zero, builds will be triggered for every submitted changelist until the backlog is this size.
-		/// </summary>
-		public int MaxChanges { get; set; }
-
-		/// <summary>
-		/// Whether the build requires a change to be submitted
-		/// </summary>
-		public bool RequireSubmittedChange { get; set; } = true;
-
-		/// <summary>
-		/// Gate allowing the schedule to trigger
-		/// </summary>
-		public ScheduleGateConfig? Gate { get; set; }
-
-		/// <summary>
-		/// Commit tags for this schedule
-		/// </summary>
-		public List<CommitTag> Commits { get; set; } = new List<CommitTag>();
-
-		/// <summary>
-		/// The types of changes to run for
-		/// </summary>
-		[Obsolete("Use Commits instead")]
-		public List<ChangeContentFlags>? Filter { get; set; }
-
-		/// <summary>
-		/// Merges the list of commit tags with the legacy filter field
-		/// </summary>
-#pragma warning disable CS0618 // Type or member is obsolete
-		internal List<CommitTag> GetAllCommitTags()
-		{
-			List<CommitTag> commits = Commits;
-			if (Filter != null && Filter.Count > 0)
-			{
-				commits = new List<CommitTag>(commits);
-				if (Filter.Contains(ChangeContentFlags.ContainsCode))
-				{
-					commits.Add(CommitTag.Code);
-				}
-				if (Filter.Contains(ChangeContentFlags.ContainsContent))
-				{
-					commits.Add(CommitTag.Content);
-				}
-			}
-			return commits;
-		}
-#pragma warning restore CS0618 // Type or member is obsolete
-
-		/// <summary>
-		/// Files that should cause the job to trigger
-		/// </summary>
-		public List<string>? Files { get; set; }
-
-		/// <summary>
-		/// Parameters for the template
-		/// </summary>
-		public Dictionary<string, string> TemplateParameters { get; set; } = new Dictionary<string, string>();
-
-		/// <summary>
-		/// New patterns for the schedule
-		/// </summary>
-		public List<SchedulePatternConfig> Patterns { get; set; } = new List<SchedulePatternConfig>();
-
-		/// <summary>
-		/// Get the next time that the schedule will trigger
-		/// </summary>
-		/// <param name="lastTimeUtc">Last time at which the schedule triggered</param>
-		/// <param name="timeZone">Timezone to evaluate the trigger</param>
-		/// <returns>Next time at which the schedule will trigger</returns>
-		public DateTime? GetNextTriggerTimeUtc(DateTime lastTimeUtc, TimeZoneInfo timeZone)
-		{
-			DateTime? nextTriggerTimeUtc = null;
-			foreach (SchedulePatternConfig pattern in Patterns)
-			{
-				DateTime patternTriggerTime = pattern.GetNextTriggerTimeUtc(lastTimeUtc, timeZone);
-				if (nextTriggerTimeUtc == null || patternTriggerTime < nextTriggerTimeUtc)
-				{
-					nextTriggerTimeUtc = patternTriggerTime;
-				}
-			}
-			return nextTriggerTimeUtc;
-		}
-	}
-
 	/// <summary>
 	/// Configuration for allocating access tokens for each job
 	/// </summary>
