Horde: Add explicit HTTP response types for schedule gates and patterns.

#rnx

diff --git a/Engine/Source/Programs/Horde/Plugins/Build/HordeServer.Build/Jobs/Schedules/ScheduleMessages.cs b/Engine/Source/Programs/Horde/Plugins/Build/HordeServer.Build/Jobs/Schedules/ScheduleMessages.cs
--- a/Engine/Source/Programs/Horde/Plugins/Build/HordeServer.Build/Jobs/Schedules/ScheduleMessages.cs
+++ b/Engine/Source/Programs/Horde/Plugins/Build/HordeServer.Build/Jobs/Schedules/ScheduleMessages.cs
@@ -2,6 +2,7 @@
 
 using EpicGames.Horde.Commits;
 using EpicGames.Horde.Jobs;
+using EpicGames.Horde.Jobs.Templates;
 using HordeServer.Streams;
 
 namespace HordeServer.Jobs.Schedules
@@ -34,7 +35,7 @@ namespace HordeServer.Jobs.Schedules
 		/// <summary>
 		/// Gate for this schedule to trigger
 		/// </summary>
-		public ScheduleGateConfig? Gate { get; set; }
+		public GetScheduleGateResponse? Gate { get; set; }
 
 		/// <summary>
 		/// Which commits to run this job for
@@ -55,7 +56,7 @@ namespace HordeServer.Jobs.Schedules
 		/// <summary>
 		/// New patterns for the schedule
 		/// </summary>
-		public List<SchedulePatternConfig> Patterns { get; set; }
+		public List<GetSchedulePatternResponse> Patterns { get; set; }
 
 		/// <summary>
 		/// Last changelist number that this was triggered for
@@ -104,13 +105,13 @@ namespace HordeServer.Jobs.Schedules
 			MaxActive = schedule.Config.MaxActive;
 			MaxChanges = schedule.Config.MaxChanges;
 			RequireSubmittedChange = schedule.Config.RequireSubmittedChange;
-			Gate = schedule.Config.Gate;
+			Gate = (schedule.Config.Gate == null) ? null : new GetScheduleGateResponse(schedule.Config.Gate);
 			Commits = schedule.Config.Commits;
 #pragma warning disable CS0618 // Type or member is obsolete
 			Filter = schedule.Config.Filter;
 #pragma warning restore CS0618 // Type or member is obsolete
 			TemplateParameters = schedule.Config.TemplateParameters;
-			Patterns = schedule.Config.Patterns;
+			Patterns = schedule.Config.Patterns.ConvertAll(x => new GetSchedulePatternResponse(x));
 			LastTriggerCommitId = schedule.LastTriggerCommitId;
 			LastTriggerTime = schedule.LastTriggerTimeUtc;
 			ActiveJobs = new List<JobId>(schedule.ActiveJobs);
@@ -131,6 +132,83 @@ namespace HordeServer.Jobs.Schedules
 		}
 	}
 
+	/// <summary>
+	/// Gate allowing a schedule to trigger.
+	/// </summary>
+	public class GetScheduleGateResponse
+	{
+		/// <summary>
+		/// The template containing the dependency
+		/// </summary>
+		public TemplateId TemplateId { get; set; }
+
+		/// <summary>
+		/// Target to wait for
+		/// </summary>
+		public string Target { get; set; }
+
+		/// <summary>
+		/// Default constructor
+		/// </summary>
+		public GetScheduleGateResponse()
+		{
+			Target = String.Empty;
+		}
+
+		/// <summary>
+		/// Constructor
+		/// </summary>
+		public GetScheduleGateResponse(ScheduleGateConfig config)
+		{
+			TemplateId = config.TemplateId;
+			Target = config.Target;
+		}
+	}
+
+	/// <summary>
+	/// Parameters to create a new schedule
+	/// </summary>
+	public class GetSchedulePatternResponse
+	{
+		/// <summary>
+		/// Days of the week to run this schedule on. If null, the schedule will run every day.
+		/// </summary>
+		public List<DayOfWeek>? DaysOfWeek { get; set; }
+
+		/// <summary>
+		/// Time during the day for the first schedule to trigger. Measured in minutes from midnight.
+		/// </summary>
+		public ScheduleTimeOfDay MinTime { get; set; } = new ScheduleTimeOfDay(0);
+
+		/// <summary>
+		/// Time during the day for the last schedule to trigger. Measured in minutes from midnight.
+		/// </summary>
+		public ScheduleTimeOfDay? MaxTime { get; set; }
+
+		/// <summary>
+		/// Interval between each schedule triggering
+		/// </summary>
+		public ScheduleInterval? Interval { get; set; }
+
+		/// <summary>
+		/// Constructor
+		/// </summary>
+		public GetSchedulePatternResponse()
+		{
+		}
+
+		/// <summary>
+		/// Constructor
+		/// </summary>
+		public GetSchedulePatternResponse(SchedulePatternConfig config)
+		{
+			DaysOfWeek = config.DaysOfWeek?.ToList();
+			MinTime = config.MinTime;
+			MaxTime = config.MaxTime;
+			Interval = config.Interval;
+		}
+	}
+
 	/// <summary>
 	/// Response describing when a schedule is expected to trigger
 	/// </summary>
