Horde: Add interfaces for accessing template schedules.

#rnx

diff --git a/Engine/Source/Programs/Horde/Plugins/Build/HordeServer.Build.Tests/Jobs/Schedules/SchedulerTests.cs b/Engine/Source/Programs/Horde/Plugins/Build/HordeServer.Build.Tests/Jobs/Schedules/SchedulerTests.cs
--- a/Engine/Source/Programs/Horde/Plugins/Build/HordeServer.Build.Tests/Jobs/Schedules/SchedulerTests.cs
+++ b/Engine/Source/Programs/Horde/Plugins/Build/HordeServer.Build.Tests/Jobs/Schedules/SchedulerTests.cs
@@ -15,6 +15,7 @@ using HordeServer.Server;
 using HordeServer.Streams;
 using HordeServer.Users;
 using Microsoft.Extensions.DependencyInjection;
+using Moq;
 
 namespace HordeServer.Tests.Jobs.Schedules
 {
@@ -129,8 +130,10 @@ namespace HordeServer.Tests.Jobs.Schedules
 			DateTime startTime = new DateTime(2021, 1, 1, 12, 0, 0, DateTimeKind.Utc); // Friday Jan 1, 2021 
 			Clock.UtcNow = startTime;
 
-			ScheduleConfig schedule = new ScheduleConfig { RequireSubmittedChange = false };
-			schedule.Patterns.Add(new SchedulePatternConfig(new List<DayOfWeek> { DayOfWeek.Friday, DayOfWeek.Sunday }, 13 * 60, null, null));
+			Mock<ISchedule> scheduleMock = new Mock<ISchedule>(MockBehavior.Strict);
+			scheduleMock.SetupGet(x => x.Patterns).Returns([new SchedulePatternConfig(new List<DayOfWeek> { DayOfWeek.Friday, DayOfWeek.Sunday }, 13 * 60, null, null)]);
+			
+			ISchedule schedule = scheduleMock.Object;
 
 			DateTime? nextTime = schedule.GetNextTriggerTimeUtc(startTime, TimeZoneInfo.Utc);
 			Assert.AreEqual(startTime + TimeSpan.FromHours(1.0), nextTime!.Value);
@@ -151,8 +154,11 @@ namespace HordeServer.Tests.Jobs.Schedules
 			DateTime startTime = new DateTime(2021, 1, 1, 12, 0, 0, DateTimeKind.Utc); // Friday Jan 1, 2021 
 			Clock.UtcNow = startTime;
 
-			ScheduleConfig schedule = new ScheduleConfig { RequireSubmittedChange = false };
-			schedule.Patterns.Add(new SchedulePatternConfig(null, 13 * 60, 14 * 60, 15));
+			Mock<ISchedule> scheduleMock = new Mock<ISchedule>(MockBehavior.Strict);
+			scheduleMock.SetupGet(x => x.RequireSubmittedChange).Returns(true);
+			scheduleMock.SetupGet(x => x.Patterns).Returns([new SchedulePatternConfig(null, 13 * 60, 14 * 60, 15)]);
+
+			ISchedule schedule = scheduleMock.Object;
 
 			DateTime? nextTime = schedule.GetNextTriggerTimeUtc(startTime, TimeZoneInfo.Utc);
 			Assert.AreEqual(startTime + TimeSpan.FromHours(1.0), nextTime!.Value);
@@ -179,9 +185,15 @@ namespace HordeServer.Tests.Jobs.Schedules
 			DateTime startTime = new DateTime(2021, 1, 1, 0, 0, 0, DateTimeKind.Utc); // Friday Jan 1, 2021 
 			Clock.UtcNow = startTime;
 
-			ScheduleConfig schedule = new ScheduleConfig { RequireSubmittedChange = false };
-			schedule.Patterns.Add(new SchedulePatternConfig(null, 11 * 60, 0, 0));
-			schedule.Patterns.Add(new SchedulePatternConfig(null, 19 * 60, 0, 0));
+			Mock<ISchedule> scheduleMock = new Mock<ISchedule>(MockBehavior.Strict);
+			scheduleMock.SetupGet(x => x.RequireSubmittedChange).Returns(false);
+			scheduleMock.SetupGet(x => x.Patterns).Returns(
+				[
+					new SchedulePatternConfig(null, 11 * 60, 0, 0),
+					new SchedulePatternConfig(null, 19 * 60, 0, 0)
+				]);
+
+			ISchedule schedule = scheduleMock.Object;
 
 			DateTime? nextTime = schedule.GetNextTriggerTimeUtc(startTime, TimeZoneInfo.Utc);
 			Assert.AreEqual(startTime + TimeSpan.FromHours(11), nextTime!.Value);
@@ -250,7 +262,7 @@ namespace HordeServer.Tests.Jobs.Schedules
 			GlobalConfig.CurrentValue.Plugins.GetBuildConfig().TryGetStream(StreamId, out streamConfig);
 
 			IStream stream2 = (await StreamCollection.GetAsync(streamConfig!))!;
-			ITemplateSchedule schedule2 = stream2.Templates.First().Value.Schedule!;
+			ISchedule schedule2 = stream2.Templates.First().Value.Schedule!;
 			Assert.AreEqual(102, schedule2.LastTriggerCommitId!.GetPerforceChange());
 			Assert.AreEqual(Clock.UtcNow, schedule2.LastTriggerTimeUtc);
 
diff --git a/Engine/Source/Programs/Horde/Plugins/Build/HordeServer.Build/Jobs/Schedules/ISchedule.cs b/Engine/Source/Programs/Horde/Plugins/Build/HordeServer.Build/Jobs/Schedules/ISchedule.cs
new file mode 100644
--- /dev/null
+++ b/Engine/Source/Programs/Horde/Plugins/Build/HordeServer.Build/Jobs/Schedules/ISchedule.cs
@@ -0,0 +1,208 @@
+// Copyright Epic Games, Inc. All Rights Reserved.
+
+using EpicGames.Horde.Commits;
+using EpicGames.Horde.Jobs;
+using EpicGames.Horde.Jobs.Templates;
+using HordeServer.Acls;
+
+namespace HordeServer.Jobs.Schedules
+{
+	/// <summary>
+	/// Schedule for a template
+	/// </summary>
+	public interface ISchedule
+	{
+		/// <summary>
+		/// Whether the schedule should be enabled
+		/// </summary>
+		bool Enabled { get; }
+
+		/// <summary>
+		/// Maximum number of builds that can be active at once
+		/// </summary>
+		int MaxActive { get; }
+
+		/// <summary>
+		/// Maximum number of changes the schedule can fall behind head revision. If greater than zero, builds will be triggered for every submitted changelist until the backlog is this size.
+		/// </summary>
+		int MaxChanges { get; }
+
+		/// <summary>
+		/// Whether the build requires a change to be submitted
+		/// </summary>
+		bool RequireSubmittedChange { get; }
+
+		/// <summary>
+		/// Gate allowing the schedule to trigger
+		/// </summary>
+		IScheduleGate? Gate { get; }
+
+		/// <summary>
+		/// Commit tags for this schedule
+		/// </summary>
+		IReadOnlyList<CommitTag> Commits { get; }
+
+		/// <summary>
+		/// Roles to impersonate for this schedule
+		/// </summary>
+		IReadOnlyList<AclClaimConfig>? Claims { get; }
+
+		/// <summary>
+		/// Last changelist number that this was triggered for
+		/// </summary>
+		CommitIdWithOrder? LastTriggerCommitId { get; }
+
+		/// <summary>
+		/// Gets the last trigger time, in UTC
+		/// </summary>
+		DateTime LastTriggerTimeUtc { get; }
+
+		/// <summary>
+		/// List of jobs that are currently active
+		/// </summary>
+		IReadOnlyList<JobId> ActiveJobs { get; }
+
+		/// <summary>
+		/// Patterns for starting this scheduled job
+		/// </summary>
+		IReadOnlyList<ISchedulePattern> Patterns { get; }
+
+		/// <summary>
+		/// Files that should cause the job to trigger
+		/// </summary>
+		IReadOnlyList<string>? Files { get; }
+
+		/// <summary>
+		/// Parameters for the template
+		/// </summary>
+		IReadOnlyDictionary<string, string> TemplateParameters { get; }
+	}
+
+	/// <summary>
+	/// Required gate for starting a schedule
+	/// </summary>
+	public interface IScheduleGate
+	{
+		/// <summary>
+		/// The template containing the dependency
+		/// </summary>
+		public TemplateId TemplateId { get; }
+
+		/// <summary>
+		/// Target to wait for
+		/// </summary>
+		public string Target { get; }
+	}
+
+	/// <summary>
+	/// Pattern for executing a schedule
+	/// </summary>
+	public interface ISchedulePattern
+	{
+		/// <summary>
+		/// Days of the week to run this schedule on. If null, the schedule will run every day.
+		/// </summary>
+		IReadOnlyList<DayOfWeek>? DaysOfWeek { get; }
+
+		/// <summary>
+		/// Time during the day for the first schedule to trigger. Measured in minutes from midnight.
+		/// </summary>
+		ScheduleTimeOfDay MinTime { get; }
+
+		/// <summary>
+		/// Time during the day for the last schedule to trigger. Measured in minutes from midnight.
+		/// </summary>
+		ScheduleTimeOfDay? MaxTime { get; }
+
+		/// <summary>
+		/// Interval between each schedule triggering
+		/// </summary>
+		ScheduleInterval? Interval { get; }
+	}
+
+	/// <summary>
+	/// Extension methods for schedules
+	/// </summary>
+	public static class ScheduleExtensions
+	{
+		/// <summary>
+		/// Gets the next trigger time for a schedule
+		/// </summary>
+		/// <param name="schedule"></param>
+		/// <param name="timeZone"></param>
+		/// <returns></returns>
+		public static DateTime? GetNextTriggerTimeUtc(this ISchedule schedule, TimeZoneInfo timeZone)
+		{
+			return schedule.GetNextTriggerTimeUtc(schedule.LastTriggerTimeUtc, timeZone);
+		}
+
+		/// <summary>
+		/// Get the next time that the schedule will trigger
+		/// </summary>
+		/// <param name="schedule">Schedule to query</param>
+		/// <param name="lastTimeUtc">Last time at which the schedule triggered</param>
+		/// <param name="timeZone">Timezone to evaluate the trigger</param>
+		/// <returns>Next time at which the schedule will trigger</returns>
+		public static DateTime? GetNextTriggerTimeUtc(this ISchedule schedule, DateTime lastTimeUtc, TimeZoneInfo timeZone)
+		{
+			DateTime? nextTriggerTimeUtc = null;
+			foreach (ISchedulePattern pattern in schedule.Patterns)
+			{
+				DateTime patternTriggerTime = pattern.GetNextTriggerTimeUtc(lastTimeUtc, timeZone);
+				if (nextTriggerTimeUtc == null || patternTriggerTime < nextTriggerTimeUtc)
+				{
+					nextTriggerTimeUtc = patternTriggerTime;
+				}
+			}
+			return nextTriggerTimeUtc;
+		}
+
+		/// <summary>
+		/// Calculates the trigger index based on the given time in minutes
+		/// </summary>
+		/// <param name="pattern">Pattern to query</param>
+		/// <param name="lastTimeUtc">Time for the last trigger</param>
+		/// <param name="timeZone">The timezone for running the schedule</param>
+		/// <returns>Index of the trigger</returns>
+		public static DateTime GetNextTriggerTimeUtc(this ISchedulePattern pattern, DateTime lastTimeUtc, TimeZoneInfo timeZone)
+		{
+			// Convert last time into the correct timezone for running the scheule
+			DateTimeOffset lastTime = TimeZoneInfo.ConvertTime((DateTimeOffset)lastTimeUtc, timeZone);
+
+			// Get the base time (ie. the start of this day) for anchoring the schedule
+			DateTimeOffset baseTime = new DateTimeOffset(lastTime.Year, lastTime.Month, lastTime.Day, 0, 0, 0, lastTime.Offset);
+			for (; ; )
+			{
+				if (pattern.DaysOfWeek == null || pattern.DaysOfWeek.Contains(baseTime.DayOfWeek))
+				{
+					// Get the last time in minutes from the start of this day
+					int lastTimeMinutes = (int)(lastTime - baseTime).TotalMinutes;
+
+					// Get the time of the first trigger of this day. If the last time is less than this, this is the next trigger.
+					if (lastTimeMinutes < pattern.MinTime.Minutes)
+					{
+						return baseTime.AddMinutes(pattern.MinTime.Minutes).UtcDateTime;
+					}
+
+					// Otherwise, get the time for the last trigger in the day.
+					if (pattern.Interval != null && pattern.Interval.Minutes > 0)
+					{
+						int actualMaxTime = pattern.MaxTime?.Minutes ?? ((24 * 60) - 1);
+						if (lastTimeMinutes < actualMaxTime)
+						{
+							int lastIndex = (lastTimeMinutes - pattern.MinTime.Minutes) / pattern.Interval.Minutes;
+							int nextIndex = lastIndex + 1;
+
+							int nextTimeMinutes = pattern.MinTime.Minutes + (nextIndex * pattern.Interval.Minutes);
+							if (nextTimeMinutes <= actualMaxTime)
+							{
+								return baseTime.AddMinutes(nextTimeMinutes).UtcDateTime;
+							}
+						}
+					}
+				}
+				baseTime = baseTime.AddDays(1.0);
+			}
+		}
+	}
+}
diff --git a/Engine/Source/Programs/Horde/Plugins/Build/HordeServer.Build/Jobs/Schedules/ScheduleConfig.cs b/Engine/Source/Programs/Horde/Plugins/Build/HordeServer.Build/Jobs/Schedules/ScheduleConfig.cs
--- a/Engine/Source/Programs/Horde/Plugins/Build/HordeServer.Build/Jobs/Schedules/ScheduleConfig.cs
+++ b/Engine/Source/Programs/Horde/Plugins/Build/HordeServer.Build/Jobs/Schedules/ScheduleConfig.cs
@@ -13,17 +13,11 @@ using HordeServer.Utilities;
 
 namespace HordeServer.Jobs.Schedules
 {
-
 	/// <summary>
 	/// Parameters to create a new schedule
 	/// </summary>
 	public class ScheduleConfig
 	{
-		/// <summary>
-		/// Roles to impersonate for this schedule
-		/// </summary>
-		public List<AclClaimConfig>? Claims { get; set; }
-
 		/// <summary>
 		/// Whether the schedule should be enabled
 		/// </summary>
@@ -54,6 +48,11 @@ namespace HordeServer.Jobs.Schedules
 		/// </summary>
 		public List<CommitTag> Commits { get; set; } = new List<CommitTag>();
 
+		/// <summary>
+		/// Roles to impersonate for this schedule
+		/// </summary>
+		public List<AclClaimConfig>? Claims { get; set; }
+
 		/// <summary>
 		/// The types of changes to run for
 		/// </summary>
@@ -97,32 +96,12 @@ namespace HordeServer.Jobs.Schedules
 		/// New patterns for the schedule
 		/// </summary>
 		public List<SchedulePatternConfig> Patterns { get; set; } = new List<SchedulePatternConfig>();
-
-		/// <summary>
-		/// Get the next time that the schedule will trigger
-		/// </summary>
-		/// <param name="lastTimeUtc">Last time at which the schedule triggered</param>
-		/// <param name="timeZone">Timezone to evaluate the trigger</param>
-		/// <returns>Next time at which the schedule will trigger</returns>
-		public DateTime? GetNextTriggerTimeUtc(DateTime lastTimeUtc, TimeZoneInfo timeZone)
-		{
-			DateTime? nextTriggerTimeUtc = null;
-			foreach (SchedulePatternConfig pattern in Patterns)
-			{
-				DateTime patternTriggerTime = pattern.GetNextTriggerTimeUtc(lastTimeUtc, timeZone);
-				if (nextTriggerTimeUtc == null || patternTriggerTime < nextTriggerTimeUtc)
-				{
-					nextTriggerTimeUtc = patternTriggerTime;
-				}
-			}
-			return nextTriggerTimeUtc;
-		}
 	}
 
 	/// <summary>
 	/// Gate allowing a schedule to trigger.
 	/// </summary>
-	public class ScheduleGateConfig
+	public class ScheduleGateConfig : IScheduleGate
 	{
 		/// <summary>
 		/// The template containing the dependency
@@ -140,7 +119,7 @@ namespace HordeServer.Jobs.Schedules
 	/// <summary>
 	/// Parameters to create a new schedule
 	/// </summary>
-	public class SchedulePatternConfig
+	public class SchedulePatternConfig : ISchedulePattern
 	{
 		/// <summary>
 		/// Days of the week to run this schedule on. If null, the schedule will run every day.
@@ -162,6 +141,8 @@ namespace HordeServer.Jobs.Schedules
 		/// </summary>
 		public ScheduleInterval? Interval { get; set; }
 
+		IReadOnlyList<DayOfWeek>? ISchedulePattern.DaysOfWeek => DaysOfWeek;
+
 		/// <summary>
 		/// Constructor
 		/// </summary>
@@ -183,53 +164,6 @@ namespace HordeServer.Jobs.Schedules
 			MaxTime = (maxTime != null) ? new ScheduleTimeOfDay(maxTime.Value) : null;
 			Interval = (interval != null) ? new ScheduleInterval(interval.Value) : null;
 		}
-
-		/// <summary>
-		/// Calculates the trigger index based on the given time in minutes
-		/// </summary>
-		/// <param name="lastTimeUtc">Time for the last trigger</param>
-		/// <param name="timeZone">The timezone for running the schedule</param>
-		/// <returns>Index of the trigger</returns>
-		public DateTime GetNextTriggerTimeUtc(DateTime lastTimeUtc, TimeZoneInfo timeZone)
-		{
-			// Convert last time into the correct timezone for running the scheule
-			DateTimeOffset lastTime = TimeZoneInfo.ConvertTime((DateTimeOffset)lastTimeUtc, timeZone);
-
-			// Get the base time (ie. the start of this day) for anchoring the schedule
-			DateTimeOffset baseTime = new DateTimeOffset(lastTime.Year, lastTime.Month, lastTime.Day, 0, 0, 0, lastTime.Offset);
-			for (; ; )
-			{
-				if (DaysOfWeek == null || DaysOfWeek.Contains(baseTime.DayOfWeek))
-				{
-					// Get the last time in minutes from the start of this day
-					int lastTimeMinutes = (int)(lastTime - baseTime).TotalMinutes;
-
-					// Get the time of the first trigger of this day. If the last time is less than this, this is the next trigger.
-					if (lastTimeMinutes < MinTime.Minutes)
-					{
-						return baseTime.AddMinutes(MinTime.Minutes).UtcDateTime;
-					}
-
-					// Otherwise, get the time for the last trigger in the day.
-					if (Interval != null && Interval.Minutes > 0)
-					{
-						int actualMaxTime = MaxTime?.Minutes ?? ((24 * 60) - 1);
-						if (lastTimeMinutes < actualMaxTime)
-						{
-							int lastIndex = (lastTimeMinutes - MinTime.Minutes) / Interval.Minutes;
-							int nextIndex = lastIndex + 1;
-
-							int nextTimeMinutes = MinTime.Minutes + (nextIndex * Interval.Minutes);
-							if (nextTimeMinutes <= actualMaxTime)
-							{
-								return baseTime.AddMinutes(nextTimeMinutes).UtcDateTime;
-							}
-						}
-					}
-				}
-				baseTime = baseTime.AddDays(1.0);
-			}
-		}
 	}
 
 	/// <summary>
diff --git a/Engine/Source/Programs/Horde/Plugins/Build/HordeServer.Build/Jobs/Schedules/ScheduleMessages.cs b/Engine/Source/Programs/Horde/Plugins/Build/HordeServer.Build/Jobs/Schedules/ScheduleMessages.cs
--- a/Engine/Source/Programs/Horde/Plugins/Build/HordeServer.Build/Jobs/Schedules/ScheduleMessages.cs
+++ b/Engine/Source/Programs/Horde/Plugins/Build/HordeServer.Build/Jobs/Schedules/ScheduleMessages.cs
@@ -1,9 +1,9 @@
 // Copyright Epic Games, Inc. All Rights Reserved.
 
+using EpicGames.Core;
 using EpicGames.Horde.Commits;
 using EpicGames.Horde.Jobs;
 using EpicGames.Horde.Jobs.Templates;
-using HordeServer.Streams;
 
 namespace HordeServer.Jobs.Schedules
 {
@@ -42,12 +42,6 @@ namespace HordeServer.Jobs.Schedules
 		/// </summary>
 		public List<CommitTag>? Commits { get; set; }
 
-		/// <summary>
-		/// The types of changes to run for
-		/// </summary>
-		[Obsolete("Use Commits instead")]
-		public List<ChangeContentFlags>? Filter { get; set; }
-
 		/// <summary>
 		/// Parameters for the template
 		/// </summary>
@@ -94,24 +88,32 @@ namespace HordeServer.Jobs.Schedules
 		/// </summary>
 		public List<JobId> ActiveJobs { get; set; }
 
+		/// <summary>
+		/// Default constructor
+		/// </summary>
+		public GetScheduleResponse()
+		{
+			Patterns = new List<GetSchedulePatternResponse>();
+			TemplateParameters = new Dictionary<string, string>();
+			NextTriggerTimesUTC = new List<DateTime>();
+			ActiveJobs = new List<JobId>();
+		}
+
 		/// <summary>
 		/// Constructor
 		/// </summary>
 		/// <param name="schedule">Schedule to construct from</param>
 		/// <param name="schedulerTimeZone">The scheduler time zone</param>
-		public GetScheduleResponse(ITemplateSchedule schedule, TimeZoneInfo schedulerTimeZone)
+		public GetScheduleResponse(ISchedule schedule, TimeZoneInfo schedulerTimeZone)
 		{
-			Enabled = schedule.Config.Enabled;
-			MaxActive = schedule.Config.MaxActive;
-			MaxChanges = schedule.Config.MaxChanges;
-			RequireSubmittedChange = schedule.Config.RequireSubmittedChange;
-			Gate = (schedule.Config.Gate == null) ? null : new GetScheduleGateResponse(schedule.Config.Gate);
-			Commits = schedule.Config.Commits;
-#pragma warning disable CS0618 // Type or member is obsolete
-			Filter = schedule.Config.Filter;
-#pragma warning restore CS0618 // Type or member is obsolete
-			TemplateParameters = schedule.Config.TemplateParameters;
-			Patterns = schedule.Config.Patterns.ConvertAll(x => new GetSchedulePatternResponse(x));
+			Enabled = schedule.Enabled;
+			MaxActive = schedule.MaxActive;
+			MaxChanges = schedule.MaxChanges;
+			RequireSubmittedChange = schedule.RequireSubmittedChange;
+			Gate = (schedule.Gate == null) ? null : new GetScheduleGateResponse(schedule.Gate);
+			Commits = schedule.Commits?.ToList();
+			TemplateParameters = schedule.TemplateParameters.ToDictionary();
+			Patterns = schedule.Patterns.ConvertAll(x => new GetSchedulePatternResponse(x));
 			LastTriggerCommitId = schedule.LastTriggerCommitId;
 			LastTriggerTime = schedule.LastTriggerTimeUtc;
 			ActiveJobs = new List<JobId>(schedule.ActiveJobs);
@@ -120,7 +122,7 @@ namespace HordeServer.Jobs.Schedules
 			NextTriggerTimesUTC = new List<DateTime>();
 			for (int i = 0; i < 16; i++)
 			{
-				DateTime? nextTime = schedule.Config.GetNextTriggerTimeUtc(curTime, schedulerTimeZone);
+				DateTime? nextTime = schedule.GetNextTriggerTimeUtc(curTime, schedulerTimeZone);
 				if (nextTime == null)
 				{
 					break;
@@ -158,10 +160,10 @@ namespace HordeServer.Jobs.Schedules
 		/// <summary>
 		/// Constructor
 		/// </summary>
-		public GetScheduleGateResponse(ScheduleGateConfig config)
+		public GetScheduleGateResponse(IScheduleGate gate)
 		{
-			TemplateId = config.TemplateId;
-			Target = config.Target;
+			TemplateId = gate.TemplateId;
+			Target = gate.Target;
 		}
 	}
 
@@ -200,12 +202,12 @@ namespace HordeServer.Jobs.Schedules
 		/// <summary>
 		/// Constructor
 		/// </summary>
-		public GetSchedulePatternResponse(SchedulePatternConfig config)
+		public GetSchedulePatternResponse(ISchedulePattern pattern)
 		{
-			DaysOfWeek = config.DaysOfWeek?.ToList();
-			MinTime = config.MinTime;
-			MaxTime = config.MaxTime;
-			Interval = config.Interval;
+			DaysOfWeek = pattern.DaysOfWeek?.ToList();
+			MinTime = pattern.MinTime;
+			MaxTime = pattern.MaxTime;
+			Interval = pattern.Interval;
 		}
 	}
 
diff --git a/Engine/Source/Programs/Horde/Plugins/Build/HordeServer.Build/Jobs/Schedules/ScheduleService.cs b/Engine/Source/Programs/Horde/Plugins/Build/HordeServer.Build/Jobs/Schedules/ScheduleService.cs
--- a/Engine/Source/Programs/Horde/Plugins/Build/HordeServer.Build/Jobs/Schedules/ScheduleService.cs
+++ b/Engine/Source/Programs/Horde/Plugins/Build/HordeServer.Build/Jobs/Schedules/ScheduleService.cs
@@ -272,14 +272,14 @@ namespace HordeServer.Jobs.Schedules
 				return false;
 			}
 
-			ITemplateSchedule? schedule = templateRef.Schedule;
+			ISchedule? schedule = templateRef.Schedule;
 			if (schedule == null)
 			{
 				return false;
 			}
 
 			ForwardingLogger scheduleLogger = new ForwardingLogger(_logger, GetAuditLog(streamId, templateId));
-			if (!schedule.Config.Enabled)
+			if (!schedule.Enabled)
 			{
 				scheduleLogger.LogDebug("Schedule is disabled. No builds started.");
 				return false;
@@ -340,12 +340,12 @@ namespace HordeServer.Jobs.Schedules
 		/// <param name="scheduleLogger"></param>
 		/// <param name="cancellationToken"></param>
 		/// <returns>Async task</returns>
-		private async Task TriggerAsync(IStream stream, TemplateId templateId, ITemplateRef templateRef, ITemplateSchedule schedule, int numActiveJobs, DateTime utcNow, ILogger scheduleLogger, CancellationToken cancellationToken)
+		private async Task TriggerAsync(IStream stream, TemplateId templateId, ITemplateRef templateRef, ISchedule schedule, int numActiveJobs, DateTime utcNow, ILogger scheduleLogger, CancellationToken cancellationToken)
 		{
 			// Check we're not already at the maximum number of allowed jobs
-			if (schedule.Config.MaxActive != 0 && numActiveJobs >= schedule.Config.MaxActive)
+			if (schedule.MaxActive != 0 && numActiveJobs >= schedule.MaxActive)
 			{
-				scheduleLogger.LogInformation("Skipping trigger of {StreamId} template {TemplateId} - already have maximum number of jobs running ({NumJobs})", stream.Id, templateId, schedule.Config.MaxActive);
+				scheduleLogger.LogInformation("Skipping trigger of {StreamId} template {TemplateId} - already have maximum number of jobs running ({NumJobs})", stream.Id, templateId, schedule.MaxActive);
 				foreach (JobId jobId in schedule.ActiveJobs)
 				{
 					_logger.LogInformation("Active job for {StreamId} template {TemplateId}: {JobId}", stream.Id, templateId, jobId);
@@ -355,20 +355,20 @@ namespace HordeServer.Jobs.Schedules
 
 			// Minimum changelist number, inclusive
 			CommitIdWithOrder? minCommitId = schedule.LastTriggerCommitId;
-			bool includeMinCommitId = !schedule.Config.RequireSubmittedChange;
+			bool includeMinCommitId = !schedule.RequireSubmittedChange;
 
 			// Maximum changelist number, exclusive
 			CommitIdWithOrder? maxCommitId = null;
 
 			// Get the maximum number of changes to trigger
 			int maxNewChanges = 1;
-			if (schedule.Config.MaxChanges != 0)
+			if (schedule.MaxChanges != 0)
 			{
-				maxNewChanges = schedule.Config.MaxChanges;
+				maxNewChanges = schedule.MaxChanges;
 			}
-			if (schedule.Config.MaxActive != 0)
+			if (schedule.MaxActive != 0)
 			{
-				maxNewChanges = Math.Min(maxNewChanges, schedule.Config.MaxActive - numActiveJobs);
+				maxNewChanges = Math.Min(maxNewChanges, schedule.MaxActive - numActiveJobs);
 			}
 
 			// Create a timer to limit the amount we look back through P4 history
@@ -376,14 +376,14 @@ namespace HordeServer.Jobs.Schedules
 
 			// Create a file filter
 			FileFilter? fileFilter = null;
-			if (schedule.Config.Files != null)
+			if (schedule.Files != null)
 			{
-				fileFilter = new FileFilter(schedule.Config.Files);
+				fileFilter = new FileFilter(schedule.Files);
 			}
 
 			// Cache the Perforce history as we're iterating through changes to improve query performance
 			ICommitCollection commits = _commitService.GetCollection(stream.Config);
-			IAsyncEnumerable<ICommit> commitEnumerable = commits.FindAsync(minCommitId, includeMinCommitId, tags: schedule.Config.GetAllCommitTags(), cancellationToken: cancellationToken);
+			IAsyncEnumerable<ICommit> commitEnumerable = commits.FindAsync(minCommitId, includeMinCommitId, tags: schedule.Commits, cancellationToken: cancellationToken);
 			await using IAsyncEnumerator<ICommit> commitEnumerator = commitEnumerable.GetAsyncEnumerator(cancellationToken);
 
 			// Start as many jobs as possible
@@ -396,9 +396,9 @@ namespace HordeServer.Jobs.Schedules
 				CommitIdWithOrder? commitId;
 				ICommit? commit;
 
-				if (schedule.Config.Gate != null)
+				if (schedule.Gate != null)
 				{
-					commitId = await GetNextChangeForGateAsync(stream.Id, templateId, schedule.Config.Gate, minCommitId, maxCommitId, scheduleLogger, cancellationToken);
+					commitId = await GetNextChangeForGateAsync(stream.Id, templateId, schedule.Gate, minCommitId, maxCommitId, scheduleLogger, cancellationToken);
 					commit = (commitId != null)? await commits.FindAsync(commitId, commitId, 1, null, cancellationToken).FirstOrDefaultAsync(cancellationToken) : null; // May be a change in a different stream
 				}
 				else if (await commitEnumerator.MoveNextAsync(cancellationToken))
@@ -423,14 +423,14 @@ namespace HordeServer.Jobs.Schedules
 					{
 						break;
 					}
-					if (commitId == minCommitId && (schedule.Config.RequireSubmittedChange || triggerChanges.Count > 0))
+					if (commitId == minCommitId && (schedule.RequireSubmittedChange || triggerChanges.Count > 0))
 					{
 						break;
 					}
 				}
 
 				// Adjust the changelist for the desired filter
-				if (commit == null || await ShouldBuildChangeAsync(commit, schedule.Config.GetAllCommitTags(), fileFilter, cancellationToken))
+				if (commit == null || await ShouldBuildChangeAsync(commit, schedule.Commits, fileFilter, cancellationToken))
 				{
 					CommitIdWithOrder codeCommitId = commitId;
 
@@ -508,7 +508,7 @@ namespace HordeServer.Jobs.Schedules
 		/// <param name="fileFilter">Filter for the files to trigger a build</param>
 		/// <param name="cancellationToken">Cancellation token for the operation</param>
 		/// <returns></returns>
-		private async ValueTask<bool> ShouldBuildChangeAsync(ICommit commit, List<CommitTag>? filterTags, FileFilter? fileFilter, CancellationToken cancellationToken)
+		private async ValueTask<bool> ShouldBuildChangeAsync(ICommit commit, IReadOnlyList<CommitTag>? filterTags, FileFilter? fileFilter, CancellationToken cancellationToken)
 		{
 			if (Regex.IsMatch(commit.Description, @"^\s*#\s*skipci", RegexOptions.Multiline))
 			{
@@ -537,7 +537,7 @@ namespace HordeServer.Jobs.Schedules
 		/// <summary>
 		/// Gets the next change to build for a schedule on a gate
 		/// </summary>
-		private async Task<CommitIdWithOrder?> GetNextChangeForGateAsync(StreamId streamId, TemplateId templateRefId, ScheduleGateConfig gate, CommitIdWithOrder? minCommitId, CommitIdWithOrder? maxCommitId, ILogger scheduleLogger, CancellationToken cancellationToken)
+		private async Task<CommitIdWithOrder?> GetNextChangeForGateAsync(StreamId streamId, TemplateId templateRefId, IScheduleGate gate, CommitIdWithOrder? minCommitId, CommitIdWithOrder? maxCommitId, ILogger scheduleLogger, CancellationToken cancellationToken)
 		{
 			for (; ; )
 			{
diff --git a/Engine/Source/Programs/Horde/Plugins/Build/HordeServer.Build/Streams/IStream.cs b/Engine/Source/Programs/Horde/Plugins/Build/HordeServer.Build/Streams/IStream.cs
--- a/Engine/Source/Programs/Horde/Plugins/Build/HordeServer.Build/Streams/IStream.cs
+++ b/Engine/Source/Programs/Horde/Plugins/Build/HordeServer.Build/Streams/IStream.cs
@@ -1,7 +1,5 @@
 // Copyright Epic Games, Inc. All Rights Reserved.
 
-using EpicGames.Horde.Commits;
-using EpicGames.Horde.Jobs;
 using EpicGames.Horde.Jobs.Templates;
 using EpicGames.Horde.Streams;
 using EpicGames.Horde.Users;
@@ -80,7 +78,7 @@ namespace HordeServer.Streams
 		/// <summary>
 		/// List of schedules for this template
 		/// </summary>
-		ITemplateSchedule? Schedule { get; }
+		ISchedule? Schedule { get; }
 
 		/// <summary>
 		/// List of template step states
@@ -88,32 +86,6 @@ namespace HordeServer.Streams
 		IReadOnlyList<ITemplateStep> StepStates { get; }
 	}
 
-	/// <summary>
-	/// Schedule for a template
-	/// </summary>
-	public interface ITemplateSchedule
-	{
-		/// <summary>
-		/// Config for this schedule
-		/// </summary>
-		ScheduleConfig Config { get; }
-
-		/// <summary>
-		/// Last changelist number that this was triggered for
-		/// </summary>
-		CommitIdWithOrder? LastTriggerCommitId { get; }
-
-		/// <summary>
-		/// Gets the last trigger time, in UTC
-		/// </summary>
-		DateTime LastTriggerTimeUtc { get; }
-
-		/// <summary>
-		/// List of jobs that are currently active
-		/// </summary>
-		IReadOnlyList<JobId> ActiveJobs { get; }
-	}
-
 	/// <summary>
 	/// Information about a paused template step
 	/// </summary>
@@ -140,17 +112,6 @@ namespace HordeServer.Streams
 	/// </summary>
 	static class StreamExtensions
 	{
-		/// <summary>
-		/// Gets the next trigger time for a schedule
-		/// </summary>
-		/// <param name="schedule"></param>
-		/// <param name="timeZone"></param>
-		/// <returns></returns>
-		public static DateTime? GetNextTriggerTimeUtc(this ITemplateSchedule schedule, TimeZoneInfo timeZone)
-		{
-			return schedule.Config.GetNextTriggerTimeUtc(schedule.LastTriggerTimeUtc, timeZone);
-		}
-
 		/// <summary>
 		/// Converts to a public response object
 		/// </summary>
diff --git a/Engine/Source/Programs/Horde/Plugins/Build/HordeServer.Build/Streams/StreamCollection.cs b/Engine/Source/Programs/Horde/Plugins/Build/HordeServer.Build/Streams/StreamCollection.cs
--- a/Engine/Source/Programs/Horde/Plugins/Build/HordeServer.Build/Streams/StreamCollection.cs
+++ b/Engine/Source/Programs/Horde/Plugins/Build/HordeServer.Build/Streams/StreamCollection.cs
@@ -6,6 +6,7 @@ using EpicGames.Horde.Jobs;
 using EpicGames.Horde.Jobs.Templates;
 using EpicGames.Horde.Streams;
 using EpicGames.Horde.Users;
+using HordeServer.Acls;
 using HordeServer.Jobs.Schedules;
 using HordeServer.Server;
 using HordeServer.Utilities;
@@ -19,10 +20,86 @@ namespace HordeServer.Streams
 	/// </summary>
 	public class StreamCollection : IStreamCollection
 	{
+		class Stream : IStream
+		{
+			readonly StreamConfig _config;
+			readonly StreamDoc _document;
+			readonly Dictionary<TemplateId, ITemplateRef> _templateRefs = new Dictionary<TemplateId, ITemplateRef>();
+
+			public StreamDoc Document => _document;
+
+			public StreamId Id => _config.Id;
+			public StreamConfig Config => _config;
+			public IReadOnlyDictionary<TemplateId, ITemplateRef> Templates => _templateRefs;
+			public DateTime? PausedUntil => _document.PausedUntil;
+			public string? PauseComment => _document.PauseComment;
+
+			public Stream(StreamConfig config, StreamDoc document)
+			{
+				_config = config;
+				_document = document;
+
+				foreach(TemplateRefConfig templateRefConfig in config.Templates)
+				{
+					TemplateRefDoc templateRefDoc = _document.Templates[templateRefConfig.Id];
+					_templateRefs[templateRefConfig.Id] = new TemplateRef(templateRefConfig, templateRefDoc);
+				}
+			}
+		}
+
+		class TemplateRef : ITemplateRef
+		{
+			readonly TemplateRefConfig _config;
+			readonly TemplateRefDoc _document;
+			readonly Schedule? _schedule;
+
+			public TemplateRef(TemplateRefConfig config, TemplateRefDoc document)
+			{
+				_config = config;
+				_document = document;
+
+				if(_config.Schedule != null)
+				{
+					_schedule = new Schedule(_config.Schedule, _document.Schedule ?? new ScheduleDoc());
+				}
+			}
+
+			public TemplateId Id => _config.Id;
+			public TemplateRefConfig Config => _config;
+			public ISchedule? Schedule => _schedule;
+			public IReadOnlyList<ITemplateStep> StepStates => (IReadOnlyList<ITemplateStep>?)StepStates ?? Array.Empty<ITemplateStep>();
+		}
+
+		class Schedule : ISchedule
+		{
+			readonly ScheduleConfig _config;
+			readonly ScheduleDoc _document;
+
+			public IReadOnlyList<AclClaimConfig>? Claims => _config.Claims;
+			public bool Enabled => _config.Enabled;
+			public int MaxActive => _config.MaxActive;
+			public int MaxChanges => _config.MaxChanges;
+			public bool RequireSubmittedChange => _config.RequireSubmittedChange;
+			public IScheduleGate? Gate => _config.Gate;
+			public IReadOnlyList<CommitTag> Commits => _config.Commits;
+			public CommitIdWithOrder? LastTriggerCommitId => _document.LastTriggerCommitId;
+			public DateTime LastTriggerTimeUtc => _document.LastTriggerTimeUtc;
+			public IReadOnlyList<JobId> ActiveJobs => _document.ActiveJobs;
+			public IReadOnlyList<ISchedulePattern> Patterns => _config.Patterns;
+			public IReadOnlyList<string>? Files => _config.Files;
+			public IReadOnlyDictionary<string, string> TemplateParameters => _config.TemplateParameters;
+
+			public Schedule(ScheduleConfig config, ScheduleDoc document)
+			{
+				_config = config;
+				_document = document;
+			}
+		}
+
 		/// <summary>
 		/// Information about a stream
 		/// </summary>
-		class StreamDoc : IStream
+		class StreamDoc
 		{
 			[BsonRequired, BsonId]
 			public StreamId Id { get; set; }
@@ -43,27 +120,8 @@ namespace HordeServer.Streams
 				Id = id;
 			}
 
-			#region IStream Implementation
-
-			[BsonIgnore]
-			public StreamConfig Config { get; private set; } = null!;
-
-			[BsonIgnore]
-			IReadOnlyDictionary<TemplateId, ITemplateRef>? _cachedTemplates;
-
-			IReadOnlyDictionary<TemplateId, ITemplateRef> IStream.Templates
-			{
-				get
-				{
-					_cachedTemplates ??= Templates.ToDictionary(x => x.Key, x => (ITemplateRef)x.Value);
-					return _cachedTemplates;
-				}
-			}
-
 			public bool PostLoad(StreamConfig config, DateTime utcNow)
 			{
-				Config = config;
-
 				bool replaceDocument = false;
 
 				Dictionary<TemplateId, TemplateRefDoc> templateRefDocs = new Dictionary<TemplateId, TemplateRefDoc>(config.Templates.Count);
@@ -85,37 +143,21 @@ namespace HordeServer.Streams
 
 				return replaceDocument;
 			}
-
-			#endregion
 		}
 
-		class TemplateRefDoc : ITemplateRef
+		class TemplateRefDoc
 		{
 			[BsonIgnoreIfNull]
-			public TemplateScheduleDoc? Schedule { get; set; }
+			public ScheduleDoc? Schedule { get; set; }
 
 			[BsonIgnoreIfNull]
 			public List<TemplateStepDoc>? StepStates { get; set; }
 
-			#region ITemplateRef implementation
-
-			[BsonIgnore]
-			public TemplateId Id { get; private set; }
-
-			[BsonIgnore]
-			public TemplateRefConfig Config { get; private set; } = null!;
-
-			ITemplateSchedule? ITemplateRef.Schedule => Schedule;
-			IReadOnlyList<ITemplateStep> ITemplateRef.StepStates => (IReadOnlyList<ITemplateStep>?)StepStates ?? Array.Empty<ITemplateStep>();
-
 			public bool PostLoad(TemplateRefConfig config, DateTime utcNow)
 			{
-				Id = config.Id;
-				Config = config;
-
 				bool replaceDocument = false;
 
-				if (Config.Schedule == null)
+				if (config.Schedule == null)
 				{
 					if (Schedule != null)
 					{
@@ -127,13 +169,13 @@ namespace HordeServer.Streams
 				{
 					if (Schedule == null)
 					{
-						Schedule = new TemplateScheduleDoc();
+						Schedule = new ScheduleDoc();
 						Schedule.LastTriggerTimeUtc = utcNow;
 						replaceDocument = true;
 					}
 				}
 
-				Schedule?.PostLoad(this);
+				Schedule?.PostLoad();
 
 				if (StepStates != null)
 				{
@@ -143,11 +185,9 @@ namespace HordeServer.Streams
 
 				return replaceDocument;
 			}
-
-			#endregion
 		}
 
-		class TemplateScheduleDoc : ITemplateSchedule
+		class ScheduleDoc
 		{
 			public CommitIdWithOrder LastTriggerCommitId
 			{
@@ -166,20 +206,8 @@ namespace HordeServer.Streams
 			public DateTime LastTriggerTimeUtc { get; set; }
 			public List<JobId> ActiveJobs { get; set; } = new List<JobId>();
 
-			#region ITemplateSchedule implementation
-
-			[BsonIgnore]
-			TemplateRefDoc? _owner;
-
-			[BsonIgnore]
-			public ScheduleConfig Config => _owner?.Config.Schedule!;
-
-			IReadOnlyList<JobId> ITemplateSchedule.ActiveJobs => ActiveJobs;
-
-			public void PostLoad(TemplateRefDoc owner)
+			public void PostLoad()
 			{
-				_owner = owner;
-
 #pragma warning disable CS0618 // Type or member is obsolete
 				if (LastTriggerTime.HasValue)
 				{
@@ -188,8 +216,6 @@ namespace HordeServer.Streams
 				}
 #pragma warning restore CS0618 // Type or member is obsolete
 			}
-
-			#endregion
 		}
 
 		class TemplateStepDoc : ITemplateStep
@@ -229,7 +255,10 @@ namespace HordeServer.Streams
 
 		/// <inheritdoc/>
 		public async Task<IStream> GetAsync(StreamConfig streamConfig, CancellationToken cancellationToken)
-			=> await GetInternalAsync(streamConfig, cancellationToken);
+		{
+			StreamDoc document = await GetInternalAsync(streamConfig, cancellationToken);
+			return new Stream(streamConfig, document);
+		}
 
 		async Task<StreamDoc> GetInternalAsync(StreamConfig streamConfig, CancellationToken cancellationToken)
 		{
@@ -249,7 +278,17 @@ namespace HordeServer.Streams
 
 		/// <inheritdoc/>
 		public async Task<IReadOnlyList<IStream>> GetAsync(IReadOnlyList<StreamConfig> streamConfigs, CancellationToken cancellationToken)
-			=> await GetInternalAsync(streamConfigs, cancellationToken);
+		{
+			List<Stream> streams = new List<Stream>();
+
+			List<StreamDoc> documents = await GetInternalAsync(streamConfigs, cancellationToken);
+			for(int idx = 0; idx < streamConfigs.Count; idx++)
+			{
+				streams.Add(new Stream(streamConfigs[idx], documents[idx]));
+			}
+
+			return streams;
+		}
 
 		async Task<List<StreamDoc>> GetInternalAsync(IReadOnlyList<StreamConfig> streamConfigs, CancellationToken cancellationToken)
 		{
@@ -291,10 +330,11 @@ namespace HordeServer.Streams
 		/// <inheritdoc/>
 		public async Task<IStream?> TryUpdateTemplateRefAsync(IStream streamInterface, TemplateId templateId, List<UpdateStepStateRequest>? stepStates = null, CancellationToken cancellationToken = default)
 		{
-			StreamDoc stream = (StreamDoc)streamInterface;
+			Stream stream = (Stream)streamInterface;
+			StreamDoc document = stream.Document;
 
 			TemplateRefDoc? templateRef;
-			if (!stream.Templates.TryGetValue(templateId, out templateRef))
+			if (!document.Templates.TryGetValue(templateId, out templateRef))
 			{
 				return null;
 			}
@@ -302,7 +342,7 @@ namespace HordeServer.Streams
 			UpdateDefinitionBuilder<StreamDoc> updateBuilder = Builders<StreamDoc>.Update;
 			List<UpdateDefinition<StreamDoc>> updates = new List<UpdateDefinition<StreamDoc>>();
 
-			Dictionary<TemplateId, TemplateRefDoc> newTemplates = new Dictionary<TemplateId, TemplateRefDoc>(stream.Templates);
+			Dictionary<TemplateId, TemplateRefDoc> newTemplates = new Dictionary<TemplateId, TemplateRefDoc>(document.Templates);
 
 			// clear
 			if (stepStates != null && stepStates.Count == 0)
@@ -372,83 +412,88 @@ namespace HordeServer.Streams
 				return streamInterface;
 			}
 
-			return await TryUpdateStreamAsync(stream, updateBuilder.Combine(updates), cancellationToken);
+			StreamDoc? newStreamDoc = await TryUpdateStreamAsync(document, updateBuilder.Combine(updates), cancellationToken);
+			return (newStreamDoc == null)? null : new Stream(stream.Config, document);
 
 		}
 
 		/// <inheritdoc/>
 		public async Task<IStream?> TryUpdatePauseStateAsync(IStream streamInterface, DateTime? newPausedUntil, string? newPauseComment, CancellationToken cancellationToken)
 		{
-			StreamDoc stream = (StreamDoc)streamInterface;
+			Stream stream = (Stream)streamInterface;
+			StreamDoc document = stream.Document;
 
 			UpdateDefinitionBuilder<StreamDoc> updateBuilder = Builders<StreamDoc>.Update;
 
 			List<UpdateDefinition<StreamDoc>> updates = new List<UpdateDefinition<StreamDoc>>();
-			stream.PausedUntil = newPausedUntil;
-			stream.PauseComment = newPauseComment;
+			document.PausedUntil = newPausedUntil;
+			document.PauseComment = newPauseComment;
 			updates.Add(updateBuilder.Set(x => x.PausedUntil, newPausedUntil));
 			updates.Add(updateBuilder.Set(x => x.PauseComment, newPauseComment));
 
-			return await TryUpdateStreamAsync(stream, updateBuilder.Combine(updates), cancellationToken);
+			StreamDoc? newDocument = await TryUpdateStreamAsync(document, updateBuilder.Combine(updates), cancellationToken);
+			return (newDocument == null)? null : new Stream(stream.Config, newDocument);
 		}
 
 		/// <inheritdoc/>
 		public async Task<IStream?> TryUpdateScheduleTriggerAsync(IStream streamInterface, TemplateId templateId, DateTime? lastTriggerTimeUtc, CommitIdWithOrder? lastTriggerCommitId, List<JobId> newActiveJobs, CancellationToken cancellationToken)
 		{
-			StreamDoc stream = (StreamDoc)streamInterface;
-			TemplateRefDoc template = stream.Templates[templateId];
-			TemplateScheduleDoc schedule = template.Schedule!;
+			Stream stream = (Stream)streamInterface;
+
+			StreamDoc streamDoc = stream.Document;
+			TemplateRefDoc templateRefDoc = streamDoc.Templates[templateId];
+			ScheduleDoc scheduleDoc = templateRefDoc.Schedule!;
 
 			// Build the updates. MongoDB driver cannot parse TemplateRefId in expression tree; need to specify field name explicitly
 			List<UpdateDefinition<StreamDoc>> updates = new List<UpdateDefinition<StreamDoc>>();
-			if (lastTriggerTimeUtc.HasValue && lastTriggerTimeUtc.Value != schedule.LastTriggerTimeUtc)
+			if (lastTriggerTimeUtc.HasValue && lastTriggerTimeUtc.Value != scheduleDoc.LastTriggerTimeUtc)
 			{
-				FieldDefinition<StreamDoc, DateTime> lastTriggerTimeField = $"{nameof(stream.Templates)}.{templateId}.{nameof(template.Schedule)}.{nameof(schedule.LastTriggerTimeUtc)}";
+				FieldDefinition<StreamDoc, DateTime> lastTriggerTimeField = $"{nameof(streamDoc.Templates)}.{templateId}.{nameof(templateRefDoc.Schedule)}.{nameof(scheduleDoc.LastTriggerTimeUtc)}";
 				updates.Add(Builders<StreamDoc>.Update.Set(lastTriggerTimeField, lastTriggerTimeUtc.Value));
 
-				schedule.LastTriggerTimeUtc = lastTriggerTimeUtc.Value;
+				scheduleDoc.LastTriggerTimeUtc = lastTriggerTimeUtc.Value;
 			}
-			if (lastTriggerCommitId != null && lastTriggerCommitId > schedule.LastTriggerCommitId)
+			if (lastTriggerCommitId != null && lastTriggerCommitId > scheduleDoc.LastTriggerCommitId)
 			{
-				FieldDefinition<StreamDoc, string> lastTriggerCommitNameField = $"{nameof(stream.Templates)}.{templateId}.{nameof(template.Schedule)}.{nameof(schedule.LastTriggerCommitName)}";
+				FieldDefinition<StreamDoc, string> lastTriggerCommitNameField = $"{nameof(streamDoc.Templates)}.{templateId}.{nameof(templateRefDoc.Schedule)}.{nameof(scheduleDoc.LastTriggerCommitName)}";
 				updates.Add(Builders<StreamDoc>.Update.Set(lastTriggerCommitNameField, lastTriggerCommitId.Name));
 
-				FieldDefinition<StreamDoc, int> lastTriggerChangeField = $"{nameof(stream.Templates)}.{templateId}.{nameof(template.Schedule)}.LastTriggerChange";
+				FieldDefinition<StreamDoc, int> lastTriggerChangeField = $"{nameof(streamDoc.Templates)}.{templateId}.{nameof(templateRefDoc.Schedule)}.LastTriggerChange";
 				updates.Add(Builders<StreamDoc>.Update.Set(lastTriggerChangeField, lastTriggerCommitId.Order));
 
-				schedule.LastTriggerCommitId = lastTriggerCommitId;
+				scheduleDoc.LastTriggerCommitId = lastTriggerCommitId;
 			}
 			if (newActiveJobs != null)
 			{
-				FieldDefinition<StreamDoc, List<JobId>> field = $"{nameof(stream.Templates)}.{templateId}.{nameof(template.Schedule)}.{nameof(schedule.ActiveJobs)}";
+				FieldDefinition<StreamDoc, List<JobId>> field = $"{nameof(streamDoc.Templates)}.{templateId}.{nameof(templateRefDoc.Schedule)}.{nameof(scheduleDoc.ActiveJobs)}";
 				updates.Add(Builders<StreamDoc>.Update.Set(field, newActiveJobs));
 
-				schedule.ActiveJobs = newActiveJobs;
+				scheduleDoc.ActiveJobs = newActiveJobs;
+			}
+
+			if(updates.Count == 0)
+			{
+				return streamInterface;
 			}
 
-			return (updates.Count == 0) ? streamInterface : await TryUpdateStreamAsync(stream, Builders<StreamDoc>.Update.Combine(updates), cancellationToken);
+			StreamDoc? newStreamDoc = await TryUpdateStreamAsync(streamDoc, Builders<StreamDoc>.Update.Combine(updates), cancellationToken);
+			return (newStreamDoc == null)? null : new Stream(stream.Config, newStreamDoc);
 		}
 
 		/// <summary>
 		/// Update a stream
 		/// </summary>
-		/// <param name="stream">The stream to update</param>
+		/// <param name="streamDoc">The stream to update</param>
 		/// <param name="update">The update definition</param>
 		/// <param name="cancellationToken">Cancellation token for the operation</param>
 		/// <returns>The updated document, or null the update failed</returns>
-		private async Task<StreamDoc?> TryUpdateStreamAsync(StreamDoc stream, UpdateDefinition<StreamDoc> update, CancellationToken cancellationToken)
+		private async Task<StreamDoc?> TryUpdateStreamAsync(StreamDoc streamDoc, UpdateDefinition<StreamDoc> update, CancellationToken cancellationToken)
 		{
-			FilterDefinition<StreamDoc> filter = Builders<StreamDoc>.Filter.Expr(x => x.Id == stream.Id && x.UpdateIndex == stream.UpdateIndex);
-			update = update.Set(x => x.UpdateIndex, stream.UpdateIndex + 1);
+			FilterDefinition<StreamDoc> filter = Builders<StreamDoc>.Filter.Expr(x => x.Id == streamDoc.Id && x.UpdateIndex == streamDoc.UpdateIndex);
+			update = update.Set(x => x.UpdateIndex, streamDoc.UpdateIndex + 1);
 
 			FindOneAndUpdateOptions<StreamDoc> options = new FindOneAndUpdateOptions<StreamDoc> { ReturnDocument = ReturnDocument.After };
-
-			StreamDoc? result = await _streams.FindOneAndUpdateAsync(filter, update, options, cancellationToken);
-			if (result != null)
-			{
-				result.PostLoad(stream.Config, _clock.UtcNow);
-			}
-			return result;
+			return await _streams.FindOneAndUpdateAsync(filter, update, options, cancellationToken);
 		}
 
 		/// <summary>
diff --git a/Engine/Source/Programs/Horde/Plugins/Build/HordeServer.Build/Streams/StreamMessages.cs b/Engine/Source/Programs/Horde/Plugins/Build/HordeServer.Build/Streams/StreamMessages.cs
--- a/Engine/Source/Programs/Horde/Plugins/Build/HordeServer.Build/Streams/StreamMessages.cs
+++ b/Engine/Source/Programs/Horde/Plugins/Build/HordeServer.Build/Streams/StreamMessages.cs
@@ -483,22 +483,22 @@ namespace HordeServer.Streams
 	/// </summary>
 	public class GetTemplateScheduleResponseV2
 	{
-		readonly ITemplateSchedule _schedule;
+		readonly ISchedule _schedule;
 
-		/// <inheritdoc cref="ITemplateSchedule.LastTriggerCommitId"/>
+		/// <inheritdoc cref="ISchedule.LastTriggerCommitId"/>
 		[Obsolete("Use LastTriggerCommit instead")]
 		public int LastTriggerChange => _schedule.LastTriggerCommitId?.GetPerforceChangeOrMinusOne() ?? 0;
 
-		/// <inheritdoc cref="ITemplateSchedule.LastTriggerCommitId"/>
+		/// <inheritdoc cref="ISchedule.LastTriggerCommitId"/>
 		public CommitIdWithOrder? LastTriggerCommit => _schedule.LastTriggerCommitId;
 
-		/// <inheritdoc cref="ITemplateSchedule.LastTriggerTimeUtc"/>
+		/// <inheritdoc cref="ISchedule.LastTriggerTimeUtc"/>
 		public DateTime LastTriggerTimeUtc => _schedule.LastTriggerTimeUtc;
 
-		/// <inheritdoc cref="ITemplateSchedule.LastTriggerTimeUtc"/>
+		/// <inheritdoc cref="ISchedule.LastTriggerTimeUtc"/>
 		public IReadOnlyList<JobId> ActiveJobs => _schedule.ActiveJobs;
 
-		internal GetTemplateScheduleResponseV2(ITemplateSchedule schedule)
+		internal GetTemplateScheduleResponseV2(ISchedule schedule)
 		{
 			_schedule = schedule;
 		}
