Horde: Remove V2 streams endpoints. These endpoints were never used by the dashboard, and include internal config classes that we don't want to publish as part of the public API.

#rnx

diff --git a/Engine/Source/Programs/Horde/Plugins/Build/HordeServer.Build/Streams/StreamMessages.cs b/Engine/Source/Programs/Horde/Plugins/Build/HordeServer.Build/Streams/StreamMessages.cs
--- a/Engine/Source/Programs/Horde/Plugins/Build/HordeServer.Build/Streams/StreamMessages.cs
+++ b/Engine/Source/Programs/Horde/Plugins/Build/HordeServer.Build/Streams/StreamMessages.cs
@@ -17,31 +17,126 @@
 namespace HordeServer.Streams
 {
 	/// <summary>
-	/// Step state update request
+	/// Response describing a stream
 	/// </summary>
-	public class UpdateStepStateRequest
+	public class GetStreamResponse
 	{
 		/// <summary>
-		/// Name of the step
+		/// Unique id of the stream
 		/// </summary>
-		[Required]
-		public string Name { get; set; } = String.Empty;
+		public string Id { get; set; }
 
 		/// <summary>
-		/// User who paused the step
+		/// Unique id of the project containing this stream
 		/// </summary>
-		public string? PausedByUserId { get; set; }
+		public string ProjectId { get; set; }
+
+		/// <summary>
+		/// Name of the stream
+		/// </summary>
+		public string Name { get; set; }
+
+		/// <summary>
+		/// The config file path on the server
+		/// </summary>
+		public string ConfigPath { get; set; } = String.Empty;
+
+		/// <summary>
+		/// Revision of the config file 
+		/// </summary>
+		public string ConfigRevision { get; set; } = String.Empty;
+
+		/// <summary>
+		/// Order to display in the list
+		/// </summary>
+		public int Order { get; set; }
+
+		/// <summary>
+		/// Notification channel for all jobs in this stream
+		/// </summary>
+		public string? NotificationChannel { get; set; }
+
+		/// <summary>
+		/// Notification channel filter for this template. Can be a combination of "Success", "Failure" and "Warnings" separated by pipe characters.
+		/// </summary>
+		public string? NotificationChannelFilter { get; set; }
+
+		/// <summary>
+		/// Channel to post issue triage notifications
+		/// </summary>
+		public string? TriageChannel { get; set; }
+
+		/// <summary>
+		/// Default template for running preflights
+		/// </summary>
+		public string? DefaultPreflightTemplate { get; set; }
+
+		/// <summary>
+		/// Default template to use for preflights
+		/// </summary>
+		public DefaultPreflightConfig? DefaultPreflight { get; set; }
+
+		/// <summary>
+		/// List of tabs to display for this stream
+		/// </summary>
+		public List<TabConfig> Tabs { get; set; }
+
+		/// <summary>
+		/// Map of agent name to type
+		/// </summary>
+		public Dictionary<string, AgentConfig> AgentTypes { get; set; }
+
+		/// <summary>
+		/// Map of workspace name to type
+		/// </summary>
+		public Dictionary<string, WorkspaceConfig>? WorkspaceTypes { get; set; }
+
+		/// <summary>
+		/// Templates for jobs in this stream
+		/// </summary>
+		public List<GetTemplateRefResponse> Templates { get; set; }
+
+		/// <summary>
+		/// Stream paused for new builds until this date
+		/// </summary>
+		public DateTime? PausedUntil { get; set; }
+
+		/// <summary>
+		/// Reason for stream being paused
+		/// </summary>
+		public string? PauseComment { get; set; }
+
+		/// <summary>
+		/// Workflows for this stream
+		/// </summary>
+		public List<WorkflowConfig> Workflows { get; set; }
+
+		/// <summary>
+		/// Constructor
+		/// </summary>
+		/// <param name="stream">The stream to construct from</param>
+		/// <param name="templates">Templates for this stream</param>
+		public GetStreamResponse(IStream stream, List<GetTemplateRefResponse> templates)
+		{
+			Id = stream.Id.ToString();
+			ProjectId = stream.Config.ProjectConfig.Id.ToString();
+			Name = stream.Config.Name;
+			ConfigPath = stream.Config.Path ?? String.Empty;
+			ConfigRevision = stream.Config.Revision;
+			Order = stream.Config.Order;
+			NotificationChannel = stream.Config.NotificationChannel;
+			NotificationChannelFilter = stream.Config.NotificationChannelFilter;
+			TriageChannel = stream.Config.TriageChannel;
+			DefaultPreflightTemplate = stream.Config.DefaultPreflight?.TemplateId?.ToString();
+			DefaultPreflight = stream.Config.DefaultPreflight;
+			Tabs = stream.Config.Tabs;
+			AgentTypes = stream.Config.AgentTypes;
+			WorkspaceTypes = stream.Config.WorkspaceTypes;
+			Templates = templates;
+			PausedUntil = stream.PausedUntil;
+			PauseComment = stream.PauseComment;
+			Workflows = stream.Config.Workflows;
+		}
 	}
 
 	/// <summary>
-	/// Updates an existing stream template ref
-	/// </summary>
-	public class UpdateTemplateRefRequest
-	{
-		/// <summary>
-		/// Step states to update
-		/// </summary>
-		public List<UpdateStepStateRequest>? StepStates { get; set; }
-	}
-
-	/// <summary>
@@ -291,242 +386,30 @@
 	}
 
 	/// <summary>
-	/// Response describing a stream
+	/// Step state update request
 	/// </summary>
-	public class GetStreamResponse
+	public class UpdateStepStateRequest
 	{
 		/// <summary>
-		/// Unique id of the stream
+		/// Name of the step
 		/// </summary>
-		public string Id { get; set; }
+		[Required]
+		public string Name { get; set; } = String.Empty;
 
 		/// <summary>
-		/// Unique id of the project containing this stream
+		/// User who paused the step
 		/// </summary>
-		public string ProjectId { get; set; }
-
-		/// <summary>
-		/// Name of the stream
-		/// </summary>
-		public string Name { get; set; }
-
-		/// <summary>
-		/// The config file path on the server
-		/// </summary>
-		public string ConfigPath { get; set; } = String.Empty;
-
-		/// <summary>
-		/// Revision of the config file 
-		/// </summary>
-		public string ConfigRevision { get; set; } = String.Empty;
-
-		/// <summary>
-		/// Order to display in the list
-		/// </summary>
-		public int Order { get; set; }
-
-		/// <summary>
-		/// Notification channel for all jobs in this stream
-		/// </summary>
-		public string? NotificationChannel { get; set; }
-
-		/// <summary>
-		/// Notification channel filter for this template. Can be a combination of "Success", "Failure" and "Warnings" separated by pipe characters.
-		/// </summary>
-		public string? NotificationChannelFilter { get; set; }
-
-		/// <summary>
-		/// Channel to post issue triage notifications
-		/// </summary>
-		public string? TriageChannel { get; set; }
-
-		/// <summary>
-		/// Default template for running preflights
-		/// </summary>
-		public string? DefaultPreflightTemplate { get; set; }
-
-		/// <summary>
-		/// Default template to use for preflights
-		/// </summary>
-		public DefaultPreflightConfig? DefaultPreflight { get; set; }
-
-		/// <summary>
-		/// List of tabs to display for this stream
-		/// </summary>
-		public List<TabConfig> Tabs { get; set; }
-
-		/// <summary>
-		/// Map of agent name to type
-		/// </summary>
-		public Dictionary<string, AgentConfig> AgentTypes { get; set; }
-
-		/// <summary>
-		/// Map of workspace name to type
-		/// </summary>
-		public Dictionary<string, WorkspaceConfig>? WorkspaceTypes { get; set; }
-
-		/// <summary>
-		/// Templates for jobs in this stream
-		/// </summary>
-		public List<GetTemplateRefResponse> Templates { get; set; }
-
-		/// <summary>
-		/// Stream paused for new builds until this date
-		/// </summary>
-		public DateTime? PausedUntil { get; set; }
-
-		/// <summary>
-		/// Reason for stream being paused
-		/// </summary>
-		public string? PauseComment { get; set; }
-
-		/// <summary>
-		/// Workflows for this stream
-		/// </summary>
-		public List<WorkflowConfig> Workflows { get; set; }
-
-		/// <summary>
-		/// Constructor
-		/// </summary>
-		/// <param name="stream">The stream to construct from</param>
-		/// <param name="templates">Templates for this stream</param>
-		public GetStreamResponse(IStream stream, List<GetTemplateRefResponse> templates)
-		{
-			Id = stream.Id.ToString();
-			ProjectId = stream.Config.ProjectConfig.Id.ToString();
-			Name = stream.Config.Name;
-			ConfigPath = stream.Config.Path ?? String.Empty;
-			ConfigRevision = stream.Config.Revision;
-			Order = stream.Config.Order;
-			NotificationChannel = stream.Config.NotificationChannel;
-			NotificationChannelFilter = stream.Config.NotificationChannelFilter;
-			TriageChannel = stream.Config.TriageChannel;
-			DefaultPreflightTemplate = stream.Config.DefaultPreflight?.TemplateId?.ToString();
-			DefaultPreflight = stream.Config.DefaultPreflight;
-			Tabs = stream.Config.Tabs;
-			AgentTypes = stream.Config.AgentTypes;
-			WorkspaceTypes = stream.Config.WorkspaceTypes;
-			Templates = templates;
-			PausedUntil = stream.PausedUntil;
-			PauseComment = stream.PauseComment;
-			Workflows = stream.Config.Workflows;
-		}
+		public string? PausedByUserId { get; set; }
 	}
 
 	/// <summary>
-	/// Response describing a stream
+	/// Updates an existing stream template ref
 	/// </summary>
-	public class GetStreamResponseV2
+	public class UpdateTemplateRefRequest
 	{
-		readonly IStream _stream;
-
-		/// <inheritdoc cref="IStream.Id"/>
-		public StreamId Id => _stream.Id;
-
-		/// <inheritdoc cref="ProjectConfig.Id"/>
-		public ProjectId ProjectId => _stream.Config.ProjectConfig.Id;
-
-		/// <inheritdoc cref="StreamConfig.Name"/>
-		public string Name => _stream.Config.Name;
-
-		/// <inheritdoc cref="StreamConfig.Revision"/>
-		public string ConfigRevision => _stream.Config.Revision;
-
-		/// <inheritdoc cref="IStream.Config"/>
-		public StreamConfig? Config { get; }
-
-		/// <inheritdoc cref="IStream.Templates"/>
-		public List<GetTemplateRefResponseV2> Templates { get; }
-
-		/// <inheritdoc cref="IStream.PausedUntil"/>
-		public DateTime? PausedUntil => _stream.PausedUntil;
-
-		/// <inheritdoc cref="IStream.PauseComment"/>
-		public string? PauseComment => _stream.PauseComment;
-
-		internal GetStreamResponseV2(IStream stream, bool includeConfig, List<GetTemplateRefResponseV2> templates)
-		{
-			_stream = stream;
-
-			Templates = templates;
-
-			if (includeConfig)
-			{
-				Config = stream.Config;
-			}
-		}
+		/// <summary>
+		/// Step states to update
+		/// </summary>
+		public List<UpdateStepStateRequest>? StepStates { get; set; }
 	}
-
-	/// <summary>
-	/// Job template in a stream
-	/// </summary>
-	public class GetTemplateRefResponseV2
-	{
-		readonly ITemplate _template;
-		readonly ITemplateRef _templateRef;
-
-		/// <inheritdoc cref="ITemplateRef.Id"/>
-		public TemplateId Id => _templateRef.Id;
-
-		/// <inheritdoc cref="ITemplate.Hash"/>
-		public ContentHash Hash => _template.Hash;
-
-		internal GetTemplateRefResponseV2(ITemplate template, ITemplateRef templateRef)
-		{
-			_template = template;
-			_templateRef = templateRef;
-		}
-	}
-
-	/// <summary>
-	/// Schedule for a template
-	/// </summary>
-	public class GetTemplateScheduleResponseV2
-	{
-		readonly ISchedule _schedule;
-
-		/// <inheritdoc cref="ISchedule.LastTriggerCommitId"/>
-		[Obsolete("Use LastTriggerCommit instead")]
-		public int LastTriggerChange => _schedule.LastTriggerCommitId?.GetPerforceChangeOrMinusOne() ?? 0;
-
-		/// <inheritdoc cref="ISchedule.LastTriggerCommitId"/>
-		public CommitIdWithOrder? LastTriggerCommit => _schedule.LastTriggerCommitId;
-
-		/// <inheritdoc cref="ISchedule.LastTriggerTimeUtc"/>
-		public DateTime LastTriggerTimeUtc => _schedule.LastTriggerTimeUtc;
-
-		/// <inheritdoc cref="ISchedule.LastTriggerTimeUtc"/>
-		public IReadOnlyList<JobId> ActiveJobs => _schedule.ActiveJobs;
-
-		internal GetTemplateScheduleResponseV2(ISchedule schedule)
-		{
-			_schedule = schedule;
-		}
-	}
-
-	/// <summary>
-	/// State information for a step in the stream
-	/// </summary>
-	public class GetTemplateStepResponseV2
-	{
-		readonly ITemplateStep _step;
-
-		/// <inheritdoc cref="ITemplateStep.Name"/>
-		public string Name => _step.Name;
-
-		/// <inheritdoc cref="ITemplateStep.PausedByUserId"/>
-		public GetThinUserInfoResponse PausedByUserInfo { get; }
-
-		/// <inheritdoc cref="ITemplateStep.PauseTimeUtc"/>
-		public DateTime PauseTimeUtc => _step.PauseTimeUtc;
-
-		/// <summary>
-		/// Constructor
-		/// </summary>
-		public GetTemplateStepResponseV2(ITemplateStep step, GetThinUserInfoResponse pausedByUserInfo)
-		{
-			_step = step;
-			PausedByUserInfo = pausedByUserInfo;
-		}
-	}
 }
diff --git a/Engine/Source/Programs/Horde/Plugins/Build/HordeServer.Build/Streams/StreamsController.cs b/Engine/Source/Programs/Horde/Plugins/Build/HordeServer.Build/Streams/StreamsController.cs
--- a/Engine/Source/Programs/Horde/Plugins/Build/HordeServer.Build/Streams/StreamsController.cs
+++ b/Engine/Source/Programs/Horde/Plugins/Build/HordeServer.Build/Streams/StreamsController.cs
@@ -98,39 +98,6 @@ namespace HordeServer.Streams
 			return responses.OrderBy(x => x.Id).Select(x => PropertyFilter.Apply(x, filter)).ToList();
 		}
 
-		/// <summary>
-		/// Query all the streams for a particular project.
-		/// </summary>
-		/// <param name="projectIds">Unique id of the project to query</param>
-		/// <param name="filter">Filter for the properties to return</param>
-		/// <returns>Information about all the projects</returns>
-		[HttpGet]
-		[Route("/api/v2/streams")]
-		[ProducesResponseType(typeof(List<GetStreamResponse>), 200)]
-		public async Task<ActionResult<List<object>>> GetStreamsV2Async([FromQuery(Name = "ProjectId")] string[] projectIds, [FromQuery] PropertyFilter? filter = null)
-		{
-			ProjectId[] projectIdValues = Array.ConvertAll(projectIds, x => new ProjectId(x));
-
-			List<GetStreamResponseV2> responses = new List<GetStreamResponseV2>();
-			foreach (ProjectConfig projectConfig in _buildConfig.Value.Projects)
-			{
-				if (projectIdValues.Length == 0 || projectIdValues.Contains(projectConfig.Id))
-				{
-					foreach (StreamConfig streamConfig in projectConfig.Streams)
-					{
-						if (streamConfig.Authorize(StreamAclAction.ViewStream, User))
-						{
-							IStream stream = await _streamCollection.GetAsync(streamConfig);
-							GetStreamResponseV2 response = await CreateGetStreamResponseV2Async(stream, false);
-							responses.Add(response);
-						}
-					}
-				}
-			}
-
-			return responses.OrderBy(x => x.Id).Select(x => PropertyFilter.Apply(x, filter)).ToList();
-		}
-
 		/// <summary>
 		/// Retrieve information about a specific stream.
 		/// </summary>
@@ -157,33 +124,6 @@ namespace HordeServer.Streams
 			return PropertyFilter.Apply(await CreateGetStreamResponseAsync(stream, cancellationToken), filter);
 		}
 
-		/// <summary>
-		/// Retrieve information about a specific stream.
-		/// </summary>
-		/// <param name="streamId">Id of the stream to get information about</param>
-		/// <param name="config">The client's cached config revision. If this matches value matches the current config version of the stream, the config object will be omitted from the response.</param>
-		/// <param name="filter">Filter for the properties to return</param>
-		/// <returns>Information about the requested project</returns>
-		[HttpGet]
-		[Route("/api/v2/streams/{streamId}/config")]
-		[ProducesResponseType(typeof(GetStreamResponseV2), 200)]
-		public async Task<ActionResult<object>> GetStreamV2Async(StreamId streamId, [FromQuery] string? config = null, [FromQuery] PropertyFilter? filter = null)
-		{
-			StreamConfig? streamConfig;
-			if (!_buildConfig.Value.TryGetStream(streamId, out streamConfig))
-			{
-				return NotFound(streamId);
-			}
-			if (!streamConfig.Authorize(StreamAclAction.ViewStream, User))
-			{
-				return Forbid(StreamAclAction.ViewStream, streamId);
-			}
-
-			bool includeConfig = !String.Equals(config, streamConfig.Revision, StringComparison.Ordinal);
-			IStream stream = await _streamCollection.GetAsync(streamConfig);
-			return PropertyFilter.Apply(await CreateGetStreamResponseV2Async(stream, includeConfig), filter);
-		}
-
 		/// <summary>
 		/// Create a stream response object, including all the templates
 		/// </summary>
@@ -228,17 +168,6 @@ namespace HordeServer.Streams
 			return stream.ToApiResponse(apiTemplateRefs);
 		}
 
-		async Task<GetStreamResponseV2> CreateGetStreamResponseV2Async(IStream stream, bool includeConfig)
-		{
-			List<GetTemplateRefResponseV2> templates = new List<GetTemplateRefResponseV2>();
-			foreach (ITemplateRef templateRef in stream.Templates.Values)
-			{
-				ITemplate template = await _templateCollection.GetOrAddAsync(templateRef.Config);
-				templates.Add(new GetTemplateRefResponseV2(template, templateRef));
-			}
-			return new GetStreamResponseV2(stream, includeConfig, templates);
-		}
-
 		/// <summary>
 		/// Gets a list of changes for a stream
 		/// </summary>
