Horde: Add workflow interface and messages to EpicGames.Horde.

#rnx

diff --git a/Engine/Source/Programs/Horde/Plugins/Build/HordeServer.Build/Issues/WorkflowConfig.cs b/Engine/Source/Programs/Horde/Plugins/Build/HordeServer.Build/Issues/WorkflowConfig.cs
--- a/Engine/Source/Programs/Horde/Plugins/Build/HordeServer.Build/Issues/WorkflowConfig.cs
+++ b/Engine/Source/Programs/Horde/Plugins/Build/HordeServer.Build/Issues/WorkflowConfig.cs
@@ -11,139 +11,96 @@ namespace HordeServer.Issues
 	/// <summary>
 	/// Configuration for an issue workflow
 	/// </summary>
-	public class WorkflowConfig
+	public class WorkflowConfig : IWorkflow
 	{
-		/// <summary>
-		/// Identifier for this workflow
-		/// </summary>
+		/// <inheritdoc/>
 		public WorkflowId Id { get; set; } = WorkflowId.Empty;
 
-		/// <summary>
-		/// Times of day at which to send a report
-		/// </summary>
+		/// <inheritdoc cref="IWorkflow.ReportTimes"/>
 		public List<TimeSpan> ReportTimes { get; set; } = new List<TimeSpan>();
+		IReadOnlyList<TimeSpan> IWorkflow.ReportTimes => ReportTimes;
 
-		/// <summary>
-		/// Name of the tab to post summary data to
-		/// </summary>
+		/// <inheritdoc/>
 		public string? SummaryTab { get; set; }
 
-		/// <summary>
-		/// Channel to post summary information for these templates.
-		/// </summary>
+		/// <inheritdoc/>
 		public string? ReportChannel { get; set; }
 
-		/// <summary>
-		/// Whether to include issues with a warning status in the summary
-		/// </summary>
+		/// <inheritdoc/>
 		public bool ReportWarnings { get; set; } = true;
 
-		/// <summary>
-		/// Whether to group issues by template in the report
-		/// </summary>
+		/// <inheritdoc/>
 		public bool GroupIssuesByTemplate { get; set; } = true;
 
-		/// <summary>
-		/// Channel to post threads for triaging new issues
-		/// </summary>
+		/// <inheritdoc/>
 		public string? TriageChannel { get; set; }
 
-		/// <summary>
-		/// Prefix for all triage messages
-		/// </summary>
+		/// <inheritdoc/>
 		public string? TriagePrefix { get; set; } = "*[NEW]* ";
 
-		/// <summary>
-		/// Suffix for all triage messages
-		/// </summary>
+		/// <inheritdoc/>
 		public string? TriageSuffix { get; set; }
 
-		/// <summary>
-		/// Instructions posted to triage threads
-		/// </summary>
+		/// <inheritdoc/>
 		public string? TriageInstructions { get; set; }
 
-		/// <summary>
-		/// User id of a Slack user/alias to ping if there is nobody assigned to an issue by default.
-		/// </summary>
+		/// <inheritdoc/>
 		public string? TriageAlias { get; set; }
 
-		/// <summary>
-		/// Slack user/alias to ping for specific issue types (such as Systemic), if there is nobody assigned to an issue by default.
-		/// </summary>
+		/// <inheritdoc cref="IWorkflow.TriageTypeAliases"/>
 		public Dictionary<string, string>? TriageTypeAliases { get; set; }
+		IReadOnlyDictionary<string, string>? IWorkflow.TriageTypeAliases => TriageTypeAliases;
 
-		/// <summary>
-		/// Alias to ping if an issue has not been resolved for a certain amount of time
-		/// </summary>
+		/// <inheritdoc/>
 		public string? EscalateAlias { get; set; }
 
-		/// <summary>
-		/// Times after an issue has been opened to escalate to the alias above, in minutes. Continues to notify on the last interval once reaching the end of the list.
-		/// </summary>
+		/// <inheritdoc cref="IWorkflow.EscalateTimes"/>
 		public List<int> EscalateTimes { get; set; } = new List<int> { 120 };
 
-		/// <summary>
-		/// Maximum number of people to mention on a triage thread
-		/// </summary>
+		IReadOnlyList<int> IWorkflow.EscalateTimes => EscalateTimes;
+
+		/// <inheritdoc/>
 		public int MaxMentions { get; set; } = 5;
 
-		/// <summary>
-		/// Whether to mention people on this thread. Useful to disable for testing.
-		/// </summary>
+		/// <inheritdoc/>
 		public bool AllowMentions { get; set; } = true;
 
-		/// <summary>
-		/// Uses the admin.conversations.invite API to invite users to the channel
-		/// </summary>
+		/// <inheritdoc/>
 		public bool InviteRestrictedUsers { get; set; }
 
-		/// <summary>
-		/// Skips sending reports when there are no active issues. 
-		/// </summary>
+		/// <inheritdoc/>
 		public bool SkipWhenEmpty { get; set; }
 
-		/// <summary>
-		/// Whether to show warnings about merging changes into the origin stream.
-		/// </summary>
+		/// <inheritdoc/>
 		public bool ShowMergeWarnings { get; set; }
 
-		/// <summary>
-		/// Additional node annotations implicit in this workflow
-		/// </summary>
+		/// <inheritdoc cref="IWorkflow.Annotations"/>
 		public NodeAnnotations Annotations { get; set; } = new NodeAnnotations();
+		IReadOnlyNodeAnnotations IWorkflow.Annotations => Annotations;
 
-		/// <summary>
-		/// External issue tracking configuration for this workflow
-		/// </summary>
-		public ExternalIssueConfig? ExternalIssues { get; set; }
+		/// <inheritdoc cref="IWorkflow.ExternalIssues"/>
+		public WorkflowExternalIssuesConfig? ExternalIssues { get; set; }
+		IWorkflowExternalIssues? IWorkflow.ExternalIssues => ExternalIssues;
 
-		/// <summary>
-		/// Additional issue handlers enabled for this workflow
-		/// </summary>
+		/// <inheritdoc cref="IWorkflow.IssueHandlers"/>
 		public List<string>? IssueHandlers { get; set; }
+		IReadOnlyList<string>? IWorkflow.IssueHandlers => IssueHandlers;
 	}
 
 	/// <summary>
 	/// External issue tracking configuration for a workflow
 	/// </summary>
-	public class ExternalIssueConfig
+	public class WorkflowExternalIssuesConfig : IWorkflowExternalIssues
 	{
-		/// <summary>
-		/// Project key in external issue tracker
-		/// </summary>
+		/// <inheritdoc/>
 		[Required]
 		public string ProjectKey { get; set; } = String.Empty;
 
-		/// <summary>
-		/// Default component id for issues using workflow
-		/// </summary>
+		/// <inheritdoc/>
 		[Required]
 		public string DefaultComponentId { get; set; } = String.Empty;
 
-		/// <summary>
-		/// Default issue type id for issues using workflow
-		/// </summary>
+		/// <inheritdoc/>
 		[Required]
 		public string DefaultIssueTypeId { get; set; } = String.Empty;
 	}
diff --git a/Engine/Source/Programs/Shared/EpicGames.Horde/Issues/IWorkflow.cs b/Engine/Source/Programs/Shared/EpicGames.Horde/Issues/IWorkflow.cs
new file mode 100644
--- /dev/null
+++ b/Engine/Source/Programs/Shared/EpicGames.Horde/Issues/IWorkflow.cs
@@ -0,0 +1,145 @@
+// Copyright Epic Games, Inc. All Rights Reserved.
+
+using System;
+using System.Collections.Generic;
+using HordeServer.Jobs.Graphs;
+
+namespace EpicGames.Horde.Issues
+{
+	/// <summary>
+	/// Configuration for an issue workflow
+	/// </summary>
+	public interface IWorkflow
+	{
+		/// <summary>
+		/// Identifier for this workflow
+		/// </summary>
+		WorkflowId Id { get; }
+
+		/// <summary>
+		/// Times of day at which to send a report
+		/// </summary>
+		IReadOnlyList<TimeSpan> ReportTimes { get; }
+
+		/// <summary>
+		/// Name of the tab to post summary data to
+		/// </summary>
+		string? SummaryTab { get; }
+
+		/// <summary>
+		/// Channel to post summary information for these templates.
+		/// </summary>
+		string? ReportChannel { get; }
+
+		/// <summary>
+		/// Whether to include issues with a warning status in the summary
+		/// </summary>
+		bool ReportWarnings { get; }
+
+		/// <summary>
+		/// Whether to group issues by template in the report
+		/// </summary>
+		bool GroupIssuesByTemplate { get; }
+
+		/// <summary>
+		/// Channel to post threads for triaging new issues
+		/// </summary>
+		string? TriageChannel { get; }
+
+		/// <summary>
+		/// Prefix for all triage messages
+		/// </summary>
+		string? TriagePrefix { get; }
+
+		/// <summary>
+		/// Suffix for all triage messages
+		/// </summary>
+		string? TriageSuffix { get; }
+
+		/// <summary>
+		/// Instructions posted to triage threads
+		/// </summary>
+		string? TriageInstructions { get; }
+
+		/// <summary>
+		/// User id of a Slack user/alias to ping if there is nobody assigned to an issue by default.
+		/// </summary>
+		string? TriageAlias { get; }
+
+		/// <summary>
+		/// Slack user/alias to ping for specific issue types (such as Systemic), if there is nobody assigned to an issue by default.
+		/// </summary>
+		IReadOnlyDictionary<string, string>? TriageTypeAliases { get; }
+
+		/// <summary>
+		/// Alias to ping if an issue has not been resolved for a certain amount of time
+		/// </summary>
+		string? EscalateAlias { get; }
+
+		/// <summary>
+		/// Times after an issue has been opened to escalate to the alias above, in minutes. Continues to notify on the last interval once reaching the end of the list.
+		/// </summary>
+		IReadOnlyList<int> EscalateTimes { get; }
+
+		/// <summary>
+		/// Maximum number of people to mention on a triage thread
+		/// </summary>
+		int MaxMentions { get; }
+
+		/// <summary>
+		/// Whether to mention people on this thread. Useful to disable for testing.
+		/// </summary>
+		bool AllowMentions { get; }
+
+		/// <summary>
+		/// Uses the admin.conversations.invite API to invite users to the channel
+		/// </summary>
+		bool InviteRestrictedUsers { get; }
+
+		/// <summary>
+		/// Skips sending reports when there are no active issues. 
+		/// </summary>
+		bool SkipWhenEmpty { get; }
+
+		/// <summary>
+		/// Whether to show warnings about merging changes into the origin stream.
+		/// </summary>
+		bool ShowMergeWarnings { get; }
+
+		/// <summary>
+		/// Additional node annotations implicit in this workflow
+		/// </summary>
+		IReadOnlyNodeAnnotations Annotations { get; }
+
+		/// <summary>
+		/// External issue tracking configuration for this workflow
+		/// </summary>
+		IWorkflowExternalIssues? ExternalIssues { get; }
+
+		/// <summary>
+		/// Additional issue handlers enabled for this workflow
+		/// </summary>
+		IReadOnlyList<string>? IssueHandlers { get; }
+	}
+
+	/// <summary>
+	/// External issue tracking configuration for a workflow
+	/// </summary>
+	public interface IWorkflowExternalIssues
+	{
+		/// <summary>
+		/// Project key in external issue tracker
+		/// </summary>
+		string ProjectKey { get; }
+
+		/// <summary>
+		/// Default component id for issues using workflow
+		/// </summary>
+		string DefaultComponentId { get; }
+
+		/// <summary>
+		/// Default issue type id for issues using workflow
+		/// </summary>
+		string DefaultIssueTypeId { get; }
+	}
+}
diff --git a/Engine/Source/Programs/Shared/EpicGames.Horde/Issues/WorkflowMessages.cs b/Engine/Source/Programs/Shared/EpicGames.Horde/Issues/WorkflowMessages.cs
new file mode 100644
--- /dev/null
+++ b/Engine/Source/Programs/Shared/EpicGames.Horde/Issues/WorkflowMessages.cs
@@ -0,0 +1,162 @@
+// Copyright Epic Games, Inc. All Rights Reserved.
+
+using System;
+using System.Collections.Generic;
+using System.Linq;
+using HordeServer.Jobs.Graphs;
+
+#pragma warning disable CA2227 //  Change x to be read-only by removing the property setter
+
+namespace EpicGames.Horde.Issues
+{
+	/// <summary>
+	/// Configuration for an issue workflow
+	/// </summary>
+	public class GetWorkflowResponse : IWorkflow
+	{
+		/// <inheritdoc/>
+		public WorkflowId Id { get; set; }
+
+		/// <inheritdoc cref="IWorkflow.ReportTimes"/>
+		public List<TimeSpan> ReportTimes { get; set; }
+		IReadOnlyList<TimeSpan> IWorkflow.ReportTimes => ReportTimes;
+
+		/// <inheritdoc/>
+		public string? SummaryTab { get; set; }
+
+		/// <inheritdoc/>
+		public string? ReportChannel { get; set; }
+
+		/// <inheritdoc/>
+		public bool ReportWarnings { get; set; } = true;
+
+		/// <inheritdoc/>
+		public bool GroupIssuesByTemplate { get; set; } = true;
+
+		/// <inheritdoc/>
+		public string? TriageChannel { get; set; }
+
+		/// <inheritdoc/>
+		public string? TriagePrefix { get; set; }
+
+		/// <inheritdoc/>
+		public string? TriageSuffix { get; set; }
+
+		/// <inheritdoc/>
+		public string? TriageInstructions { get; set; }
+
+		/// <inheritdoc/>
+		public string? TriageAlias { get; set; }
+
+		/// <inheritdoc cref="IWorkflow.TriageTypeAliases"/>
+		public Dictionary<string, string>? TriageTypeAliases { get; set; }
+		IReadOnlyDictionary<string, string>? IWorkflow.TriageTypeAliases => TriageTypeAliases;
+
+		/// <inheritdoc/>
+		public string? EscalateAlias { get; set; }
+
+		/// <inheritdoc cref="IWorkflow.EscalateTimes"/>
+		public List<int> EscalateTimes { get; set; }
+		IReadOnlyList<int> IWorkflow.EscalateTimes => EscalateTimes;
+
+		/// <inheritdoc/>
+		public int MaxMentions { get; set; }
+
+		/// <inheritdoc/>
+		public bool AllowMentions { get; set; }
+
+		/// <inheritdoc/>
+		public bool InviteRestrictedUsers { get; set; }
+
+		/// <inheritdoc/>
+		public bool SkipWhenEmpty { get; set; }
+
+		/// <inheritdoc/>
+		public bool ShowMergeWarnings { get; set; }
+
+		/// <inheritdoc cref="IWorkflow.Annotations"/>
+		public NodeAnnotations Annotations { get; set; } = new NodeAnnotations();
+		IReadOnlyNodeAnnotations IWorkflow.Annotations => Annotations;
+
+		/// <inheritdoc cref="IWorkflow.ExternalIssues"/>
+		public GetWorkflowExternalIssuesResponse? ExternalIssues { get; set; }
+		IWorkflowExternalIssues? IWorkflow.ExternalIssues => ExternalIssues;
+
+		/// <inheritdoc cref="IWorkflow.IssueHandlers"/>
+		public List<string>? IssueHandlers { get; set; }
+		IReadOnlyList<string>? IWorkflow.IssueHandlers => IssueHandlers;
+
+		/// <summary>
+		/// Constructor
+		/// </summary>
+		public GetWorkflowResponse()
+		{
+			ReportTimes = new List<TimeSpan>();
+			EscalateTimes = new List<int>();
+		}
+
+		/// <summary>
+		/// Constructor
+		/// </summary>
+		public GetWorkflowResponse(IWorkflow workflow)
+		{
+			Id = workflow.Id;
+			ReportTimes = workflow.ReportTimes.ToList();
+			SummaryTab = workflow.SummaryTab;
+			ReportChannel = workflow.ReportChannel;
+			ReportWarnings = workflow.ReportWarnings;
+			GroupIssuesByTemplate = workflow.GroupIssuesByTemplate;
+			TriageChannel = workflow.TriageChannel;
+			TriagePrefix = workflow.TriagePrefix;
+			TriageSuffix = workflow.TriageSuffix;
+			TriageInstructions = workflow.TriageInstructions;
+			TriageAlias = workflow.TriageAlias;
+			TriageTypeAliases = workflow.TriageTypeAliases?.ToDictionary();
+			EscalateAlias = workflow.EscalateAlias;
+			EscalateTimes = workflow.EscalateTimes.ToList();
+			MaxMentions = workflow.MaxMentions;
+			AllowMentions = workflow.AllowMentions;
+			InviteRestrictedUsers = workflow.InviteRestrictedUsers;
+			SkipWhenEmpty = workflow.SkipWhenEmpty;
+			ShowMergeWarnings = workflow.ShowMergeWarnings;
+			Annotations = new NodeAnnotations(workflow.Annotations);
+			ExternalIssues = (workflow.ExternalIssues == null) ? null : new GetWorkflowExternalIssuesResponse(workflow.ExternalIssues);
+			IssueHandlers = workflow.IssueHandlers?.ToList();
+		}
+	}
+
+	/// <summary>
+	/// External issue tracking configuration for a workflow
+	/// </summary>
+	public class GetWorkflowExternalIssuesResponse : IWorkflowExternalIssues
+	{
+		/// <inheritdoc/>
+		public string ProjectKey { get; set; } = String.Empty;
+
+		/// <inheritdoc/>
+		public string DefaultComponentId { get; set; } = String.Empty;
+
+		/// <inheritdoc/>
+		public string DefaultIssueTypeId { get; set; } = String.Empty;
+
+		/// <summary>
+		/// Constructor
+		/// </summary>
+		public GetWorkflowExternalIssuesResponse()
+		{
+			ProjectKey = String.Empty;
+			DefaultComponentId = String.Empty;
+			DefaultIssueTypeId = String.Empty;
+		}
+
+		/// <summary>
+		/// Constructor
+		/// </summary>
+		public GetWorkflowExternalIssuesResponse(IWorkflowExternalIssues issues)
+		{
+			ProjectKey = issues.ProjectKey;
+			DefaultComponentId = issues.DefaultComponentId;
+			DefaultIssueTypeId = issues.DefaultIssueTypeId;
+		}
+	}
+}
