Horde: Populate the initial agent type for templates from the stream config value.

#rnx

diff --git a/Engine/Source/Programs/Horde/Plugins/Build/HordeServer.Build.Tests/Devices/DeviceServiceTests.cs b/Engine/Source/Programs/Horde/Plugins/Build/HordeServer.Build.Tests/Devices/DeviceServiceTests.cs
--- a/Engine/Source/Programs/Horde/Plugins/Build/HordeServer.Build.Tests/Devices/DeviceServiceTests.cs
+++ b/Engine/Source/Programs/Horde/Plugins/Build/HordeServer.Build.Tests/Devices/DeviceServiceTests.cs
@@ -186,7 +186,7 @@ namespace HordeServer.Tests.Devices
 			Mock<ITemplate> templateMock = new Mock<ITemplate>(MockBehavior.Strict);
 			templateMock.SetupGet(x => x.InitialAgentType).Returns((string?)null);
 
-			IGraph baseGraph = await GraphCollection.AddAsync(templateMock.Object, null);
+			IGraph baseGraph = await GraphCollection.AddAsync(templateMock.Object);
 
 			CreateJobOptions options = new CreateJobOptions();
 			options.Arguments.Add("-Target=Run Tests");
diff --git a/Engine/Source/Programs/Horde/Plugins/Build/HordeServer.Build.Tests/Fixture.cs b/Engine/Source/Programs/Horde/Plugins/Build/HordeServer.Build.Tests/Fixture.cs
--- a/Engine/Source/Programs/Horde/Plugins/Build/HordeServer.Build.Tests/Fixture.cs
+++ b/Engine/Source/Programs/Horde/Plugins/Build/HordeServer.Build.Tests/Fixture.cs
@@ -50,7 +50,7 @@ namespace HordeServer.Tests
 		{
 			GraphStub fg = new (ContentHash.Empty, 1122);
 			Template = await templateCollection.GetOrAddAsync(new TemplateConfig { Name = "Test template" });
-			Graph = await graphCollection.AddAsync(Template, null);
+			Graph = await graphCollection.AddAsync(Template);
 
 			TemplateRefId1 = new TemplateId("template1");
 			TemplateRefId2 = new TemplateId("template2");
diff --git a/Engine/Source/Programs/Horde/Plugins/Build/HordeServer.Build.Tests/Jobs/Bisect/BisectTests.cs b/Engine/Source/Programs/Horde/Plugins/Build/HordeServer.Build.Tests/Jobs/Bisect/BisectTests.cs
--- a/Engine/Source/Programs/Horde/Plugins/Build/HordeServer.Build.Tests/Jobs/Bisect/BisectTests.cs
+++ b/Engine/Source/Programs/Horde/Plugins/Build/HordeServer.Build.Tests/Jobs/Bisect/BisectTests.cs
@@ -148,7 +148,7 @@ namespace HordeServer.Tests.Jobs.Bisect
 			Mock<ITemplate> templateMock = new Mock<ITemplate>(MockBehavior.Strict);
 			templateMock.SetupGet(x => x.InitialAgentType).Returns((string?)null);
 
-			IGraph graph = await GraphCollection.AddAsync(templateMock.Object, null);
+			IGraph graph = await GraphCollection.AddAsync(templateMock.Object);
 			graph = await GraphCollection.AppendAsync(graph, groups);
 
 			// Create two job runs, one which fails, one that succeeds.
diff --git a/Engine/Source/Programs/Horde/Plugins/Build/HordeServer.Build.Tests/Jobs/JobCollectionTests.cs b/Engine/Source/Programs/Horde/Plugins/Build/HordeServer.Build.Tests/Jobs/JobCollectionTests.cs
--- a/Engine/Source/Programs/Horde/Plugins/Build/HordeServer.Build.Tests/Jobs/JobCollectionTests.cs
+++ b/Engine/Source/Programs/Horde/Plugins/Build/HordeServer.Build.Tests/Jobs/JobCollectionTests.cs
@@ -126,7 +126,7 @@ namespace HordeServer.Tests.Jobs
 			Mock<ITemplate> templateMock = new Mock<ITemplate>(MockBehavior.Strict);
 			templateMock.SetupGet(x => x.InitialAgentType).Returns((string?)null);
 
-			IGraph baseGraph = await GraphCollection.AddAsync(templateMock.Object, null);
+			IGraph baseGraph = await GraphCollection.AddAsync(templateMock.Object);
 
 			CreateJobOptions options = new CreateJobOptions();
 			options.Arguments.Add("-Target=Publish Client");
@@ -174,7 +174,7 @@ namespace HordeServer.Tests.Jobs
 			Mock<ITemplate> templateMock = new Mock<ITemplate>(MockBehavior.Strict);
 			templateMock.SetupGet(x => x.InitialAgentType).Returns((string?)null);
 
-			IGraph baseGraph = await GraphCollection.AddAsync(templateMock.Object, null);
+			IGraph baseGraph = await GraphCollection.AddAsync(templateMock.Object);
 
 			List<NewGroup> groups = new List<NewGroup>();
 			groups.Add(new NewGroup("Test", new List<NewNode> { new NewNode("Initial Node", outputs: new List<string> { "#InitialOutput" }) }));
@@ -251,7 +251,7 @@ namespace HordeServer.Tests.Jobs
 			Mock<ITemplate> templateMock = new Mock<ITemplate>(MockBehavior.Strict);
 			templateMock.SetupGet(x => x.InitialAgentType).Returns((string?)null);
 
-			IGraph baseGraph = await GraphCollection.AddAsync(templateMock.Object, null);
+			IGraph baseGraph = await GraphCollection.AddAsync(templateMock.Object);
 
 			List<NewGroup> groups = new List<NewGroup>();
 			groups.Add(new NewGroup("Test", new List<NewNode> { new NewNode("Initial Node", outputs: new List<string> { "#InitialOutput" }) }));
@@ -319,7 +319,7 @@ namespace HordeServer.Tests.Jobs
 			Mock<ITemplate> templateMock = new Mock<ITemplate>(MockBehavior.Strict);
 			templateMock.SetupGet(x => x.InitialAgentType).Returns((string?)null);
 
-			IGraph baseGraph = await GraphCollection.AddAsync(templateMock.Object, null);
+			IGraph baseGraph = await GraphCollection.AddAsync(templateMock.Object);
 
 			CreateJobOptions options = new CreateJobOptions();
 			options.Arguments.Add("-Target=Step 1");
@@ -375,7 +375,7 @@ namespace HordeServer.Tests.Jobs
 			Mock<ITemplate> templateMock = new Mock<ITemplate>(MockBehavior.Strict);
 			templateMock.SetupGet(x => x.InitialAgentType).Returns((string?)null);
 
-			IGraph baseGraph = await GraphCollection.AddAsync(templateMock.Object, null);
+			IGraph baseGraph = await GraphCollection.AddAsync(templateMock.Object);
 
 			CreateJobOptions options = new CreateJobOptions();
 			options.Arguments.Add("-Target=Step 1");
@@ -426,7 +426,7 @@ namespace HordeServer.Tests.Jobs
 			Mock<ITemplate> templateMock = new Mock<ITemplate>(MockBehavior.Strict);
 			templateMock.SetupGet(x => x.InitialAgentType).Returns((string?)null);
 
-			IGraph baseGraph = await GraphCollection.AddAsync(templateMock.Object, null);
+			IGraph baseGraph = await GraphCollection.AddAsync(templateMock.Object);
 
 			CreateJobOptions options = new CreateJobOptions();
 			options.Arguments.Add("-Target=Step 1");
@@ -827,7 +827,7 @@ namespace HordeServer.Tests.Jobs
 			Mock<ITemplate> templateMock = new Mock<ITemplate>(MockBehavior.Strict);
 			templateMock.SetupGet(x => x.InitialAgentType).Returns((string?)null);
 
-			IGraph baseGraph = await GraphCollection.AddAsync(templateMock.Object, null);
+			IGraph baseGraph = await GraphCollection.AddAsync(templateMock.Object);
 
 			CreateJobOptions options = new CreateJobOptions();
 			options.Arguments.Add("-Target=Step 1");
@@ -848,7 +848,7 @@ namespace HordeServer.Tests.Jobs
 			Mock<ITemplate> templateMock = new Mock<ITemplate>(MockBehavior.Strict);
 			templateMock.SetupGet(x => x.InitialAgentType).Returns((string?)null);
 
-			IGraph baseGraph = await GraphCollection.AddAsync(templateMock.Object, null);
+			IGraph baseGraph = await GraphCollection.AddAsync(templateMock.Object);
 
 			CreateJobOptions options = new CreateJobOptions();
 			options.Arguments.Add("-Target=Compile Editor");
@@ -899,7 +899,7 @@ namespace HordeServer.Tests.Jobs
 			Mock<ITemplate> templateMock = new Mock<ITemplate>(MockBehavior.Strict);
 			templateMock.SetupGet(x => x.InitialAgentType).Returns((string?)null);
 
-			IGraph baseGraph = await GraphCollection.AddAsync(templateMock.Object, null);
+			IGraph baseGraph = await GraphCollection.AddAsync(templateMock.Object);
 
 			List<NewGraphArtifact> newArtifacts = new List<NewGraphArtifact>();
 			newArtifacts.Add(new NewGraphArtifact(new ArtifactName("foo"), new ArtifactType("type"), "hello world", "Engine/Source", new List<string>(), new List<string>(), null, "fileset"));
@@ -916,7 +916,7 @@ namespace HordeServer.Tests.Jobs
 		[TestMethod]
 		public async Task FindJobsAsync()
 		{
-			IGraph graph = await GraphCollection.AddAsync(_template1, null);
+			IGraph graph = await GraphCollection.AddAsync(_template1);
 
 			List<IJob> jobs = [];
 			for (int i = 0; i < 10; i++)
diff --git a/Engine/Source/Programs/Horde/Plugins/Build/HordeServer.Build.Tests/Jobs/JobExpirationTests.cs b/Engine/Source/Programs/Horde/Plugins/Build/HordeServer.Build.Tests/Jobs/JobExpirationTests.cs
--- a/Engine/Source/Programs/Horde/Plugins/Build/HordeServer.Build.Tests/Jobs/JobExpirationTests.cs
+++ b/Engine/Source/Programs/Horde/Plugins/Build/HordeServer.Build.Tests/Jobs/JobExpirationTests.cs
@@ -42,7 +42,7 @@ namespace HordeServer.Tests.Jobs
 
 			ITemplate template = await TemplateCollection.GetOrAddAsync(streamConfig.Templates[0]);
 
-			IGraph graph = await GraphCollection.AddAsync(template, null);
+			IGraph graph = await GraphCollection.AddAsync(template);
 
 			return await JobService.CreateJobAsync(null, streamConfig, templateId, template.Hash, graph, "Hello", CommitIdWithOrder.FromPerforceChange(1234), CommitIdWithOrder.FromPerforceChange(1233), options);
 		}
diff --git a/Engine/Source/Programs/Horde/Plugins/Build/HordeServer.Build.Tests/Jobs/JobServiceTests.cs b/Engine/Source/Programs/Horde/Plugins/Build/HordeServer.Build.Tests/Jobs/JobServiceTests.cs
--- a/Engine/Source/Programs/Horde/Plugins/Build/HordeServer.Build.Tests/Jobs/JobServiceTests.cs
+++ b/Engine/Source/Programs/Horde/Plugins/Build/HordeServer.Build.Tests/Jobs/JobServiceTests.cs
@@ -57,7 +57,7 @@ namespace HordeServer.Tests.Jobs
 
 			ITemplate template = await TemplateCollection.GetOrAddAsync(streamConfig.Templates[0]);
 
-			IGraph graph = await GraphCollection.AddAsync(template, null);
+			IGraph graph = await GraphCollection.AddAsync(template);
 
 			IJob job = await JobService.CreateJobAsync(null, streamConfig, templateRefId1, template.Hash, graph, "Hello", CommitIdWithOrder.FromPerforceChange(1234), CommitIdWithOrder.FromPerforceChange(1233), options);
 			Assert.AreEqual(1, job.ChainedJobs.Count);
@@ -255,7 +255,7 @@ namespace HordeServer.Tests.Jobs
 			await AgentService.CreateSessionAsync(agent, new RpcAgentCapabilities(), null);
 
 			ITemplate template = await TemplateCollection.GetOrAddAsync(new TemplateConfig { Name = "Test template" });
-			IGraph graph = await GraphCollection.AddAsync(template, null);
+			IGraph graph = await GraphCollection.AddAsync(template);
 
 			NewGroup groupA = new NewGroup("win", new List<NewNode>());
 			groupA.Nodes.Add(new NewNode("Compile"));
diff --git a/Engine/Source/Programs/Horde/Plugins/Build/HordeServer.Build.Tests/Jobs/Schedules/SchedulerTests.cs b/Engine/Source/Programs/Horde/Plugins/Build/HordeServer.Build.Tests/Jobs/Schedules/SchedulerTests.cs
--- a/Engine/Source/Programs/Horde/Plugins/Build/HordeServer.Build.Tests/Jobs/Schedules/SchedulerTests.cs
+++ b/Engine/Source/Programs/Horde/Plugins/Build/HordeServer.Build.Tests/Jobs/Schedules/SchedulerTests.cs
@@ -451,7 +451,7 @@ namespace HordeServer.Tests.Jobs.Schedules
 			IStream stream = await StreamCollection.GetAsync(config);
 
 			// Create the TriggerNext step and mark it as complete
-			IGraph graphA = await GraphCollection.AddAsync(newTemplate1, null);
+			IGraph graphA = await GraphCollection.AddAsync(newTemplate1);
 			NewGroup groupA = new NewGroup("win", new List<NewNode> { new NewNode("TriggerNext") });
 			graphA = await GraphCollection.AppendAsync(graphA, new List<NewGroup> { groupA });
 
@@ -550,7 +550,7 @@ namespace HordeServer.Tests.Jobs.Schedules
 			Assert.IsNotNull(template2);
 
 			// Create the graph
-			IGraph graphA = await GraphCollection.AddAsync(template1, null);
+			IGraph graphA = await GraphCollection.AddAsync(template1);
 			NewGroup groupA = new NewGroup("win", new List<NewNode> { new NewNode("TriggerNext") });
 			graphA = await GraphCollection.AppendAsync(graphA, new List<NewGroup> { groupA });
 
diff --git a/Engine/Source/Programs/Horde/Plugins/Build/HordeServer.Build/Jobs/Bisect/BisectService.cs b/Engine/Source/Programs/Horde/Plugins/Build/HordeServer.Build/Jobs/Bisect/BisectService.cs
--- a/Engine/Source/Programs/Horde/Plugins/Build/HordeServer.Build/Jobs/Bisect/BisectService.cs
+++ b/Engine/Source/Programs/Horde/Plugins/Build/HordeServer.Build/Jobs/Bisect/BisectService.cs
@@ -156,7 +156,7 @@ namespace HordeServer.Jobs.Bisect
 			}
 
 			ITemplate template = await _templateCollection.GetOrAddAsync(templateRefConfig);
-			IGraph graph = await _graphCollection.AddAsync(template, streamConfig.InitialAgentType, cancellationToken);
+			IGraph graph = await _graphCollection.AddAsync(template, cancellationToken);
 
 			CreateJobOptions options = new CreateJobOptions(templateRefConfig);
 			options.StartedByBisectTaskId = bisectTask.Id;
diff --git a/Engine/Source/Programs/Horde/Plugins/Build/HordeServer.Build/Jobs/Graphs/GraphCollection.cs b/Engine/Source/Programs/Horde/Plugins/Build/HordeServer.Build/Jobs/Graphs/GraphCollection.cs
--- a/Engine/Source/Programs/Horde/Plugins/Build/HordeServer.Build/Jobs/Graphs/GraphCollection.cs
+++ b/Engine/Source/Programs/Horde/Plugins/Build/HordeServer.Build/Jobs/Graphs/GraphCollection.cs
@@ -462,10 +462,10 @@ namespace HordeServer.Jobs.Graphs
 		}
 
 		/// <inheritdoc/>
-		public async Task<IGraph> AddAsync(ITemplate template, string? streamInitialAgentType, CancellationToken cancellationToken)
+		public async Task<IGraph> AddAsync(ITemplate template, CancellationToken cancellationToken)
 		{
 			Node node = new Node(IJob.SetupNodeName, null, null, Array.Empty<NodeRef>(), Array.Empty<NodeRef>(), Priority.High, true, false, true, null, null, null);
-			NodeGroup group = new NodeGroup(template.InitialAgentType ?? streamInitialAgentType ?? "Win64", new List<Node> { node });
+			NodeGroup group = new NodeGroup(template.InitialAgentType ?? "Win64", new List<Node> { node });
 
 			GraphDocument graph = new GraphDocument(new List<NodeGroup> { group }, new List<Aggregate>(), new List<Label>(), new List<GraphArtifact>());
 			await AddAsync(graph, cancellationToken);
diff --git a/Engine/Source/Programs/Horde/Plugins/Build/HordeServer.Build/Jobs/Graphs/IGraphCollection.cs b/Engine/Source/Programs/Horde/Plugins/Build/HordeServer.Build/Jobs/Graphs/IGraphCollection.cs
--- a/Engine/Source/Programs/Horde/Plugins/Build/HordeServer.Build/Jobs/Graphs/IGraphCollection.cs
+++ b/Engine/Source/Programs/Horde/Plugins/Build/HordeServer.Build/Jobs/Graphs/IGraphCollection.cs
@@ -14,10 +14,9 @@ namespace HordeServer.Jobs.Graphs
 		/// Adds a graph from a template
 		/// </summary>
 		/// <param name="template">The template</param>
-		/// <param name="streamInitialAgentType">Default agent type for the stream</param>
 		/// <param name="cancellationToken">Cancellation token for the operation</param>
 		/// <returns>New graph</returns>
-		Task<IGraph> AddAsync(ITemplate template, string? streamInitialAgentType, CancellationToken cancellationToken = default);
+		Task<IGraph> AddAsync(ITemplate template, CancellationToken cancellationToken = default);
 
 		/// <summary>
 		/// Creates a graph by appending groups and aggregates to an existing graph.
diff --git a/Engine/Source/Programs/Horde/Plugins/Build/HordeServer.Build/Jobs/JobService.cs b/Engine/Source/Programs/Horde/Plugins/Build/HordeServer.Build/Jobs/JobService.cs
--- a/Engine/Source/Programs/Horde/Plugins/Build/HordeServer.Build/Jobs/JobService.cs
+++ b/Engine/Source/Programs/Horde/Plugins/Build/HordeServer.Build/Jobs/JobService.cs
@@ -1179,7 +1179,7 @@ namespace HordeServer.Jobs
 
 					ITemplate template = await _templateCollection.GetOrAddAsync(templateRefConfig);
 
-					IGraph triggerGraph = await _graphs.AddAsync(template, streamConfig.InitialAgentType);
+					IGraph triggerGraph = await _graphs.AddAsync(template, cancellationToken);
 
 					CommitIdWithOrder? commitId = null;
 					CommitIdWithOrder? codeCommitId;
diff --git a/Engine/Source/Programs/Horde/Plugins/Build/HordeServer.Build/Jobs/JobsController.cs b/Engine/Source/Programs/Horde/Plugins/Build/HordeServer.Build/Jobs/JobsController.cs
--- a/Engine/Source/Programs/Horde/Plugins/Build/HordeServer.Build/Jobs/JobsController.cs
+++ b/Engine/Source/Programs/Horde/Plugins/Build/HordeServer.Build/Jobs/JobsController.cs
@@ -163,7 +163,7 @@ namespace HordeServer.Jobs
 			Priority priority = create.Priority ?? template.Priority ?? Priority.Normal;
 
 			// New groups for the job
-			IGraph graph = await _graphs.AddAsync(template, streamConfig.InitialAgentType, cancellationToken);
+			IGraph graph = await _graphs.AddAsync(template, cancellationToken);
 
 			// Get the commits for this stream
 			ICommitCollection commits = _commitService.GetCollection(streamConfig);
diff --git a/Engine/Source/Programs/Horde/Plugins/Build/HordeServer.Build/Jobs/Schedules/ScheduleService.cs b/Engine/Source/Programs/Horde/Plugins/Build/HordeServer.Build/Jobs/Schedules/ScheduleService.cs
--- a/Engine/Source/Programs/Horde/Plugins/Build/HordeServer.Build/Jobs/Schedules/ScheduleService.cs
+++ b/Engine/Source/Programs/Horde/Plugins/Build/HordeServer.Build/Jobs/Schedules/ScheduleService.cs
@@ -473,7 +473,7 @@ namespace HordeServer.Jobs.Schedules
 			ITemplate template = await _templateCollection.GetOrAddAsync(templateRef.Config);
 
 			// Register the graph for it
-			IGraph graph = await _graphs.AddAsync(template, stream.Config.InitialAgentType, cancellationToken);
+			IGraph graph = await _graphs.AddAsync(template, cancellationToken);
 
 			// We may need to submit a new change for any new jobs. This only makes sense if there's one change.
 			if (template.SubmitNewChange != null)
diff --git a/Engine/Source/Programs/Horde/Plugins/Build/HordeServer.Build/Streams/StreamConfig.cs b/Engine/Source/Programs/Horde/Plugins/Build/HordeServer.Build/Streams/StreamConfig.cs
--- a/Engine/Source/Programs/Horde/Plugins/Build/HordeServer.Build/Streams/StreamConfig.cs
+++ b/Engine/Source/Programs/Horde/Plugins/Build/HordeServer.Build/Streams/StreamConfig.cs
@@ -933,6 +933,8 @@ namespace HordeServer.Streams
 		{
 			StreamConfig = streamConfig;
 
+			InitialAgentType ??= streamConfig.InitialAgentType;
+
 			if (Id.IsEmpty)
 			{
 				Id = new TemplateId(StringId.Sanitize(Name));
