Horde: Fix test error.

#rnx

diff --git a/Engine/Source/Programs/Horde/Plugins/Build/HordeServer.Build/Streams/StreamCollection.cs b/Engine/Source/Programs/Horde/Plugins/Build/HordeServer.Build/Streams/StreamCollection.cs
--- a/Engine/Source/Programs/Horde/Plugins/Build/HordeServer.Build/Streams/StreamCollection.cs
+++ b/Engine/Source/Programs/Horde/Plugins/Build/HordeServer.Build/Streams/StreamCollection.cs
@@ -67,7 +67,7 @@ namespace HordeServer.Streams
 			public TemplateId Id => _config.Id;
 			public TemplateRefConfig Config => _config;
 			public ISchedule? Schedule => _schedule;
-			public IReadOnlyList<ITemplateStep> StepStates => (IReadOnlyList<ITemplateStep>?)StepStates ?? Array.Empty<ITemplateStep>();
+			public IReadOnlyList<ITemplateStep> StepStates => (IReadOnlyList<ITemplateStep>?)_document.StepStates ?? Array.Empty<ITemplateStep>();
 		}
 
 		class Schedule : ISchedule
