Horde: Fix static analysis warnings.

#rnx

diff --git a/Engine/Source/Programs/Horde/Plugins/Build/HordeServer.Build/Streams/StreamMessages.cs b/Engine/Source/Programs/Horde/Plugins/Build/HordeServer.Build/Streams/StreamMessages.cs
--- a/Engine/Source/Programs/Horde/Plugins/Build/HordeServer.Build/Streams/StreamMessages.cs
+++ b/Engine/Source/Programs/Horde/Plugins/Build/HordeServer.Build/Streams/StreamMessages.cs
@@ -1,16 +1,11 @@
 // Copyright Epic Games, Inc. All Rights Reserved.
 
 using System.ComponentModel.DataAnnotations;
-using EpicGames.Core;
-using EpicGames.Horde.Commits;
 using EpicGames.Horde.Jobs;
 using EpicGames.Horde.Jobs.Templates;
-using EpicGames.Horde.Projects;
-using EpicGames.Horde.Streams;
 using EpicGames.Horde.Users;
 using HordeServer.Issues;
 using HordeServer.Jobs.Schedules;
-using HordeServer.Projects;
 
 #pragma warning disable CA2227 // Change x to be read-only by removing the property setter
 
