Horde: Adjust namespaces for graph classes in EpicGames.Horde.

#rnx

diff --git a/Engine/Source/Programs/Horde/Plugins/Build/HordeServer.Build.Tests/Devices/DeviceServiceTests.cs b/Engine/Source/Programs/Horde/Plugins/Build/HordeServer.Build.Tests/Devices/DeviceServiceTests.cs
--- a/Engine/Source/Programs/Horde/Plugins/Build/HordeServer.Build.Tests/Devices/DeviceServiceTests.cs
+++ b/Engine/Source/Programs/Horde/Plugins/Build/HordeServer.Build.Tests/Devices/DeviceServiceTests.cs
@@ -5,12 +5,12 @@ using EpicGames.Core;
 using EpicGames.Horde.Commits;
 using EpicGames.Horde.Devices;
 using EpicGames.Horde.Jobs;
+using EpicGames.Horde.Jobs.Graphs;
 using EpicGames.Horde.Jobs.Templates;
 using EpicGames.Horde.Projects;
 using EpicGames.Horde.Streams;
 using HordeServer.Devices;
 using HordeServer.Jobs;
-using HordeServer.Jobs.Graphs;
 using HordeServer.Users;
 using HordeServer.Utilities;
 using Microsoft.AspNetCore.Http;
diff --git a/Engine/Source/Programs/Horde/Plugins/Build/HordeServer.Build.Tests/Fixture.cs b/Engine/Source/Programs/Horde/Plugins/Build/HordeServer.Build.Tests/Fixture.cs
--- a/Engine/Source/Programs/Horde/Plugins/Build/HordeServer.Build.Tests/Fixture.cs
+++ b/Engine/Source/Programs/Horde/Plugins/Build/HordeServer.Build.Tests/Fixture.cs
@@ -5,6 +5,7 @@ using EpicGames.Horde.Agents;
 using EpicGames.Horde.Agents.Pools;
 using EpicGames.Horde.Commits;
 using EpicGames.Horde.Jobs;
+using EpicGames.Horde.Jobs.Graphs;
 using EpicGames.Horde.Jobs.Templates;
 using EpicGames.Horde.Projects;
 using EpicGames.Horde.Streams;
@@ -48,7 +49,7 @@ namespace HordeServer.Tests
 
 		private async Task PopulateAsync(ConfigService configService, IGraphCollection graphCollection, ITemplateCollectionInternal templateCollection, JobService jobService, AgentService agentService, IPluginCollection pluginCollection, ServerSettings serverSettings)
 		{
-			GraphStub fg = new (ContentHash.Empty, 1122);
+			GraphStub fg = new(ContentHash.Empty, 1122);
 			Template = await templateCollection.GetOrAddAsync(new TemplateConfig { Name = "Test template" });
 			Graph = await graphCollection.AddAsync(Template);
 
@@ -114,7 +115,7 @@ namespace HordeServer.Tests
 			Agent1 = await agentService.CreateAgentAsync(new CreateAgentOptions(new AgentId(Agent1Name), false, ""));
 		}
 	}
-	
+
 	internal class GraphStub(ContentHash id, int schema = 1) : IGraph
 	{
 		public ContentHash Id { get; set; } = id;
@@ -124,7 +125,7 @@ namespace HordeServer.Tests
 		public IReadOnlyList<ILabel> Labels { get; set; } = [];
 		public IReadOnlyList<IGraphArtifact> Artifacts { get; set; } = [];
 	}
-	
+
 	internal class TemplateStub(ContentHash hash, string name) : ITemplate
 	{
 		public ContentHash Hash { get; } = hash;
diff --git a/Engine/Source/Programs/Horde/Plugins/Build/HordeServer.Build.Tests/Fleet/JobQueueStrategyTest.cs b/Engine/Source/Programs/Horde/Plugins/Build/HordeServer.Build.Tests/Fleet/JobQueueStrategyTest.cs
--- a/Engine/Source/Programs/Horde/Plugins/Build/HordeServer.Build.Tests/Fleet/JobQueueStrategyTest.cs
+++ b/Engine/Source/Programs/Horde/Plugins/Build/HordeServer.Build.Tests/Fleet/JobQueueStrategyTest.cs
@@ -4,6 +4,7 @@ using EpicGames.Core;
 using EpicGames.Horde.Agents.Pools;
 using EpicGames.Horde.Commits;
 using EpicGames.Horde.Jobs;
+using EpicGames.Horde.Jobs.Graphs;
 using EpicGames.Horde.Jobs.Templates;
 using EpicGames.Horde.Projects;
 using EpicGames.Horde.Streams;
@@ -11,7 +12,6 @@ using HordeServer.Agents;
 using HordeServer.Agents.Fleet;
 using HordeServer.Agents.Pools;
 using HordeServer.Jobs;
-using HordeServer.Jobs.Graphs;
 using HordeServer.Projects;
 using HordeServer.Server;
 using HordeServer.Streams;
diff --git a/Engine/Source/Programs/Horde/Plugins/Build/HordeServer.Build.Tests/Issues/IssueServiceTests.cs b/Engine/Source/Programs/Horde/Plugins/Build/HordeServer.Build.Tests/Issues/IssueServiceTests.cs
--- a/Engine/Source/Programs/Horde/Plugins/Build/HordeServer.Build.Tests/Issues/IssueServiceTests.cs
+++ b/Engine/Source/Programs/Horde/Plugins/Build/HordeServer.Build.Tests/Issues/IssueServiceTests.cs
@@ -9,6 +9,7 @@ using EpicGames.Core;
 using EpicGames.Horde.Commits;
 using EpicGames.Horde.Issues;
 using EpicGames.Horde.Jobs;
+using EpicGames.Horde.Jobs.Graphs;
 using EpicGames.Horde.Jobs.Templates;
 using EpicGames.Horde.Logs;
 using EpicGames.Horde.Projects;
@@ -18,7 +19,6 @@ using EpicGames.Horde.Users;
 using EpicGames.Perforce;
 using HordeServer.Issues;
 using HordeServer.Jobs;
-using HordeServer.Jobs.Graphs;
 using HordeServer.Logs;
 using HordeServer.Projects;
 using HordeServer.Server;
diff --git a/Engine/Source/Programs/Horde/Plugins/Build/HordeServer.Build.Tests/Jobs/Bisect/BisectTests.cs b/Engine/Source/Programs/Horde/Plugins/Build/HordeServer.Build.Tests/Jobs/Bisect/BisectTests.cs
--- a/Engine/Source/Programs/Horde/Plugins/Build/HordeServer.Build.Tests/Jobs/Bisect/BisectTests.cs
+++ b/Engine/Source/Programs/Horde/Plugins/Build/HordeServer.Build.Tests/Jobs/Bisect/BisectTests.cs
@@ -4,12 +4,12 @@ using System.Security.Claims;
 using EpicGames.Core;
 using EpicGames.Horde.Commits;
 using EpicGames.Horde.Jobs;
+using EpicGames.Horde.Jobs.Graphs;
 using EpicGames.Horde.Jobs.Templates;
 using EpicGames.Horde.Projects;
 using EpicGames.Horde.Streams;
 using HordeServer.Jobs;
 using HordeServer.Jobs.Bisect;
-using HordeServer.Jobs.Graphs;
 using HordeServer.Logs;
 using HordeServer.Projects;
 using HordeServer.Streams;
diff --git a/Engine/Source/Programs/Horde/Plugins/Build/HordeServer.Build.Tests/Jobs/GetLabelStatesTests.cs b/Engine/Source/Programs/Horde/Plugins/Build/HordeServer.Build.Tests/Jobs/GetLabelStatesTests.cs
--- a/Engine/Source/Programs/Horde/Plugins/Build/HordeServer.Build.Tests/Jobs/GetLabelStatesTests.cs
+++ b/Engine/Source/Programs/Horde/Plugins/Build/HordeServer.Build.Tests/Jobs/GetLabelStatesTests.cs
@@ -1,8 +1,8 @@
 // Copyright Epic Games, Inc. All Rights Reserved.
 
 using EpicGames.Horde.Jobs;
+using EpicGames.Horde.Jobs.Graphs;
 using HordeServer.Jobs;
-using HordeServer.Jobs.Graphs;
 using Moq;
 
 namespace HordeServer.Tests.Jobs
diff --git a/Engine/Source/Programs/Horde/Plugins/Build/HordeServer.Build.Tests/Jobs/JobCollectionTests.cs b/Engine/Source/Programs/Horde/Plugins/Build/HordeServer.Build.Tests/Jobs/JobCollectionTests.cs
--- a/Engine/Source/Programs/Horde/Plugins/Build/HordeServer.Build.Tests/Jobs/JobCollectionTests.cs
+++ b/Engine/Source/Programs/Horde/Plugins/Build/HordeServer.Build.Tests/Jobs/JobCollectionTests.cs
@@ -9,6 +9,7 @@ using EpicGames.Horde.Agents.Sessions;
 using EpicGames.Horde.Artifacts;
 using EpicGames.Horde.Commits;
 using EpicGames.Horde.Jobs;
+using EpicGames.Horde.Jobs.Graphs;
 using EpicGames.Horde.Jobs.Templates;
 using EpicGames.Horde.Logs;
 using EpicGames.Horde.Projects;
@@ -16,7 +17,6 @@ using EpicGames.Horde.Streams;
 using EpicGames.Horde.Users;
 using HordeServer.Agents.Sessions;
 using HordeServer.Jobs;
-using HordeServer.Jobs.Graphs;
 using HordeServer.Logs;
 using HordeServer.Projects;
 using HordeServer.Streams;
diff --git a/Engine/Source/Programs/Horde/Plugins/Build/HordeServer.Build.Tests/Jobs/JobExpirationTests.cs b/Engine/Source/Programs/Horde/Plugins/Build/HordeServer.Build.Tests/Jobs/JobExpirationTests.cs
--- a/Engine/Source/Programs/Horde/Plugins/Build/HordeServer.Build.Tests/Jobs/JobExpirationTests.cs
+++ b/Engine/Source/Programs/Horde/Plugins/Build/HordeServer.Build.Tests/Jobs/JobExpirationTests.cs
@@ -1,11 +1,11 @@
 // Copyright Epic Games, Inc. All Rights Reserved.
 
 using EpicGames.Horde.Commits;
+using EpicGames.Horde.Jobs.Graphs;
 using EpicGames.Horde.Jobs.Templates;
 using EpicGames.Horde.Projects;
 using EpicGames.Horde.Streams;
 using HordeServer.Jobs;
-using HordeServer.Jobs.Graphs;
 using HordeServer.Projects;
 using HordeServer.Server;
 using HordeServer.Streams;
diff --git a/Engine/Source/Programs/Horde/Plugins/Build/HordeServer.Build.Tests/Jobs/JobServiceTests.cs b/Engine/Source/Programs/Horde/Plugins/Build/HordeServer.Build.Tests/Jobs/JobServiceTests.cs
--- a/Engine/Source/Programs/Horde/Plugins/Build/HordeServer.Build.Tests/Jobs/JobServiceTests.cs
+++ b/Engine/Source/Programs/Horde/Plugins/Build/HordeServer.Build.Tests/Jobs/JobServiceTests.cs
@@ -6,13 +6,13 @@ using EpicGames.Horde.Agents;
 using EpicGames.Horde.Agents.Pools;
 using EpicGames.Horde.Commits;
 using EpicGames.Horde.Jobs;
+using EpicGames.Horde.Jobs.Graphs;
 using EpicGames.Horde.Jobs.Templates;
 using EpicGames.Horde.Streams;
 using HordeCommon.Rpc.Messages;
 using HordeServer.Agents;
 using HordeServer.Commits;
 using HordeServer.Jobs;
-using HordeServer.Jobs.Graphs;
 using HordeServer.Jobs.Templates;
 using HordeServer.Logs;
 using HordeServer.Projects;
diff --git a/Engine/Source/Programs/Horde/Plugins/Build/HordeServer.Build.Tests/Jobs/JobTaskSourceTests.cs b/Engine/Source/Programs/Horde/Plugins/Build/HordeServer.Build.Tests/Jobs/JobTaskSourceTests.cs
--- a/Engine/Source/Programs/Horde/Plugins/Build/HordeServer.Build.Tests/Jobs/JobTaskSourceTests.cs
+++ b/Engine/Source/Programs/Horde/Plugins/Build/HordeServer.Build.Tests/Jobs/JobTaskSourceTests.cs
@@ -4,12 +4,12 @@ using EpicGames.Horde.Agents;
 using EpicGames.Horde.Agents.Pools;
 using EpicGames.Horde.Commits;
 using EpicGames.Horde.Jobs;
+using EpicGames.Horde.Jobs.Graphs;
 using EpicGames.Horde.Users;
 using HordeCommon.Rpc.Messages;
 using HordeServer.Agents;
 using HordeServer.Agents.Pools;
 using HordeServer.Jobs;
-using HordeServer.Jobs.Graphs;
 using HordeServer.Streams;
 using HordeServer.Utilities;
 
diff --git a/Engine/Source/Programs/Horde/Plugins/Build/HordeServer.Build.Tests/Jobs/Schedules/SchedulerTests.cs b/Engine/Source/Programs/Horde/Plugins/Build/HordeServer.Build.Tests/Jobs/Schedules/SchedulerTests.cs
--- a/Engine/Source/Programs/Horde/Plugins/Build/HordeServer.Build.Tests/Jobs/Schedules/SchedulerTests.cs
+++ b/Engine/Source/Programs/Horde/Plugins/Build/HordeServer.Build.Tests/Jobs/Schedules/SchedulerTests.cs
@@ -2,11 +2,11 @@
 
 using EpicGames.Horde.Commits;
 using EpicGames.Horde.Jobs;
+using EpicGames.Horde.Jobs.Graphs;
 using EpicGames.Horde.Jobs.Templates;
 using EpicGames.Horde.Projects;
 using EpicGames.Horde.Streams;
 using HordeServer.Jobs;
-using HordeServer.Jobs.Graphs;
 using HordeServer.Jobs.Schedules;
 using HordeServer.Jobs.Templates;
 using HordeServer.Logs;
@@ -132,7 +132,7 @@ namespace HordeServer.Tests.Jobs.Schedules
 
 			Mock<ISchedule> scheduleMock = new Mock<ISchedule>(MockBehavior.Strict);
 			scheduleMock.SetupGet(x => x.Patterns).Returns([new SchedulePatternConfig(new List<DayOfWeek> { DayOfWeek.Friday, DayOfWeek.Sunday }, 13 * 60, null, null)]);
-			
+
 			ISchedule schedule = scheduleMock.Object;
 
 			DateTime? nextTime = schedule.GetNextTriggerTimeUtc(startTime, TimeZoneInfo.Utc);
diff --git a/Engine/Source/Programs/Horde/Plugins/Build/HordeServer.Build.Tests/Jobs/TestData/TestDataTests.cs b/Engine/Source/Programs/Horde/Plugins/Build/HordeServer.Build.Tests/Jobs/TestData/TestDataTests.cs
--- a/Engine/Source/Programs/Horde/Plugins/Build/HordeServer.Build.Tests/Jobs/TestData/TestDataTests.cs
+++ b/Engine/Source/Programs/Horde/Plugins/Build/HordeServer.Build.Tests/Jobs/TestData/TestDataTests.cs
@@ -4,13 +4,13 @@ using System.Security.Claims;
 using EpicGames.Horde;
 using EpicGames.Horde.Commits;
 using EpicGames.Horde.Jobs;
+using EpicGames.Horde.Jobs.Graphs;
 using EpicGames.Horde.Jobs.Templates;
 using EpicGames.Horde.Jobs.TestData;
 using EpicGames.Horde.Logs;
 using EpicGames.Horde.Projects;
 using EpicGames.Horde.Streams;
 using HordeServer.Jobs;
-using HordeServer.Jobs.Graphs;
 using HordeServer.Jobs.TestData;
 using HordeServer.Logs;
 using HordeServer.Projects;
diff --git a/Engine/Source/Programs/Horde/Plugins/Build/HordeServer.Build.Tests/Notifications/NotificationServiceTests.cs b/Engine/Source/Programs/Horde/Plugins/Build/HordeServer.Build.Tests/Notifications/NotificationServiceTests.cs
--- a/Engine/Source/Programs/Horde/Plugins/Build/HordeServer.Build.Tests/Notifications/NotificationServiceTests.cs
+++ b/Engine/Source/Programs/Horde/Plugins/Build/HordeServer.Build.Tests/Notifications/NotificationServiceTests.cs
@@ -2,6 +2,7 @@
 
 using EpicGames.Horde.Commits;
 using EpicGames.Horde.Jobs;
+using EpicGames.Horde.Jobs.Graphs;
 using EpicGames.Horde.Streams;
 using HordeServer.Agents;
 using HordeServer.Agents.Pools;
@@ -9,7 +10,6 @@ using HordeServer.Configuration;
 using HordeServer.Devices;
 using HordeServer.Issues;
 using HordeServer.Jobs;
-using HordeServer.Jobs.Graphs;
 using HordeServer.Logs;
 using HordeServer.Notifications;
 using HordeServer.Streams;
diff --git a/Engine/Source/Programs/Horde/Plugins/Build/HordeServer.Build/Agents/Fleet/JobQueueStrategy.cs b/Engine/Source/Programs/Horde/Plugins/Build/HordeServer.Build/Agents/Fleet/JobQueueStrategy.cs
--- a/Engine/Source/Programs/Horde/Plugins/Build/HordeServer.Build/Agents/Fleet/JobQueueStrategy.cs
+++ b/Engine/Source/Programs/Horde/Plugins/Build/HordeServer.Build/Agents/Fleet/JobQueueStrategy.cs
@@ -3,6 +3,7 @@
 using EpicGames.Core;
 using EpicGames.Horde.Agents.Pools;
 using EpicGames.Horde.Jobs;
+using EpicGames.Horde.Jobs.Graphs;
 using EpicGames.Horde.Streams;
 using HordeServer.Agents.Pools;
 using HordeServer.Jobs;
diff --git a/Engine/Source/Programs/Horde/Plugins/Build/HordeServer.Build/Devices/DeviceService.cs b/Engine/Source/Programs/Horde/Plugins/Build/HordeServer.Build/Devices/DeviceService.cs
--- a/Engine/Source/Programs/Horde/Plugins/Build/HordeServer.Build/Devices/DeviceService.cs
+++ b/Engine/Source/Programs/Horde/Plugins/Build/HordeServer.Build/Devices/DeviceService.cs
@@ -5,10 +5,10 @@ using EpicGames.Core;
 using EpicGames.Horde.Acls;
 using EpicGames.Horde.Devices;
 using EpicGames.Horde.Jobs;
+using EpicGames.Horde.Jobs.Graphs;
 using EpicGames.Horde.Projects;
 using EpicGames.Horde.Users;
 using HordeServer.Jobs;
-using HordeServer.Jobs.Graphs;
 using HordeServer.Notifications;
 using HordeServer.Projects;
 using HordeServer.Server;
diff --git a/Engine/Source/Programs/Horde/Plugins/Build/HordeServer.Build/Issues/IIssueCollection.cs b/Engine/Source/Programs/Horde/Plugins/Build/HordeServer.Build/Issues/IIssueCollection.cs
--- a/Engine/Source/Programs/Horde/Plugins/Build/HordeServer.Build/Issues/IIssueCollection.cs
+++ b/Engine/Source/Programs/Horde/Plugins/Build/HordeServer.Build/Issues/IIssueCollection.cs
@@ -3,13 +3,13 @@
 using EpicGames.Horde.Commits;
 using EpicGames.Horde.Issues;
 using EpicGames.Horde.Jobs;
+using EpicGames.Horde.Jobs.Graphs;
 using EpicGames.Horde.Jobs.Templates;
 using EpicGames.Horde.Logs;
 using EpicGames.Horde.Streams;
 using EpicGames.Horde.Users;
 using HordeServer.Auditing;
 using HordeServer.Jobs;
-using HordeServer.Jobs.Graphs;
 using MongoDB.Bson;
 using MongoDB.Driver;
 
@@ -245,23 +245,23 @@ namespace HordeServer.Issues
 	/// <param name="WorkflowThreadUrl">The workflow thread url associated with the issue</param>
 	public record class UpdateIssueOptions
 	(
-		IssueSeverity? Severity = null, 
-		string? Summary = null, 
-		string? UserSummary = null, 
-		string? Description = null, 
-		bool? Promoted = null, 
-		UserId? OwnerId = null, 
-		UserId? NominatedById = null, 
-		bool? Acknowledged = null, 
-		UserId? DeclinedById = null, 
-		CommitId? FixCommitId = null, 
-		bool? FixSystemic = null, 
-		UserId? ResolvedById = null, 
-		List<ObjectId>? ExcludeSpanIds = null, 
-		DateTime? LastSeenAt = null, 
-		string? ExternalIssueKey = null, 
-		UserId? QuarantinedById = null, 
-		UserId? ForceClosedById = null, 
+		IssueSeverity? Severity = null,
+		string? Summary = null,
+		string? UserSummary = null,
+		string? Description = null,
+		bool? Promoted = null,
+		UserId? OwnerId = null,
+		UserId? NominatedById = null,
+		bool? Acknowledged = null,
+		UserId? DeclinedById = null,
+		CommitId? FixCommitId = null,
+		bool? FixSystemic = null,
+		UserId? ResolvedById = null,
+		List<ObjectId>? ExcludeSpanIds = null,
+		DateTime? LastSeenAt = null,
+		string? ExternalIssueKey = null,
+		UserId? QuarantinedById = null,
+		UserId? ForceClosedById = null,
 		Uri? WorkflowThreadUrl = null
 	);
 
diff --git a/Engine/Source/Programs/Horde/Plugins/Build/HordeServer.Build/Issues/IIssueStep.cs b/Engine/Source/Programs/Horde/Plugins/Build/HordeServer.Build/Issues/IIssueStep.cs
--- a/Engine/Source/Programs/Horde/Plugins/Build/HordeServer.Build/Issues/IIssueStep.cs
+++ b/Engine/Source/Programs/Horde/Plugins/Build/HordeServer.Build/Issues/IIssueStep.cs
@@ -3,8 +3,8 @@
 using EpicGames.Horde.Commits;
 using EpicGames.Horde.Issues;
 using EpicGames.Horde.Jobs;
+using EpicGames.Horde.Jobs.Graphs;
 using EpicGames.Horde.Logs;
-using HordeServer.Jobs.Graphs;
 using MongoDB.Bson;
 
 namespace HordeServer.Issues
diff --git a/Engine/Source/Programs/Horde/Plugins/Build/HordeServer.Build/Issues/IssueCollection.cs b/Engine/Source/Programs/Horde/Plugins/Build/HordeServer.Build/Issues/IssueCollection.cs
--- a/Engine/Source/Programs/Horde/Plugins/Build/HordeServer.Build/Issues/IssueCollection.cs
+++ b/Engine/Source/Programs/Horde/Plugins/Build/HordeServer.Build/Issues/IssueCollection.cs
@@ -7,6 +7,7 @@ using EpicGames.Core;
 using EpicGames.Horde.Commits;
 using EpicGames.Horde.Issues;
 using EpicGames.Horde.Jobs;
+using EpicGames.Horde.Jobs.Graphs;
 using EpicGames.Horde.Jobs.Templates;
 using EpicGames.Horde.Logs;
 using EpicGames.Horde.Streams;
@@ -16,7 +17,6 @@ using EpicGames.Redis.Utility;
 using HordeServer.Auditing;
 using HordeServer.Commits;
 using HordeServer.Jobs;
-using HordeServer.Jobs.Graphs;
 using HordeServer.Logs;
 using HordeServer.Server;
 using HordeServer.Streams;
@@ -102,7 +102,7 @@ namespace HordeServer.Issues
 			[BsonIgnore]
 			public CommitId? FixCommitId
 			{
-				get => (FixCommitName != null)? new CommitId(FixCommitName) : (FixChange != null)? CommitId.FromPerforceChange(FixChange.Value) : null;
+				get => (FixCommitName != null) ? new CommitId(FixCommitName) : (FixChange != null) ? CommitId.FromPerforceChange(FixChange.Value) : null;
 				set => (FixCommitName, FixChange) = (value?.Name, null);
 			}
 
@@ -521,7 +521,7 @@ namespace HordeServer.Issues
 			[BsonIgnore]
 			public CommitIdWithOrder CommitId
 			{
-				get => (CommitName != null)? new CommitIdWithOrder(CommitName, CommitOrder) : CommitIdWithOrder.FromPerforceChange(CommitOrder);
+				get => (CommitName != null) ? new CommitIdWithOrder(CommitName, CommitOrder) : CommitIdWithOrder.FromPerforceChange(CommitOrder);
 				set => (CommitName, CommitOrder) = (value.Name, value.Order);
 			}
 
@@ -536,7 +536,7 @@ namespace HordeServer.Issues
 			[BsonIgnore]
 			public CommitIdWithOrder? SourceCommitId
 			{
-				get => (SourceCommitName != null)? new CommitIdWithOrder(SourceCommitName, SourceCommitOrder ?? 0) : (SourceCommitOrder != null)? CommitIdWithOrder.FromPerforceChange(SourceCommitOrder.Value) : null;
+				get => (SourceCommitName != null) ? new CommitIdWithOrder(SourceCommitName, SourceCommitOrder ?? 0) : (SourceCommitOrder != null) ? CommitIdWithOrder.FromPerforceChange(SourceCommitOrder.Value) : null;
 				set => (SourceCommitName, SourceCommitOrder) = (value?.Name, value?.Order);
 			}
 
@@ -567,7 +567,7 @@ namespace HordeServer.Issues
 			[BsonIgnore]
 			public CommitIdWithOrder CommitId
 			{
-				get => (CommitName != null)? new CommitIdWithOrder(CommitName, CommitOrder) : CommitIdWithOrder.FromPerforceChange(CommitOrder);
+				get => (CommitName != null) ? new CommitIdWithOrder(CommitName, CommitOrder) : CommitIdWithOrder.FromPerforceChange(CommitOrder);
 				set => (CommitName, CommitOrder) = (value.Name, value.Order);
 			}
 
diff --git a/Engine/Source/Programs/Horde/Plugins/Build/HordeServer.Build/Issues/IssueReportService.cs b/Engine/Source/Programs/Horde/Plugins/Build/HordeServer.Build/Issues/IssueReportService.cs
--- a/Engine/Source/Programs/Horde/Plugins/Build/HordeServer.Build/Issues/IssueReportService.cs
+++ b/Engine/Source/Programs/Horde/Plugins/Build/HordeServer.Build/Issues/IssueReportService.cs
@@ -3,6 +3,7 @@
 using EpicGames.Core;
 using EpicGames.Horde.Issues;
 using EpicGames.Horde.Jobs;
+using EpicGames.Horde.Jobs.Graphs;
 using EpicGames.Horde.Jobs.Templates;
 using EpicGames.Horde.Streams;
 using HordeServer.Jobs;
diff --git a/Engine/Source/Programs/Horde/Plugins/Build/HordeServer.Build/Issues/IssueService.cs b/Engine/Source/Programs/Horde/Plugins/Build/HordeServer.Build/Issues/IssueService.cs
--- a/Engine/Source/Programs/Horde/Plugins/Build/HordeServer.Build/Issues/IssueService.cs
+++ b/Engine/Source/Programs/Horde/Plugins/Build/HordeServer.Build/Issues/IssueService.cs
@@ -9,13 +9,13 @@ using EpicGames.Core;
 using EpicGames.Horde.Commits;
 using EpicGames.Horde.Issues;
 using EpicGames.Horde.Jobs;
+using EpicGames.Horde.Jobs.Graphs;
 using EpicGames.Horde.Jobs.Templates;
 using EpicGames.Horde.Logs;
 using EpicGames.Horde.Streams;
 using EpicGames.Horde.Users;
 using HordeServer.Commits;
 using HordeServer.Jobs;
-using HordeServer.Jobs.Graphs;
 using HordeServer.Logs;
 using HordeServer.Streams;
 using HordeServer.Users;
diff --git a/Engine/Source/Programs/Horde/Plugins/Build/HordeServer.Build/Issues/IssuesController.cs b/Engine/Source/Programs/Horde/Plugins/Build/HordeServer.Build/Issues/IssuesController.cs
--- a/Engine/Source/Programs/Horde/Plugins/Build/HordeServer.Build/Issues/IssuesController.cs
+++ b/Engine/Source/Programs/Horde/Plugins/Build/HordeServer.Build/Issues/IssuesController.cs
@@ -4,6 +4,7 @@ using EpicGames.Core;
 using EpicGames.Horde.Commits;
 using EpicGames.Horde.Issues;
 using EpicGames.Horde.Jobs;
+using EpicGames.Horde.Jobs.Graphs;
 using EpicGames.Horde.Jobs.Templates;
 using EpicGames.Horde.Logs;
 using EpicGames.Horde.Streams;
@@ -12,7 +13,6 @@ using HordeServer.Auditing;
 using HordeServer.Commits;
 using HordeServer.Issues.External;
 using HordeServer.Jobs;
-using HordeServer.Jobs.Graphs;
 using HordeServer.Logs;
 using HordeServer.Streams;
 using HordeServer.Users;
diff --git a/Engine/Source/Programs/Horde/Plugins/Build/HordeServer.Build/Issues/WorkflowConfig.cs b/Engine/Source/Programs/Horde/Plugins/Build/HordeServer.Build/Issues/WorkflowConfig.cs
--- a/Engine/Source/Programs/Horde/Plugins/Build/HordeServer.Build/Issues/WorkflowConfig.cs
+++ b/Engine/Source/Programs/Horde/Plugins/Build/HordeServer.Build/Issues/WorkflowConfig.cs
@@ -2,7 +2,7 @@
 
 using System.ComponentModel.DataAnnotations;
 using EpicGames.Horde.Issues;
-using HordeServer.Jobs.Graphs;
+using EpicGames.Horde.Jobs.Graphs;
 
 #pragma warning disable CA2227 //  Change x to be read-only by removing the property setter
 
diff --git a/Engine/Source/Programs/Horde/Plugins/Build/HordeServer.Build/Jobs/Bisect/BisectController.cs b/Engine/Source/Programs/Horde/Plugins/Build/HordeServer.Build/Jobs/Bisect/BisectController.cs
--- a/Engine/Source/Programs/Horde/Plugins/Build/HordeServer.Build/Jobs/Bisect/BisectController.cs
+++ b/Engine/Source/Programs/Horde/Plugins/Build/HordeServer.Build/Jobs/Bisect/BisectController.cs
@@ -3,6 +3,7 @@
 using EpicGames.Horde.Commits;
 using EpicGames.Horde.Jobs;
 using EpicGames.Horde.Jobs.Bisect;
+using EpicGames.Horde.Jobs.Graphs;
 using EpicGames.Horde.Jobs.Templates;
 using EpicGames.Horde.Streams;
 using EpicGames.Horde.Users;
diff --git a/Engine/Source/Programs/Horde/Plugins/Build/HordeServer.Build/Jobs/Bisect/BisectService.cs b/Engine/Source/Programs/Horde/Plugins/Build/HordeServer.Build/Jobs/Bisect/BisectService.cs
--- a/Engine/Source/Programs/Horde/Plugins/Build/HordeServer.Build/Jobs/Bisect/BisectService.cs
+++ b/Engine/Source/Programs/Horde/Plugins/Build/HordeServer.Build/Jobs/Bisect/BisectService.cs
@@ -2,6 +2,7 @@
 
 using EpicGames.Core;
 using EpicGames.Horde.Jobs;
+using EpicGames.Horde.Jobs.Graphs;
 using EpicGames.Horde.Jobs.Templates;
 using HordeServer.Commits;
 using HordeServer.Jobs.Graphs;
diff --git a/Engine/Source/Programs/Horde/Plugins/Build/HordeServer.Build/Jobs/Bisect/IBisectTaskCollection.cs b/Engine/Source/Programs/Horde/Plugins/Build/HordeServer.Build/Jobs/Bisect/IBisectTaskCollection.cs
--- a/Engine/Source/Programs/Horde/Plugins/Build/HordeServer.Build/Jobs/Bisect/IBisectTaskCollection.cs
+++ b/Engine/Source/Programs/Horde/Plugins/Build/HordeServer.Build/Jobs/Bisect/IBisectTaskCollection.cs
@@ -3,8 +3,8 @@
 using EpicGames.Horde.Commits;
 using EpicGames.Horde.Jobs;
 using EpicGames.Horde.Jobs.Bisect;
+using EpicGames.Horde.Jobs.Graphs;
 using EpicGames.Horde.Users;
-using HordeServer.Jobs.Graphs;
 using Microsoft.Extensions.Logging;
 
 namespace HordeServer.Jobs.Bisect
diff --git a/Engine/Source/Programs/Horde/Plugins/Build/HordeServer.Build/Jobs/DebugJobsController.cs b/Engine/Source/Programs/Horde/Plugins/Build/HordeServer.Build/Jobs/DebugJobsController.cs
--- a/Engine/Source/Programs/Horde/Plugins/Build/HordeServer.Build/Jobs/DebugJobsController.cs
+++ b/Engine/Source/Programs/Horde/Plugins/Build/HordeServer.Build/Jobs/DebugJobsController.cs
@@ -4,8 +4,8 @@ using System.Net;
 using System.Reflection;
 using System.Text;
 using EpicGames.Horde.Jobs;
+using EpicGames.Horde.Jobs.Graphs;
 using HordeServer.Agents.Leases;
-using HordeServer.Jobs.Graphs;
 using HordeServer.Server;
 using HordeServer.Streams;
 using HordeServer.Utilities;
diff --git a/Engine/Source/Programs/Horde/Plugins/Build/HordeServer.Build/Jobs/Graphs/GraphCollection.cs b/Engine/Source/Programs/Horde/Plugins/Build/HordeServer.Build/Jobs/Graphs/GraphCollection.cs
--- a/Engine/Source/Programs/Horde/Plugins/Build/HordeServer.Build/Jobs/Graphs/GraphCollection.cs
+++ b/Engine/Source/Programs/Horde/Plugins/Build/HordeServer.Build/Jobs/Graphs/GraphCollection.cs
@@ -4,6 +4,7 @@ using System.Diagnostics;
 using EpicGames.Core;
 using EpicGames.Horde.Artifacts;
 using EpicGames.Horde.Jobs;
+using EpicGames.Horde.Jobs.Graphs;
 using EpicGames.Horde.Jobs.Templates;
 using HordeServer.Server;
 using Microsoft.Extensions.Caching.Memory;
diff --git a/Engine/Source/Programs/Horde/Plugins/Build/HordeServer.Build/Jobs/Graphs/IGraphCollection.cs b/Engine/Source/Programs/Horde/Plugins/Build/HordeServer.Build/Jobs/Graphs/IGraphCollection.cs
--- a/Engine/Source/Programs/Horde/Plugins/Build/HordeServer.Build/Jobs/Graphs/IGraphCollection.cs
+++ b/Engine/Source/Programs/Horde/Plugins/Build/HordeServer.Build/Jobs/Graphs/IGraphCollection.cs
@@ -1,6 +1,7 @@
 // Copyright Epic Games, Inc. All Rights Reserved.
 
 using EpicGames.Core;
+using EpicGames.Horde.Jobs.Graphs;
 using EpicGames.Horde.Jobs.Templates;
 
 namespace HordeServer.Jobs.Graphs
diff --git a/Engine/Source/Programs/Horde/Plugins/Build/HordeServer.Build/Jobs/IJob.cs b/Engine/Source/Programs/Horde/Plugins/Build/HordeServer.Build/Jobs/IJob.cs
--- a/Engine/Source/Programs/Horde/Plugins/Build/HordeServer.Build/Jobs/IJob.cs
+++ b/Engine/Source/Programs/Horde/Plugins/Build/HordeServer.Build/Jobs/IJob.cs
@@ -11,13 +11,13 @@ using EpicGames.Horde.Agents.Sessions;
 using EpicGames.Horde.Commits;
 using EpicGames.Horde.Jobs;
 using EpicGames.Horde.Jobs.Bisect;
+using EpicGames.Horde.Jobs.Graphs;
 using EpicGames.Horde.Jobs.Templates;
 using EpicGames.Horde.Logs;
 using EpicGames.Horde.Streams;
 using EpicGames.Horde.Users;
 using Horde.Common.Rpc;
 using HordeServer.Acls;
-using HordeServer.Jobs.Graphs;
 using HordeServer.Jobs.Timing;
 using HordeServer.Streams;
 using HordeServer.Ugs;
diff --git a/Engine/Source/Programs/Horde/Plugins/Build/HordeServer.Build/Jobs/IJobCollection.cs b/Engine/Source/Programs/Horde/Plugins/Build/HordeServer.Build/Jobs/IJobCollection.cs
--- a/Engine/Source/Programs/Horde/Plugins/Build/HordeServer.Build/Jobs/IJobCollection.cs
+++ b/Engine/Source/Programs/Horde/Plugins/Build/HordeServer.Build/Jobs/IJobCollection.cs
@@ -4,11 +4,11 @@ using EpicGames.Core;
 using EpicGames.Horde.Commits;
 using EpicGames.Horde.Jobs;
 using EpicGames.Horde.Jobs.Bisect;
+using EpicGames.Horde.Jobs.Graphs;
 using EpicGames.Horde.Jobs.Templates;
 using EpicGames.Horde.Streams;
 using EpicGames.Horde.Users;
 using HordeServer.Acls;
-using HordeServer.Jobs.Graphs;
 using HordeServer.Streams;
 
 #pragma warning disable CA2227 //  Change x to be read-only by removing the property setter
diff --git a/Engine/Source/Programs/Horde/Plugins/Build/HordeServer.Build/Jobs/IJobStepRefCollection.cs b/Engine/Source/Programs/Horde/Plugins/Build/HordeServer.Build/Jobs/IJobStepRefCollection.cs
--- a/Engine/Source/Programs/Horde/Plugins/Build/HordeServer.Build/Jobs/IJobStepRefCollection.cs
+++ b/Engine/Source/Programs/Horde/Plugins/Build/HordeServer.Build/Jobs/IJobStepRefCollection.cs
@@ -4,10 +4,10 @@ using EpicGames.Horde.Agents;
 using EpicGames.Horde.Agents.Pools;
 using EpicGames.Horde.Commits;
 using EpicGames.Horde.Jobs;
+using EpicGames.Horde.Jobs.Graphs;
 using EpicGames.Horde.Jobs.Templates;
 using EpicGames.Horde.Logs;
 using EpicGames.Horde.Streams;
-using HordeServer.Jobs.Graphs;
 using Microsoft.Extensions.Logging;
 
 namespace HordeServer.Jobs
diff --git a/Engine/Source/Programs/Horde/Plugins/Build/HordeServer.Build/Jobs/JobCollection.cs b/Engine/Source/Programs/Horde/Plugins/Build/HordeServer.Build/Jobs/JobCollection.cs
--- a/Engine/Source/Programs/Horde/Plugins/Build/HordeServer.Build/Jobs/JobCollection.cs
+++ b/Engine/Source/Programs/Horde/Plugins/Build/HordeServer.Build/Jobs/JobCollection.cs
@@ -12,6 +12,7 @@ using EpicGames.Horde.Agents.Sessions;
 using EpicGames.Horde.Commits;
 using EpicGames.Horde.Jobs;
 using EpicGames.Horde.Jobs.Bisect;
+using EpicGames.Horde.Jobs.Graphs;
 using EpicGames.Horde.Jobs.Templates;
 using EpicGames.Horde.Logs;
 using EpicGames.Horde.Streams;
diff --git a/Engine/Source/Programs/Horde/Plugins/Build/HordeServer.Build/Jobs/JobRpcService.cs b/Engine/Source/Programs/Horde/Plugins/Build/HordeServer.Build/Jobs/JobRpcService.cs
--- a/Engine/Source/Programs/Horde/Plugins/Build/HordeServer.Build/Jobs/JobRpcService.cs
+++ b/Engine/Source/Programs/Horde/Plugins/Build/HordeServer.Build/Jobs/JobRpcService.cs
@@ -11,6 +11,7 @@ using EpicGames.Horde.Agents.Sessions;
 using EpicGames.Horde.Artifacts;
 using EpicGames.Horde.Commits;
 using EpicGames.Horde.Jobs;
+using EpicGames.Horde.Jobs.Graphs;
 using EpicGames.Horde.Jobs.Templates;
 using EpicGames.Horde.Logs;
 using Google.Protobuf.WellKnownTypes;
diff --git a/Engine/Source/Programs/Horde/Plugins/Build/HordeServer.Build/Jobs/JobService.cs b/Engine/Source/Programs/Horde/Plugins/Build/HordeServer.Build/Jobs/JobService.cs
--- a/Engine/Source/Programs/Horde/Plugins/Build/HordeServer.Build/Jobs/JobService.cs
+++ b/Engine/Source/Programs/Horde/Plugins/Build/HordeServer.Build/Jobs/JobService.cs
@@ -9,6 +9,7 @@ using EpicGames.Horde.Agents.Leases;
 using EpicGames.Horde.Agents.Sessions;
 using EpicGames.Horde.Commits;
 using EpicGames.Horde.Jobs;
+using EpicGames.Horde.Jobs.Graphs;
 using EpicGames.Horde.Jobs.Templates;
 using EpicGames.Horde.Logs;
 using EpicGames.Horde.Streams;
diff --git a/Engine/Source/Programs/Horde/Plugins/Build/HordeServer.Build/Jobs/JobTaskSource.cs b/Engine/Source/Programs/Horde/Plugins/Build/HordeServer.Build/Jobs/JobTaskSource.cs
--- a/Engine/Source/Programs/Horde/Plugins/Build/HordeServer.Build/Jobs/JobTaskSource.cs
+++ b/Engine/Source/Programs/Horde/Plugins/Build/HordeServer.Build/Jobs/JobTaskSource.cs
@@ -10,6 +10,7 @@ using EpicGames.Horde.Agents.Leases;
 using EpicGames.Horde.Agents.Pools;
 using EpicGames.Horde.Common;
 using EpicGames.Horde.Jobs;
+using EpicGames.Horde.Jobs.Graphs;
 using EpicGames.Horde.Logs;
 using EpicGames.Horde.Streams;
 using HordeCommon.Rpc.Messages;
diff --git a/Engine/Source/Programs/Horde/Plugins/Build/HordeServer.Build/Jobs/Schedules/ScheduleService.cs b/Engine/Source/Programs/Horde/Plugins/Build/HordeServer.Build/Jobs/Schedules/ScheduleService.cs
--- a/Engine/Source/Programs/Horde/Plugins/Build/HordeServer.Build/Jobs/Schedules/ScheduleService.cs
+++ b/Engine/Source/Programs/Horde/Plugins/Build/HordeServer.Build/Jobs/Schedules/ScheduleService.cs
@@ -5,6 +5,7 @@ using System.Text.RegularExpressions;
 using EpicGames.Core;
 using EpicGames.Horde.Commits;
 using EpicGames.Horde.Jobs;
+using EpicGames.Horde.Jobs.Graphs;
 using EpicGames.Horde.Jobs.Templates;
 using EpicGames.Horde.Streams;
 using EpicGames.Redis;
diff --git a/Engine/Source/Programs/Horde/Plugins/Build/HordeServer.Build/Notifications/INotificationService.cs b/Engine/Source/Programs/Horde/Plugins/Build/HordeServer.Build/Notifications/INotificationService.cs
--- a/Engine/Source/Programs/Horde/Plugins/Build/HordeServer.Build/Notifications/INotificationService.cs
+++ b/Engine/Source/Programs/Horde/Plugins/Build/HordeServer.Build/Notifications/INotificationService.cs
@@ -3,12 +3,12 @@
 using System.Diagnostics.CodeAnalysis;
 using System.Security.Claims;
 using EpicGames.Horde.Jobs;
+using EpicGames.Horde.Jobs.Graphs;
 using HordeServer.Agents;
 using HordeServer.Configuration;
 using HordeServer.Devices;
 using HordeServer.Issues;
 using HordeServer.Jobs;
-using HordeServer.Jobs.Graphs;
 using HordeServer.Streams;
 using HordeServer.Users;
 using MongoDB.Bson;
diff --git a/Engine/Source/Programs/Horde/Plugins/Build/HordeServer.Build/Notifications/INotificationSink.cs b/Engine/Source/Programs/Horde/Plugins/Build/HordeServer.Build/Notifications/INotificationSink.cs
--- a/Engine/Source/Programs/Horde/Plugins/Build/HordeServer.Build/Notifications/INotificationSink.cs
+++ b/Engine/Source/Programs/Horde/Plugins/Build/HordeServer.Build/Notifications/INotificationSink.cs
@@ -1,12 +1,12 @@
 // Copyright Epic Games, Inc. All Rights Reserved.
 
 using EpicGames.Horde.Jobs;
+using EpicGames.Horde.Jobs.Graphs;
 using HordeServer.Agents;
 using HordeServer.Configuration;
 using HordeServer.Devices;
 using HordeServer.Issues;
 using HordeServer.Jobs;
-using HordeServer.Jobs.Graphs;
 using HordeServer.Logs;
 using HordeServer.Streams;
 using HordeServer.Users;
diff --git a/Engine/Source/Programs/Horde/Plugins/Build/HordeServer.Build/Notifications/NotificationService.cs b/Engine/Source/Programs/Horde/Plugins/Build/HordeServer.Build/Notifications/NotificationService.cs
--- a/Engine/Source/Programs/Horde/Plugins/Build/HordeServer.Build/Notifications/NotificationService.cs
+++ b/Engine/Source/Programs/Horde/Plugins/Build/HordeServer.Build/Notifications/NotificationService.cs
@@ -5,6 +5,7 @@ using System.Diagnostics.Metrics;
 using System.Security.Claims;
 using EpicGames.Core;
 using EpicGames.Horde.Jobs;
+using EpicGames.Horde.Jobs.Graphs;
 using EpicGames.Horde.Users;
 using EpicGames.Redis;
 using HordeServer.Agents;
diff --git a/Engine/Source/Programs/Horde/Plugins/Build/HordeServer.Build/Notifications/Sinks/SlackNotificationSink.cs b/Engine/Source/Programs/Horde/Plugins/Build/HordeServer.Build/Notifications/Sinks/SlackNotificationSink.cs
--- a/Engine/Source/Programs/Horde/Plugins/Build/HordeServer.Build/Notifications/Sinks/SlackNotificationSink.cs
+++ b/Engine/Source/Programs/Horde/Plugins/Build/HordeServer.Build/Notifications/Sinks/SlackNotificationSink.cs
@@ -14,6 +14,7 @@ using EpicGames.Horde.Agents;
 using EpicGames.Horde.Commits;
 using EpicGames.Horde.Issues;
 using EpicGames.Horde.Jobs;
+using EpicGames.Horde.Jobs.Graphs;
 using EpicGames.Horde.Jobs.Templates;
 using EpicGames.Horde.Logs;
 using EpicGames.Horde.Streams;
@@ -30,7 +31,6 @@ using HordeServer.Devices;
 using HordeServer.Issues;
 using HordeServer.Issues.External;
 using HordeServer.Jobs;
-using HordeServer.Jobs.Graphs;
 using HordeServer.Logs;
 using HordeServer.Perforce;
 using HordeServer.Server;
@@ -1266,9 +1266,9 @@ namespace HordeServer.Notifications.Sinks
 							Uri? swarmLink = GetSwarmLinkForSpan(span);
 							if (swarmLink != null)
 							{
-								suspectMessage +=  $"<{swarmLink}|View changes in Swarm>";
+								suspectMessage += $"<{swarmLink}|View changes in Swarm>";
 							}
-							
+
 							Uri? robomergeTrackChangesLink = GetRobomergeTrackChangesUrlForSpan(span);
 							if (robomergeTrackChangesLink != null)
 							{
@@ -1373,7 +1373,7 @@ namespace HordeServer.Notifications.Sinks
 			}
 
 			await UpdateReactionsAsync(triageChannel, issue, fixFailedSpan, eventId, threadId, cancellationToken);
-			
+
 			if (issue.ExternalIssueKey != null)
 			{
 				string extIssueEventId = $"issue_{issue.Id}_ext_{issue.ExternalIssueKey}";
@@ -1381,7 +1381,7 @@ namespace HordeServer.Notifications.Sinks
 				await PostSingleMessageToThreadAsync(triageChannel, extIssueEventId, threadId, extIssueMessage, cancellationToken);
 			}
 		}
-		
+
 		private async Task UpdateReactionsAsync(string triageChannel, IIssue issue, IIssueSpan? fixFailedSpan, string eventId, SlackMessageId threadId, CancellationToken cancellationToken)
 		{
 			ReactionFlags reactions = ReactionFlags.None;
@@ -1405,10 +1405,10 @@ namespace HordeServer.Notifications.Sinks
 			{
 				reactions |= ReactionFlags.AssignedButNotAcknowledged;
 			}
-			
+
 			string reactionEventId = $"{eventId}_reactions";
 			string reactionStateDigest = ((int)reactions).ToString();
-			
+
 			MessageStateDocument? reactionState = await GetMessageStateAsync(triageChannel, reactionEventId, cancellationToken);
 			if ((reactionState == null && reactions != ReactionFlags.None) || (reactionState != null && reactionState.Digest != reactionStateDigest))
 			{
@@ -1420,7 +1420,7 @@ namespace HordeServer.Notifications.Sinks
 				{
 					await _slackClient.RemoveReactionAsync(threadId, "eyes", cancellationToken);
 				}
-				
+
 				if ((reactions & ReactionFlags.Quarantined) != 0)
 				{
 					await _slackClient.AddReactionAsync(threadId, "mask", cancellationToken);
@@ -1429,7 +1429,7 @@ namespace HordeServer.Notifications.Sinks
 				{
 					await _slackClient.RemoveReactionAsync(threadId, "mask", cancellationToken);
 				}
-				
+
 				if ((reactions & ReactionFlags.FixFailed) != 0)
 				{
 					await _slackClient.AddReactionAsync(threadId, "x", cancellationToken);
@@ -1438,7 +1438,7 @@ namespace HordeServer.Notifications.Sinks
 				{
 					await _slackClient.RemoveReactionAsync(threadId, "x", cancellationToken);
 				}
-				
+
 				if ((reactions & ReactionFlags.Resolved) != 0)
 				{
 					await _slackClient.AddReactionAsync(threadId, "tick", cancellationToken);
@@ -1447,7 +1447,7 @@ namespace HordeServer.Notifications.Sinks
 				{
 					await _slackClient.RemoveReactionAsync(threadId, "tick", cancellationToken);
 				}
-				
+
 				if ((reactions & ReactionFlags.AssignedButNotAcknowledged) != 0)
 				{
 					await _slackClient.AddReactionAsync(threadId, "mailbox", cancellationToken);
@@ -1456,11 +1456,11 @@ namespace HordeServer.Notifications.Sinks
 				{
 					await _slackClient.RemoveReactionAsync(threadId, "mailbox", cancellationToken);
 				}
-				
+
 				await AddOrUpdateMessageStateAsync(triageChannel, reactionEventId, null, reactionStateDigest, threadId, cancellationToken);
 			}
 		}
-		
+
 		Uri? GetSwarmLinkForSpan(IIssueSpan span)
 		{
 			if (_settings.P4SwarmUrl == null || span.LastSuccess == null)
diff --git a/Engine/Source/Programs/Horde/Plugins/Build/HordeServer.Build/Streams/StreamConfig.cs b/Engine/Source/Programs/Horde/Plugins/Build/HordeServer.Build/Streams/StreamConfig.cs
--- a/Engine/Source/Programs/Horde/Plugins/Build/HordeServer.Build/Streams/StreamConfig.cs
+++ b/Engine/Source/Programs/Horde/Plugins/Build/HordeServer.Build/Streams/StreamConfig.cs
@@ -15,6 +15,7 @@ using EpicGames.Horde.Artifacts;
 using EpicGames.Horde.Commits;
 using EpicGames.Horde.Issues;
 using EpicGames.Horde.Jobs;
+using EpicGames.Horde.Jobs.Graphs;
 using EpicGames.Horde.Jobs.Templates;
 using EpicGames.Horde.Replicators;
 using EpicGames.Horde.Streams;
@@ -27,7 +28,6 @@ using HordeServer.Configuration;
 using HordeServer.Issues;
 using HordeServer.Jobs;
 using HordeServer.Jobs.Bisect;
-using HordeServer.Jobs.Graphs;
 using HordeServer.Jobs.Schedules;
 using HordeServer.Jobs.Templates;
 using HordeServer.Notifications;
@@ -273,7 +273,7 @@ namespace HordeServer.Streams
 		{
 			Id = id;
 			ProjectConfig = projectConfig;
-			
+
 			AclAction[] aclActions = AclConfig.GetActions(
 				[
 					typeof(BisectTaskAclAction),
diff --git a/Engine/Source/Programs/Horde/Plugins/Experimental/HordeServer.Experimental/Notifications/Sinks/ExperimentalSlackNotificationSink.cs b/Engine/Source/Programs/Horde/Plugins/Experimental/HordeServer.Experimental/Notifications/Sinks/ExperimentalSlackNotificationSink.cs
--- a/Engine/Source/Programs/Horde/Plugins/Experimental/HordeServer.Experimental/Notifications/Sinks/ExperimentalSlackNotificationSink.cs
+++ b/Engine/Source/Programs/Horde/Plugins/Experimental/HordeServer.Experimental/Notifications/Sinks/ExperimentalSlackNotificationSink.cs
@@ -1,13 +1,13 @@
 // Copyright Epic Games, Inc. All Rights Reserved.
 
 using EpicGames.Horde.Jobs;
+using EpicGames.Horde.Jobs.Graphs;
 using EpicGames.Slack;
 using HordeServer.Agents;
 using HordeServer.Configuration;
 using HordeServer.Devices;
 using HordeServer.Issues;
 using HordeServer.Jobs;
-using HordeServer.Jobs.Graphs;
 using HordeServer.Logs;
 using HordeServer.Streams;
 using HordeServer.Users;
diff --git a/Engine/Source/Programs/Shared/EpicGames.Horde/Issues/IWorkflow.cs b/Engine/Source/Programs/Shared/EpicGames.Horde/Issues/IWorkflow.cs
--- a/Engine/Source/Programs/Shared/EpicGames.Horde/Issues/IWorkflow.cs
+++ b/Engine/Source/Programs/Shared/EpicGames.Horde/Issues/IWorkflow.cs
@@ -2,7 +2,7 @@
 
 using System;
 using System.Collections.Generic;
-using HordeServer.Jobs.Graphs;
+using EpicGames.Horde.Jobs.Graphs;
 
 namespace EpicGames.Horde.Issues
 {
diff --git a/Engine/Source/Programs/Shared/EpicGames.Horde/Issues/WorkflowMessages.cs b/Engine/Source/Programs/Shared/EpicGames.Horde/Issues/WorkflowMessages.cs
--- a/Engine/Source/Programs/Shared/EpicGames.Horde/Issues/WorkflowMessages.cs
+++ b/Engine/Source/Programs/Shared/EpicGames.Horde/Issues/WorkflowMessages.cs
@@ -3,7 +3,7 @@
 using System;
 using System.Collections.Generic;
 using System.Linq;
-using HordeServer.Jobs.Graphs;
+using EpicGames.Horde.Jobs.Graphs;
 
 #pragma warning disable CA2227 //  Change x to be read-only by removing the property setter
 
diff --git a/Engine/Source/Programs/Shared/EpicGames.Horde/Jobs/Graphs/IGraph.cs b/Engine/Source/Programs/Shared/EpicGames.Horde/Jobs/Graphs/IGraph.cs
--- a/Engine/Source/Programs/Shared/EpicGames.Horde/Jobs/Graphs/IGraph.cs
+++ b/Engine/Source/Programs/Shared/EpicGames.Horde/Jobs/Graphs/IGraph.cs
@@ -7,13 +7,10 @@ using System.Diagnostics.CodeAnalysis;
 using System.Linq;
 using EpicGames.Core;
 using EpicGames.Horde.Artifacts;
-using EpicGames.Horde.Jobs;
-//using HordeCommon;
-//using MongoDB.Bson.Serialization.Attributes;
 
 #pragma warning disable CA2227 // Change x to be read-only by removing the property setter
 
-namespace HordeServer.Jobs.Graphs
+namespace EpicGames.Horde.Jobs.Graphs
 {
 	/// <summary>
 	/// A unique dependency graph instance
diff --git a/Engine/Source/Programs/Shared/EpicGames.Horde/Jobs/Graphs/NodeAnnotations.cs b/Engine/Source/Programs/Shared/EpicGames.Horde/Jobs/Graphs/NodeAnnotations.cs
--- a/Engine/Source/Programs/Shared/EpicGames.Horde/Jobs/Graphs/NodeAnnotations.cs
+++ b/Engine/Source/Programs/Shared/EpicGames.Horde/Jobs/Graphs/NodeAnnotations.cs
@@ -4,7 +4,7 @@ using System;
 using System.Collections.Generic;
 using EpicGames.Horde.Issues;
 
-namespace HordeServer.Jobs.Graphs
+namespace EpicGames.Horde.Jobs.Graphs
 {
 	/// <summary>
 	/// Interface which wraps a generic key/value dictionary to provide specific node annotations
