Horde: Use identifiers to query stream objects rather than StreamConfig objects.

#rnx

diff --git a/Engine/Source/Programs/Horde/Plugins/Build/HordeServer.Build.Tests/Jobs/JobTaskSourceTests.cs b/Engine/Source/Programs/Horde/Plugins/Build/HordeServer.Build.Tests/Jobs/JobTaskSourceTests.cs
--- a/Engine/Source/Programs/Horde/Plugins/Build/HordeServer.Build.Tests/Jobs/JobTaskSourceTests.cs
+++ b/Engine/Source/Programs/Horde/Plugins/Build/HordeServer.Build.Tests/Jobs/JobTaskSourceTests.cs
@@ -100,8 +100,8 @@ namespace HordeServer.Tests.Jobs
 			Fixture fixture = await SetupPoolWithAgentAsync(isPoolAutoScaled: true, shouldCreateAgent: true, isAgentEnabled: true);
 
 			// update template with some step states
-			IStream stream = await StreamCollection.GetAsync(fixture.StreamConfig!);
-			stream = Deref(await StreamCollection.TryUpdateTemplateRefAsync(stream, fixture.TemplateRefId1, new List<UpdateStepStateRequest>() { new UpdateStepStateRequest() { Name = "Paused Step", PausedByUserId = new UserId(BinaryIdUtils.CreateNew()).ToString() } }));
+			IStream stream = (await StreamCollection.GetAsync(fixture.StreamConfig!.Id))!;
+			stream = Deref(await stream.TryUpdateTemplateRefAsync(fixture.TemplateRefId1, new List<UpdateStepStateRequest>() { new UpdateStepStateRequest() { Name = "Paused Step", PausedByUserId = new UserId(BinaryIdUtils.CreateNew()).ToString() } }));
 
 			// create a new graph with the associated nodes
 			List<NewGroup> newGroups = new List<NewGroup>();
diff --git a/Engine/Source/Programs/Horde/Plugins/Build/HordeServer.Build.Tests/Jobs/Schedules/SchedulerTests.cs b/Engine/Source/Programs/Horde/Plugins/Build/HordeServer.Build.Tests/Jobs/Schedules/SchedulerTests.cs
--- a/Engine/Source/Programs/Horde/Plugins/Build/HordeServer.Build.Tests/Jobs/Schedules/SchedulerTests.cs
+++ b/Engine/Source/Programs/Horde/Plugins/Build/HordeServer.Build.Tests/Jobs/Schedules/SchedulerTests.cs
@@ -14,7 +14,6 @@ using HordeServer.Projects;
 using HordeServer.Server;
 using HordeServer.Streams;
 using HordeServer.Users;
-using Microsoft.Extensions.DependencyInjection;
 using Moq;
 
 namespace HordeServer.Tests.Jobs.Schedules
@@ -63,7 +62,7 @@ namespace HordeServer.Tests.Jobs.Schedules
 
 			UpdateConfig(x => x.Plugins.GetBuildConfig().Projects[0].Streams = new List<StreamConfig> { streamConfig });
 
-			return await StreamCollection.GetAsync(streamConfig);
+			return (await StreamCollection.GetAsync(streamConfig.Id))!;
 		}
 
 		public async Task<List<IJob>> FileTestHelperAsync(params string[] files)
@@ -261,7 +260,7 @@ namespace HordeServer.Tests.Jobs.Schedules
 			StreamConfig? streamConfig;
 			GlobalConfig.CurrentValue.Plugins.GetBuildConfig().TryGetStream(StreamId, out streamConfig);
 
-			IStream stream2 = (await StreamCollection.GetAsync(streamConfig!))!;
+			IStream stream2 = (await StreamCollection.GetAsync(streamConfig!.Id))!;
 			ISchedule schedule2 = stream2.Templates.First().Value.Schedule!;
 			Assert.AreEqual(102, schedule2.LastTriggerCommitId!.GetPerforceChange());
 			Assert.AreEqual(Clock.UtcNow, schedule2.LastTriggerTimeUtc);
@@ -448,7 +447,7 @@ namespace HordeServer.Tests.Jobs.Schedules
 			config.Templates = new() { new TemplateRefConfig { Id = newTemplateRefId1 }, new TemplateRefConfig { Id = newTemplateRefId2 } };
 			UpdateConfig(x => x.Plugins.GetBuildConfig().Projects[0].Streams = new List<StreamConfig> { config });
 
-			IStream stream = await StreamCollection.GetAsync(config);
+			IStream stream = (await StreamCollection.GetAsync(config.Id))!;
 
 			// Create the TriggerNext step and mark it as complete
 			IGraph graphA = await GraphCollection.AddAsync(newTemplate1);
@@ -540,7 +539,7 @@ namespace HordeServer.Tests.Jobs.Schedules
 			config.Templates.Add(newTemplate2);
 			UpdateConfig(x => x.Plugins.GetBuildConfig().Projects[0].Streams = new List<StreamConfig> { config });
 
-			IStream stream = await StreamCollection.GetAsync(config);
+			IStream stream = (await StreamCollection.GetAsync(config.Id))!;
 
 			//			stream = (await CreateOrReplaceStreamAsync(StreamId, stream, ProjectId, config))!;
 
@@ -609,7 +608,7 @@ namespace HordeServer.Tests.Jobs.Schedules
 			Assert.AreEqual(100, jobs1[0].CodeCommitId!.GetPerforceChange());
 
 			// Make sure the job is registered
-			IStream? stream1 = await StreamCollection.GetAsync(GlobalConfig.CurrentValue.Plugins.GetBuildConfig().Streams[0]);
+			IStream? stream1 = await StreamCollection.GetAsync(GlobalConfig.CurrentValue.Plugins.GetBuildConfig().Streams[0].Id);
 			ITemplateRef templateRef1 = stream1!.Templates.First().Value;
 			Assert.AreEqual(1, templateRef1.Schedule!.ActiveJobs.Count);
 			Assert.AreEqual(jobs1[0].Id, templateRef1.Schedule!.ActiveJobs[0]);
@@ -618,7 +617,7 @@ namespace HordeServer.Tests.Jobs.Schedules
 			await SetScheduleAsync(schedule);
 
 			// Make sure the job is still registered
-			IStream? stream2 = await StreamCollection.GetAsync(GlobalConfig.CurrentValue.Plugins.GetBuildConfig().Streams[0]);
+			IStream? stream2 = await StreamCollection.GetAsync(GlobalConfig.CurrentValue.Plugins.GetBuildConfig().Streams[0].Id);
 			ITemplateRef templateRef2 = stream2!.Templates.First().Value;
 			Assert.AreEqual(1, templateRef2.Schedule!.ActiveJobs.Count);
 			Assert.AreEqual(jobs1[0].Id, templateRef2.Schedule!.ActiveJobs[0]);
@@ -635,8 +634,7 @@ namespace HordeServer.Tests.Jobs.Schedules
 			schedule.Patterns.Add(new SchedulePatternConfig { MinTime = ScheduleTimeOfDay.Parse("13:00"), MaxTime = ScheduleTimeOfDay.Parse("14:00"), Interval = ScheduleInterval.Parse("15m") });
 			IStream stream = await SetScheduleAsync(schedule);
 
-			IStreamCollection streamCollection = ServiceProvider.GetRequiredService<IStreamCollection>();
-			await streamCollection.TryUpdatePauseStateAsync(stream, newPausedUntil: startTime.AddHours(5), newPauseComment: "testing");
+			await stream.TryUpdatePauseStateAsync(newPausedUntil: startTime.AddHours(5), newPauseComment: "testing");
 
 			// Try trigger a job. No job should be scheduled as the stream is paused
 			await Clock.AdvanceAsync(TimeSpan.FromHours(1.25));
diff --git a/Engine/Source/Programs/Horde/Plugins/Build/HordeServer.Build.Tests/Streams/StreamServiceTest.cs b/Engine/Source/Programs/Horde/Plugins/Build/HordeServer.Build.Tests/Streams/StreamServiceTest.cs
--- a/Engine/Source/Programs/Horde/Plugins/Build/HordeServer.Build.Tests/Streams/StreamServiceTest.cs
+++ b/Engine/Source/Programs/Horde/Plugins/Build/HordeServer.Build.Tests/Streams/StreamServiceTest.cs
@@ -13,21 +13,19 @@ namespace HordeServer.Tests.Streams
 		{
 			Fixture fixture = await CreateFixtureAsync();
 
-			IStream stream = (await StreamCollection.GetAsync(fixture!.StreamConfig!))!;
+			IStream stream = (await StreamCollection.GetAsync(fixture!.StreamConfig!.Id))!;
 			Assert.IsFalse(stream.IsPaused(DateTime.UtcNow));
 			Assert.IsNull(stream.PausedUntil);
 			Assert.IsNull(stream.PauseComment);
 
 			DateTime pausedUntil = DateTime.UtcNow.AddHours(1);
-			await StreamCollection.UpdatePauseStateAsync(stream, newPausedUntil: pausedUntil, newPauseComment: "mycomment");
-			stream = (await StreamCollection.GetAsync(fixture!.StreamConfig!))!;
+			stream = (await stream.TryUpdatePauseStateAsync(newPausedUntil: pausedUntil, newPauseComment: "mycomment"))!;
 			// Comparing by string to avoid comparing exact milliseconds as those are not persisted in MongoDB fields
 			Assert.IsTrue(stream.IsPaused(DateTime.UtcNow));
 			Assert.AreEqual(pausedUntil.ToString(CultureInfo.InvariantCulture), stream.PausedUntil!.Value.ToString(CultureInfo.InvariantCulture));
 			Assert.AreEqual("mycomment", stream.PauseComment);
 
-			await StreamCollection.UpdatePauseStateAsync(stream, newPausedUntil: null, newPauseComment: null);
-			stream = (await StreamCollection.GetAsync(fixture!.StreamConfig!))!;
+			stream = (await stream.TryUpdatePauseStateAsync(newPausedUntil: null, newPauseComment: null))!;
 			Assert.IsFalse(stream.IsPaused(DateTime.UtcNow));
 			Assert.IsNull(stream.PausedUntil);
 			Assert.IsNull(stream.PauseComment);
diff --git a/Engine/Source/Programs/Horde/Plugins/Build/HordeServer.Build/Jobs/JobTaskSource.cs b/Engine/Source/Programs/Horde/Plugins/Build/HordeServer.Build/Jobs/JobTaskSource.cs
--- a/Engine/Source/Programs/Horde/Plugins/Build/HordeServer.Build/Jobs/JobTaskSource.cs
+++ b/Engine/Source/Programs/Horde/Plugins/Build/HordeServer.Build/Jobs/JobTaskSource.cs
@@ -370,7 +370,7 @@ namespace HordeServer.Jobs
 
 			// Query all the current streams
 			BuildConfig buildConfig = _buildConfig.CurrentValue;
-			IReadOnlyList<IStream> streamsList = await _streamCollection.GetAsync(buildConfig.Streams, cancellationToken);
+			IReadOnlyList<IStream> streamsList = await _streamCollection.FindAsync(buildConfig.Streams.Select(x => x.Id).ToList(), cancellationToken);
 			Dictionary<StreamId, IStream> streams = streamsList.ToDictionary(x => x.Id, x => x);
 
 			// Find all the pools which are valid (ie. have at least one online agent)
diff --git a/Engine/Source/Programs/Horde/Plugins/Build/HordeServer.Build/Jobs/Schedules/ScheduleService.cs b/Engine/Source/Programs/Horde/Plugins/Build/HordeServer.Build/Jobs/Schedules/ScheduleService.cs
--- a/Engine/Source/Programs/Horde/Plugins/Build/HordeServer.Build/Jobs/Schedules/ScheduleService.cs
+++ b/Engine/Source/Programs/Horde/Plugins/Build/HordeServer.Build/Jobs/Schedules/ScheduleService.cs
@@ -220,7 +220,12 @@ namespace HordeServer.Jobs.Schedules
 			BuildConfig buildConfig = _buildConfig.CurrentValue;
 			foreach (StreamConfig streamConfig in buildConfig.Streams)
 			{
-				IStream stream = await _streamCollection.GetAsync(streamConfig, cancellationToken);
+				IStream? stream = await _streamCollection.GetAsync(streamConfig.Id, cancellationToken);
+				if (stream == null)
+				{
+					continue;
+				}
+
 				foreach ((TemplateId templateId, ITemplateRef templateRef) in stream.Templates)
 				{
 					if (templateRef.Schedule != null)
@@ -233,7 +238,7 @@ namespace HordeServer.Jobs.Schedules
 								double score = QueueItem.GetScoreFromTime(nextTriggerTimeUtc.Value);
 								queueItems.Add(new SortedSetEntry<QueueItem>(new QueueItem(stream.Id, templateId), score));
 
-								await _streamCollection.UpdateScheduleTriggerAsync(stream, templateId, utcNow, cancellationToken: cancellationToken);
+								await stream.TryUpdateScheduleTriggerAsync(templateId, utcNow, cancellationToken: cancellationToken);
 							}
 						}
 					}
@@ -267,7 +272,7 @@ namespace HordeServer.Jobs.Schedules
 				return false;
 			}
 
-			IStream? stream = await _streamCollection.GetAsync(streamConfig, cancellationToken);
+			IStream? stream = await _streamCollection.GetAsync(streamConfig.Id, cancellationToken);
 			if (stream == null || !stream.Templates.TryGetValue(templateId, out ITemplateRef? templateRef))
 			{
 				return false;
@@ -304,7 +309,7 @@ namespace HordeServer.Jobs.Schedules
 					removeJobIds.Add(activeJobId);
 				}
 			}
-			await _streamCollection.UpdateScheduleTriggerAsync(stream, templateId, removeJobs: removeJobIds, cancellationToken: cancellationToken);
+			await stream.TryUpdateScheduleTriggerAsync(templateId, removeJobs: removeJobIds, cancellationToken: cancellationToken);
 
 			// If the stream is paused, bail out
 			if (stream.IsPaused(utcNow))
@@ -497,7 +502,7 @@ namespace HordeServer.Jobs.Schedules
 
 				IJob newJob = await _jobService.CreateJobAsync(null, stream.Config, templateId, template.Hash, graph, template.Name, commitId, codeCommitId, options, cancellationToken);
 				scheduleLogger.LogInformation("Started new job for {StreamId} template {TemplateId} at CL {Change} (Code CL {CodeChange}): {JobId}", stream.Id, templateId, commitId, codeCommitId, newJob.Id);
-				await _streamCollection.UpdateScheduleTriggerAsync(stream, templateId, utcNow, commitId, new List<JobId> { newJob.Id }, new List<JobId>(), cancellationToken);
+				await stream.TryUpdateScheduleTriggerAsync(templateId, utcNow, commitId, new List<JobId> { newJob.Id }, new List<JobId>(), cancellationToken);
 			}
 		}
 
diff --git a/Engine/Source/Programs/Horde/Plugins/Build/HordeServer.Build/Streams/IStream.cs b/Engine/Source/Programs/Horde/Plugins/Build/HordeServer.Build/Streams/IStream.cs
--- a/Engine/Source/Programs/Horde/Plugins/Build/HordeServer.Build/Streams/IStream.cs
+++ b/Engine/Source/Programs/Horde/Plugins/Build/HordeServer.Build/Streams/IStream.cs
@@ -1,5 +1,7 @@
 // Copyright Epic Games, Inc. All Rights Reserved.
 
+using EpicGames.Horde.Commits;
+using EpicGames.Horde.Jobs;
 using EpicGames.Horde.Jobs.Templates;
 using EpicGames.Horde.Streams;
 using EpicGames.Horde.Users;
@@ -58,6 +60,42 @@ namespace HordeServer.Streams
 		/// Comment/reason for why the stream was paused
 		/// </summary>
 		string? PauseComment { get; }
+
+		/// <summary>
+		/// Get the latest stream state
+		/// </summary>
+		/// <param name="cancellationToken">Cancellation toke for this operation</param>
+		/// <returns>Updated stream, or null if it no longer exists</returns>
+		Task<IStream?> RefreshAsync(CancellationToken cancellationToken = default);
+
+		/// <summary>
+		/// Updates user-facing properties for an existing stream
+		/// </summary>
+		/// <param name="newPausedUntil">The new datetime for pausing builds</param>
+		/// <param name="newPauseComment">The reason for pausing</param>
+		/// <param name="cancellationToken">Cancellation token for the operation</param>
+		/// <returns>The updated stream if successful, null otherwise</returns>
+		Task<IStream?> TryUpdatePauseStateAsync(DateTime? newPausedUntil, string? newPauseComment, CancellationToken cancellationToken = default);
+
+		/// <summary>
+		/// Attempts to update the last trigger time for a schedule
+		/// </summary>
+		/// <param name="templateRefId">The template ref id</param>
+		/// <param name="lastTriggerTimeUtc">New last trigger time for the schedule</param>
+		/// <param name="lastTriggerCommitId">New last trigger commit for the schedule</param>
+		/// <param name="newActiveJobs">New list of active jobs</param>
+		/// <param name="cancellationToken">Cancellation token for the operation</param>
+		/// <returns>The updated stream if successful, null otherwise</returns>
+		Task<IStream?> TryUpdateScheduleTriggerAsync(TemplateId templateRefId, DateTime? lastTriggerTimeUtc, CommitIdWithOrder? lastTriggerCommitId, List<JobId> newActiveJobs, CancellationToken cancellationToken = default);
+
+		/// <summary>
+		/// Attempts to update a stream template ref
+		/// </summary>
+		/// <param name="templateRefId">The template ref to update</param>
+		/// <param name="stepStates">The stream states to update, pass an empty list to clear all step states, otherwise will be a partial update based on included step updates</param>
+		/// <param name="cancellationToken">Cancellation token for the operation</param>
+		/// <returns></returns>
+		Task<IStream?> TryUpdateTemplateRefAsync(TemplateId templateRefId, List<UpdateStepStateRequest>? stepStates = null, CancellationToken cancellationToken = default);
 	}
 
 	/// <summary>
diff --git a/Engine/Source/Programs/Horde/Plugins/Build/HordeServer.Build/Streams/IStreamCollection.cs b/Engine/Source/Programs/Horde/Plugins/Build/HordeServer.Build/Streams/IStreamCollection.cs
--- a/Engine/Source/Programs/Horde/Plugins/Build/HordeServer.Build/Streams/IStreamCollection.cs
+++ b/Engine/Source/Programs/Horde/Plugins/Build/HordeServer.Build/Streams/IStreamCollection.cs
@@ -3,6 +3,7 @@
 using EpicGames.Horde.Commits;
 using EpicGames.Horde.Jobs;
 using EpicGames.Horde.Jobs.Templates;
+using EpicGames.Horde.Streams;
 
 namespace HordeServer.Streams
 {
@@ -14,50 +15,18 @@ namespace HordeServer.Streams
 		/// <summary>
 		/// Gets a stream by ID
 		/// </summary>
-		/// <param name="streamConfig">The stream config object</param>
+		/// <param name="id">The stream identifier</param>
 		/// <param name="cancellationToken">Cancellation token for the operation</param>
 		/// <returns>The stream document</returns>
-		Task<IStream> GetAsync(StreamConfig streamConfig, CancellationToken cancellationToken = default);
+		Task<IStream?> GetAsync(StreamId id, CancellationToken cancellationToken = default);
 
 		/// <summary>
 		/// Gets a stream by ID
 		/// </summary>
-		/// <param name="streamConfigs">The stream config object</param>
+		/// <param name="ids">The stream identifiers</param>
 		/// <param name="cancellationToken">Cancellation token for the operation</param>
 		/// <returns>The stream document</returns>
-		Task<IReadOnlyList<IStream>> GetAsync(IReadOnlyList<StreamConfig> streamConfigs, CancellationToken cancellationToken = default);
-
-		/// <summary>
-		/// Updates user-facing properties for an existing stream
-		/// </summary>
-		/// <param name="stream">The stream to update</param>
-		/// <param name="newPausedUntil">The new datetime for pausing builds</param>
-		/// <param name="newPauseComment">The reason for pausing</param>
-		/// <param name="cancellationToken">Cancellation token for the operation</param>
-		/// <returns>The updated stream if successful, null otherwise</returns>
-		Task<IStream?> TryUpdatePauseStateAsync(IStream stream, DateTime? newPausedUntil, string? newPauseComment, CancellationToken cancellationToken = default);
-
-		/// <summary>
-		/// Attempts to update the last trigger time for a schedule
-		/// </summary>
-		/// <param name="stream">The stream to update</param>
-		/// <param name="templateRefId">The template ref id</param>
-		/// <param name="lastTriggerTimeUtc">New last trigger time for the schedule</param>
-		/// <param name="lastTriggerCommitId">New last trigger commit for the schedule</param>
-		/// <param name="newActiveJobs">New list of active jobs</param>
-		/// <param name="cancellationToken">Cancellation token for the operation</param>
-		/// <returns>The updated stream if successful, null otherwise</returns>
-		Task<IStream?> TryUpdateScheduleTriggerAsync(IStream stream, TemplateId templateRefId, DateTime? lastTriggerTimeUtc, CommitIdWithOrder? lastTriggerCommitId, List<JobId> newActiveJobs, CancellationToken cancellationToken = default);
-
-		/// <summary>
-		/// Attempts to update a stream template ref
-		/// </summary>
-		/// <param name="streamInterface">The stream containing the template ref</param>
-		/// <param name="templateRefId">The template ref to update</param>
-		/// <param name="stepStates">The stream states to update, pass an empty list to clear all step states, otherwise will be a partial update based on included step updates</param>
-		/// <param name="cancellationToken">Cancellation token for the operation</param>
-		/// <returns></returns>
-		Task<IStream?> TryUpdateTemplateRefAsync(IStream streamInterface, TemplateId templateRefId, List<UpdateStepStateRequest>? stepStates = null, CancellationToken cancellationToken = default);
+		Task<IReadOnlyList<IStream>> FindAsync(IReadOnlyList<StreamId>? ids = null, CancellationToken cancellationToken = default);
 	}
 
 	/// <summary>
@@ -68,17 +37,16 @@ namespace HordeServer.Streams
 		/// <summary>
 		/// Updates an existing stream
 		/// </summary>
-		/// <param name="streamCollection">The stream collection</param>
 		/// <param name="stream">The stream to update</param>
 		/// <param name="newPausedUntil">The new datetime for pausing builds</param>
 		/// <param name="newPauseComment">The reason for pausing</param>
 		/// <param name="cancellationToken">Cancellation token for the operation</param>
 		/// <returns>Async task object</returns>
-		public static async Task<IStream?> UpdatePauseStateAsync(this IStreamCollection streamCollection, IStream? stream, DateTime? newPausedUntil = null, string? newPauseComment = null, CancellationToken cancellationToken = default)
+		public static async Task<IStream?> TryUpdatePauseStateAsync(this IStream? stream, DateTime? newPausedUntil = null, string? newPauseComment = null, CancellationToken cancellationToken = default)
 		{
-			for (; stream != null; stream = await streamCollection.GetAsync(stream.Config, cancellationToken))
+			for (; stream != null; stream = await stream.RefreshAsync(cancellationToken))
 			{
-				IStream? newStream = await streamCollection.TryUpdatePauseStateAsync(stream, newPausedUntil, newPauseComment, cancellationToken);
+				IStream? newStream = await stream.TryUpdatePauseStateAsync(newPausedUntil, newPauseComment, cancellationToken);
 				if (newStream != null)
 				{
 					return newStream;
@@ -90,7 +58,6 @@ namespace HordeServer.Streams
 		/// <summary>
 		/// Attempts to update the last trigger time for a schedule
 		/// </summary>
-		/// <param name="streamCollection">The stream collection</param>
 		/// <param name="stream">The stream to update</param>
 		/// <param name="templateRefId">The template ref id</param>
 		/// <param name="lastTriggerTimeUtc"></param>
@@ -99,7 +66,7 @@ namespace HordeServer.Streams
 		/// <param name="removeJobs">Jobs to remove</param>
 		/// <param name="cancellationToken">Cancellation token for the operation</param>
 		/// <returns>True if the stream was updated</returns>
-		public static async Task<IStream?> UpdateScheduleTriggerAsync(this IStreamCollection streamCollection, IStream stream, TemplateId templateRefId, DateTime? lastTriggerTimeUtc = null, CommitIdWithOrder? lastTriggerCommitId = null, List<JobId>? addJobs = null, List<JobId>? removeJobs = null, CancellationToken cancellationToken = default)
+		public static async Task<IStream?> TryUpdateScheduleTriggerAsync(this IStream stream, TemplateId templateRefId, DateTime? lastTriggerTimeUtc = null, CommitIdWithOrder? lastTriggerCommitId = null, List<JobId>? addJobs = null, List<JobId>? removeJobs = null, CancellationToken cancellationToken = default)
 		{
 			IStream? newStream = stream;
 			while (newStream != null)
@@ -124,14 +91,14 @@ namespace HordeServer.Streams
 					newActiveJobs = newActiveJobs.Union(addJobs);
 				}
 
-				newStream = await streamCollection.TryUpdateScheduleTriggerAsync(newStream, templateRefId, lastTriggerTimeUtc, lastTriggerCommitId, newActiveJobs.ToList(), cancellationToken);
+				newStream = await newStream.TryUpdateScheduleTriggerAsync(templateRefId, lastTriggerTimeUtc, lastTriggerCommitId, newActiveJobs.ToList(), cancellationToken);
 
 				if (newStream != null)
 				{
 					return newStream;
 				}
 
-				newStream = await streamCollection.GetAsync(stream.Config, cancellationToken);
+				newStream = await stream.RefreshAsync(cancellationToken);
 			}
 			return null;
 		}
diff --git a/Engine/Source/Programs/Horde/Plugins/Build/HordeServer.Build/Streams/StreamCollection.cs b/Engine/Source/Programs/Horde/Plugins/Build/HordeServer.Build/Streams/StreamCollection.cs
--- a/Engine/Source/Programs/Horde/Plugins/Build/HordeServer.Build/Streams/StreamCollection.cs
+++ b/Engine/Source/Programs/Horde/Plugins/Build/HordeServer.Build/Streams/StreamCollection.cs
@@ -10,6 +10,7 @@ using HordeServer.Acls;
 using HordeServer.Jobs.Schedules;
 using HordeServer.Server;
 using HordeServer.Utilities;
+using Microsoft.Extensions.Options;
 using MongoDB.Bson.Serialization.Attributes;
 using MongoDB.Driver;
 
@@ -22,6 +23,7 @@ namespace HordeServer.Streams
 	{
 		class Stream : IStream
 		{
+			readonly StreamCollection _collection;
 			readonly StreamConfig _config;
 			readonly StreamDoc _document;
 			readonly Dictionary<TemplateId, ITemplateRef> _templateRefs = new Dictionary<TemplateId, ITemplateRef>();
@@ -34,17 +36,39 @@ namespace HordeServer.Streams
 			public DateTime? PausedUntil => _document.PausedUntil;
 			public string? PauseComment => _document.PauseComment;
 
-			public Stream(StreamConfig config, StreamDoc document)
+			public Stream(StreamCollection collection, StreamConfig config, StreamDoc document)
 			{
+				_collection = collection;
 				_config = config;
 				_document = document;
 
-				foreach(TemplateRefConfig templateRefConfig in config.Templates)
+				foreach (TemplateRefConfig templateRefConfig in config.Templates)
 				{
 					TemplateRefDoc templateRefDoc = _document.Templates[templateRefConfig.Id];
 					_templateRefs[templateRefConfig.Id] = new TemplateRef(templateRefConfig, templateRefDoc);
 				}
 			}
+
+			public Task<IStream?> RefreshAsync(CancellationToken cancellationToken = default)
+				=> _collection.GetAsync(Id, cancellationToken);
+
+			public async Task<IStream?> TryUpdatePauseStateAsync(DateTime? newPausedUntil, string? newPauseComment, CancellationToken cancellationToken = default)
+			{
+				StreamDoc? newDocument = await _collection.TryUpdatePauseStateAsync(_document, newPausedUntil, newPauseComment, cancellationToken);
+				return (newDocument == null) ? null : new Stream(_collection, _config, newDocument);
+			}
+
+			public async Task<IStream?> TryUpdateScheduleTriggerAsync(TemplateId templateRefId, DateTime? lastTriggerTimeUtc, CommitIdWithOrder? lastTriggerCommitId, List<JobId> newActiveJobs, CancellationToken cancellationToken = default)
+			{
+				StreamDoc? newDocument = await _collection.TryUpdateScheduleTriggerAsync(_document, templateRefId, lastTriggerTimeUtc, lastTriggerCommitId, newActiveJobs, cancellationToken);
+				return (newDocument == null) ? null : new Stream(_collection, _config, newDocument);
+			}
+
+			public async Task<IStream?> TryUpdateTemplateRefAsync(TemplateId templateRefId, List<UpdateStepStateRequest>? stepStates = null, CancellationToken cancellationToken = default)
+			{
+				StreamDoc? newDocument = await _collection.TryUpdateTemplateRefAsync(_document, templateRefId, stepStates, cancellationToken);
+				return (newDocument == null) ? null : new Stream(_collection, _config, newDocument);
+			}
 		}
 
 		class TemplateRef : ITemplateRef
@@ -58,7 +82,7 @@ namespace HordeServer.Streams
 				_config = config;
 				_document = document;
 
-				if(_config.Schedule != null)
+				if (_config.Schedule != null)
 				{
 					_schedule = new Schedule(_config.Schedule, _document.Schedule ?? new ScheduleDoc());
 				}
@@ -240,24 +264,30 @@ namespace HordeServer.Streams
 		}
 
 		readonly IMongoCollection<StreamDoc> _streams;
+		readonly IOptionsMonitor<BuildConfig> _buildConfig;
 		readonly IClock _clock;
 
 		/// <summary>
 		/// Constructor
 		/// </summary>
-		/// <param name="mongoService">The database service instance</param>
-		/// <param name="clock"></param>
-		public StreamCollection(IMongoService mongoService, IClock clock)
+		public StreamCollection(IMongoService mongoService, IOptionsMonitor<BuildConfig> buildConfig, IClock clock)
 		{
 			_streams = mongoService.GetCollection<StreamDoc>("Streams");
+			_buildConfig = buildConfig;
 			_clock = clock;
 		}
 
 		/// <inheritdoc/>
-		public async Task<IStream> GetAsync(StreamConfig streamConfig, CancellationToken cancellationToken)
+		public async Task<IStream?> GetAsync(StreamId id, CancellationToken cancellationToken)
 		{
+			StreamConfig? streamConfig;
+			if (!_buildConfig.CurrentValue.TryGetStream(id, out streamConfig))
+			{
+				return null;
+			}
+
 			StreamDoc document = await GetInternalAsync(streamConfig, cancellationToken);
-			return new Stream(streamConfig, document);
+			return new Stream(this, streamConfig, document);
 		}
 
 		async Task<StreamDoc> GetInternalAsync(StreamConfig streamConfig, CancellationToken cancellationToken)
@@ -277,14 +307,23 @@ namespace HordeServer.Streams
 		}
 
 		/// <inheritdoc/>
-		public async Task<IReadOnlyList<IStream>> GetAsync(IReadOnlyList<StreamConfig> streamConfigs, CancellationToken cancellationToken)
+		public async Task<IReadOnlyList<IStream>> FindAsync(IReadOnlyList<StreamId>? ids = null, CancellationToken cancellationToken = default)
 		{
-			List<Stream> streams = new List<Stream>();
+			List<StreamConfig> configs = new List<StreamConfig>();
+			foreach (StreamConfig config in _buildConfig.CurrentValue.Streams)
+			{
+				if (ids == null || ids.Contains(config.Id))
+				{
+					configs.Add(config);
+				}
+			}
+
+			List<StreamDoc> documents = await GetInternalAsync(configs, cancellationToken);
 
-			List<StreamDoc> documents = await GetInternalAsync(streamConfigs, cancellationToken);
-			for(int idx = 0; idx < streamConfigs.Count; idx++)
+			Stream[] streams = new Stream[configs.Count];
+			for (int idx = 0; idx < configs.Count; idx++)
 			{
-				streams.Add(new Stream(streamConfigs[idx], documents[idx]));
+				streams[idx] = new Stream(this, configs[idx], documents[idx]);
 			}
 
 			return streams;
@@ -328,11 +367,8 @@ namespace HordeServer.Streams
 		}
 
 		/// <inheritdoc/>
-		public async Task<IStream?> TryUpdateTemplateRefAsync(IStream streamInterface, TemplateId templateId, List<UpdateStepStateRequest>? stepStates = null, CancellationToken cancellationToken = default)
+		async Task<StreamDoc?> TryUpdateTemplateRefAsync(StreamDoc document, TemplateId templateId, List<UpdateStepStateRequest>? stepStates = null, CancellationToken cancellationToken = default)
 		{
-			Stream stream = (Stream)streamInterface;
-			StreamDoc document = stream.Document;
-
 			TemplateRefDoc? templateRef;
 			if (!document.Templates.TryGetValue(templateId, out templateRef))
 			{
@@ -409,20 +445,16 @@ namespace HordeServer.Streams
 
 			if (updates.Count == 0)
 			{
-				return streamInterface;
+				return document;
 			}
 
-			StreamDoc? newStreamDoc = await TryUpdateStreamAsync(document, updateBuilder.Combine(updates), cancellationToken);
-			return (newStreamDoc == null)? null : new Stream(stream.Config, document);
+			return await TryUpdateStreamAsync(document, updateBuilder.Combine(updates), cancellationToken);
 
 		}
 
 		/// <inheritdoc/>
-		public async Task<IStream?> TryUpdatePauseStateAsync(IStream streamInterface, DateTime? newPausedUntil, string? newPauseComment, CancellationToken cancellationToken)
+		async Task<StreamDoc?> TryUpdatePauseStateAsync(StreamDoc document, DateTime? newPausedUntil, string? newPauseComment, CancellationToken cancellationToken)
 		{
-			Stream stream = (Stream)streamInterface;
-			StreamDoc document = stream.Document;
-
 			UpdateDefinitionBuilder<StreamDoc> updateBuilder = Builders<StreamDoc>.Update;
 
 			List<UpdateDefinition<StreamDoc>> updates = new List<UpdateDefinition<StreamDoc>>();
@@ -431,16 +463,12 @@ namespace HordeServer.Streams
 			updates.Add(updateBuilder.Set(x => x.PausedUntil, newPausedUntil));
 			updates.Add(updateBuilder.Set(x => x.PauseComment, newPauseComment));
 
-			StreamDoc? newDocument = await TryUpdateStreamAsync(document, updateBuilder.Combine(updates), cancellationToken);
-			return (newDocument == null)? null : new Stream(stream.Config, newDocument);
+			return await TryUpdateStreamAsync(document, updateBuilder.Combine(updates), cancellationToken);
 		}
 
 		/// <inheritdoc/>
-		public async Task<IStream?> TryUpdateScheduleTriggerAsync(IStream streamInterface, TemplateId templateId, DateTime? lastTriggerTimeUtc, CommitIdWithOrder? lastTriggerCommitId, List<JobId> newActiveJobs, CancellationToken cancellationToken)
+		async Task<StreamDoc?> TryUpdateScheduleTriggerAsync(StreamDoc streamDoc, TemplateId templateId, DateTime? lastTriggerTimeUtc, CommitIdWithOrder? lastTriggerCommitId, List<JobId> newActiveJobs, CancellationToken cancellationToken)
 		{
-			Stream stream = (Stream)streamInterface;
-
-			StreamDoc streamDoc = stream.Document;
 			TemplateRefDoc templateRefDoc = streamDoc.Templates[templateId];
 			ScheduleDoc scheduleDoc = templateRefDoc.Schedule!;
 
@@ -471,13 +499,12 @@ namespace HordeServer.Streams
 				scheduleDoc.ActiveJobs = newActiveJobs;
 			}
 
-			if(updates.Count == 0)
+			if (updates.Count == 0)
 			{
-				return streamInterface;
+				return streamDoc;
 			}
 
-			StreamDoc? newStreamDoc = await TryUpdateStreamAsync(streamDoc, Builders<StreamDoc>.Update.Combine(updates), cancellationToken);
-			return (newStreamDoc == null)? null : new Stream(stream.Config, newStreamDoc);
+			return await TryUpdateStreamAsync(streamDoc, Builders<StreamDoc>.Update.Combine(updates), cancellationToken);
 		}
 
 		/// <summary>
diff --git a/Engine/Source/Programs/Horde/Plugins/Build/HordeServer.Build/Streams/StreamsController.cs b/Engine/Source/Programs/Horde/Plugins/Build/HordeServer.Build/Streams/StreamsController.cs
--- a/Engine/Source/Programs/Horde/Plugins/Build/HordeServer.Build/Streams/StreamsController.cs
+++ b/Engine/Source/Programs/Horde/Plugins/Build/HordeServer.Build/Streams/StreamsController.cs
@@ -71,7 +71,7 @@ namespace HordeServer.Streams
 		{
 			ProjectId[] projectIdValues = Array.ConvertAll(projectIds, x => new ProjectId(x));
 
-			List<StreamConfig> streamConfigs = new List<StreamConfig>();
+			List<StreamId> streamIds = new List<StreamId>();
 			foreach (ProjectConfig projectConfig in _buildConfig.Value.Projects)
 			{
 				if (projectIdValues.Length == 0 || projectIdValues.Contains(projectConfig.Id))
@@ -80,13 +80,13 @@ namespace HordeServer.Streams
 					{
 						if (streamConfig.Authorize(StreamAclAction.ViewStream, User))
 						{
-							streamConfigs.Add(streamConfig);
+							streamIds.Add(streamConfig.Id);
 						}
 					}
 				}
 			}
 
-			IReadOnlyList<IStream> streams = await _streamCollection.GetAsync(streamConfigs, cancellationToken);
+			IReadOnlyList<IStream> streams = await _streamCollection.FindAsync(streamIds, cancellationToken);
 
 			List<GetStreamResponse> responses = new List<GetStreamResponse>();
 			foreach (IStream stream in streams)
@@ -120,7 +120,12 @@ namespace HordeServer.Streams
 				return Forbid(StreamAclAction.ViewStream, streamId);
 			}
 
-			IStream stream = await _streamCollection.GetAsync(streamConfig, cancellationToken);
+			IStream? stream = await _streamCollection.GetAsync(streamConfig.Id, cancellationToken);
+			if(stream == null)
+			{
+				return NotFound(streamId);
+			}
+
 			return PropertyFilter.Apply(await CreateGetStreamResponseAsync(stream, cancellationToken), filter);
 		}
 
@@ -406,19 +411,17 @@ namespace HordeServer.Streams
 				return Forbid(AdminAclAction.AdminWrite);
 			}
 
-			StreamConfig? streamConfig;
-			if (!_buildConfig.Value.TryGetStream(streamId, out streamConfig))
+			IStream? stream = await _streamCollection.GetAsync(streamId);
+			if(stream == null)
 			{
 				return NotFound(streamId);
 			}
-
-			IStream stream = await _streamCollection.GetAsync(streamConfig);
 			if (!stream.Templates.ContainsKey(templateRefId))
 			{
 				return NotFound(streamId, templateRefId);
 			}
 
-			await _streamCollection.TryUpdateTemplateRefAsync(stream, templateRefId, update.StepStates);
+			await stream.TryUpdateTemplateRefAsync(templateRefId, update.StepStates);
 			return Ok();
 		}
 	}
