Horde: Updating schema documentation

#rnx

diff --git a/Engine/Source/Programs/Horde/Docs/Config/Schema/Streams.md b/Engine/Source/Programs/Horde/Docs/Config/Schema/Streams.md
--- a/Engine/Source/Programs/Horde/Docs/Config/Schema/Streams.md
+++ b/Engine/Source/Programs/Horde/Docs/Config/Schema/Streams.md
@@ -241,7 +241,7 @@ Name | Description
 `submitDescription` | `string`<br>Description for new changelists
 `defaultChange` | [ChangeQueryConfig](#changequeryconfig)`[]`<br>Default change to build at. Each object has a condition parameter which can evaluated by the server to determine which change to use.
 `arguments` | `string[]`<br>Fixed arguments for the new job
-`parameters` | [TextParameterData](#textparameterdata)/[ListParameterData](#listparameterdata)/[BoolParameterData](#boolparameterdata)`[]`<br>Parameters for this template
+`parameters` | [TemplateTextParameterConfig](#templatetextparameterconfig)/[TemplateListParameterConfig](#templatelistparameterconfig)/[TemplateBoolParameterConfig](#templateboolparameterconfig)`[]`<br>Parameters for this template
 `jobOptions` | [JobOptions](#joboptions)<br>Default settings for jobs
 
 ## ScheduleConfig
@@ -250,35 +250,35 @@ Parameters to create a new schedule
 
 Name | Description
 ---- | -----------
-`claims` | [AclClaimConfig](#aclclaimconfig)`[]`<br>Roles to impersonate for this schedule
 `enabled` | `boolean`<br>Whether the schedule should be enabled
 `maxActive` | `integer`<br>Maximum number of builds that can be active at once
 `maxChanges` | `integer`<br>Maximum number of changes the schedule can fall behind head revision. If greater than zero, builds will be triggered for every submitted changelist until the backlog is this size.
 `requireSubmittedChange` | `boolean`<br>Whether the build requires a change to be submitted
 `gate` | [ScheduleGateConfig](#schedulegateconfig)<br>Gate allowing the schedule to trigger
 `commits` | `string[]`<br>Commit tags for this schedule
+`claims` | [AclClaimConfig](#aclclaimconfig)`[]`<br>Roles to impersonate for this schedule
 `filter` | [ChangeContentFlags](#changecontentflags-enum)`[]`<br>The types of changes to run for
 `files` | `string[]`<br>Files that should cause the job to trigger
 `templateParameters` | `string` `->` `string`<br>Parameters for the template
 `patterns` | [SchedulePatternConfig](#schedulepatternconfig)`[]`<br>New patterns for the schedule
 
-## AclClaimConfig
+## ScheduleGateConfig
 
-New claim to create
+Gate allowing a schedule to trigger.
 
 Name | Description
 ---- | -----------
-`type` | `string`<br>The claim type
-`value` | `string`<br>The claim value
+`templateId` | `string`<br>The template containing the dependency
+`target` | `string`<br>Target to wait for
 
-## ScheduleGateConfig
+## AclClaimConfig
 
-Gate allowing a schedule to trigger.
+New claim to create
 
 Name | Description
 ---- | -----------
-`templateId` | `string`<br>The template containing the dependency
-`target` | `string`<br>Target to wait for
+`type` | `string`<br>The claim type
+`value` | `string`<br>The claim value
 
 ## ChangeContentFlags (Enum)
 
@@ -368,7 +368,7 @@ Name | Description
 `High` | High priority
 `Highest` | Highest priority
 
-## TextParameterData
+## TemplateTextParameterConfig
 
 Free-form text entry parameter
 
@@ -385,7 +385,7 @@ Name | Description
 `validationError` | `string`<br>Message displayed if validation fails, informing user of valid values.
 `toolTip` | `string`<br>Tool-tip text to display
 
-## ListParameterData
+## TemplateListParameterConfig
 
 Allows the user to select a value from a constrained list of choices
 
@@ -393,11 +393,11 @@ Name | Description
 ---- | -----------
 `type` | List<br>Type discriminator
 `label` | `string`<br>Label to display next to this parameter. Defaults to the parameter name.
-`style` | [ListParameterStyle](#listparameterstyle-enum)<br>The type of list parameter
-`items` | [ListParameterItemData](#listparameteritemdata)`[]`<br>List of values to display in the list
+`style` | [TemplateListParameterStyle](#templatelistparameterstyle-enum)<br>The type of list parameter
+`items` | [TemplateListParameterItemConfig](#templatelistparameteritemconfig)`[]`<br>List of values to display in the list
 `toolTip` | `string`<br>Tool tip text to display
 
-## ListParameterStyle (Enum)
+## TemplateListParameterStyle (Enum)
 
 Style of list parameter
 
@@ -407,7 +407,7 @@ Name | Description
 `MultiList` | Drop-down list with checkboxes
 `TagPicker` | Tag picker from list of options
 
-## ListParameterItemData
+## TemplateListParameterItemConfig
 
 Possible option for a list parameter
 
@@ -423,7 +423,7 @@ Name | Description
 `default` | `boolean`<br>Whether this item is selected by default
 `scheduleOverride` | `boolean`<br>Overridden value for this property in schedule builds
 
-## BoolParameterData
+## TemplateBoolParameterConfig
 
 Allows the user to toggle an option on or off
 
@@ -479,10 +479,10 @@ Name | Description
 `skipWhenEmpty` | `boolean`<br>Skips sending reports when there are no active issues.
 `showMergeWarnings` | `boolean`<br>Whether to show warnings about merging changes into the origin stream.
 `annotations` | `string` `->` `string`<br>Additional node annotations implicit in this workflow
-`externalIssues` | [ExternalIssueConfig](#externalissueconfig)<br>External issue tracking configuration for this workflow
+`externalIssues` | [WorkflowExternalIssuesConfig](#workflowexternalissuesconfig)<br>External issue tracking configuration for this workflow
 `issueHandlers` | `string[]`<br>Additional issue handlers enabled for this workflow
 
-## ExternalIssueConfig
+## WorkflowExternalIssuesConfig
 
 External issue tracking configuration for a workflow
 
