Horde: Add parameter for controlling inactivity timeout for compute executions

For UBA-based compute tasks this value needs to be increased as CPU or I/O can at times get starved, leading to a timeout.

diff --git a/Engine/Source/Programs/Horde/HordeAgent/Leases/Handlers/ComputeHandler.cs b/Engine/Source/Programs/Horde/HordeAgent/Leases/Handlers/ComputeHandler.cs
--- a/Engine/Source/Programs/Horde/HordeAgent/Leases/Handlers/ComputeHandler.cs
+++ b/Engine/Source/Programs/Horde/HordeAgent/Leases/Handlers/ComputeHandler.cs
@@ -51,12 +51,12 @@ namespace HordeAgent.Leases.Handlers
 			TcpClient? tcpClient = null;
 			try
 			{
-				const int TimeoutSeconds = 30;
+				TimeSpan timeout = TimeSpan.FromMilliseconds(computeTask.InactivityTimeoutMs <= 0 ? 30000 : computeTask.InactivityTimeoutMs);
 
-				tcpClient = await _listenerService.WaitForClientAsync(new ByteString(computeTask.Nonce.Memory), TimeSpan.FromSeconds(TimeoutSeconds), cancellationToken);
+				tcpClient = await _listenerService.WaitForClientAsync(new ByteString(computeTask.Nonce.Memory), timeout, cancellationToken);
 				if (tcpClient == null)
 				{
-					logger.LogInformation("Timed out waiting for connection after {Time}s", TimeoutSeconds);
+					logger.LogInformation("Timed out waiting for connection after {Time} ms", (int)timeout.TotalMilliseconds);
 					return LeaseResult.Success;
 				}
 
@@ -65,7 +65,7 @@ namespace HordeAgent.Leases.Handlers
 				using (CancellationTokenSource cts = CancellationTokenSource.CreateLinkedTokenSource(cancellationToken))
 				{
 					await using ComputeTransport innerTransport = await CreateTransportAsync(computeTask, tcpClient.Client, cts.Token);
-					await using IdleTimeoutTransport idleTimeoutTransport = new(innerTransport);
+					await using IdleTimeoutTransport idleTimeoutTransport = new(innerTransport, timeout);
 
 					await using BackgroundTask timeoutTask = BackgroundTask.StartNew(ctx => idleTimeoutTransport.StartWatchdogTimerAsync(cts, logger, ctx));
 					try
diff --git a/Engine/Source/Programs/Horde/Plugins/Compute/HordeServer.Compute/Compute/ComputeController.cs b/Engine/Source/Programs/Horde/Plugins/Compute/HordeServer.Compute/Compute/ComputeController.cs
--- a/Engine/Source/Programs/Horde/Plugins/Compute/HordeServer.Compute/Compute/ComputeController.cs
+++ b/Engine/Source/Programs/Horde/Plugins/Compute/HordeServer.Compute/Compute/ComputeController.cs
@@ -119,6 +119,7 @@ namespace HordeServer.Compute
 				RequesterPublicIp = request.Connection?.ClientPublicIp,
 				UsePublicIp = request.Connection?.PreferPublicIp,
 				Encryption = ComputeService.ConvertEncryptionToProto(request.Connection?.Encryption),
+				InactivityTimeoutMs = request.Connection?.InactivityTimeoutMs,
 				UserId = User.GetUserId(),
 			};
 
diff --git a/Engine/Source/Programs/Horde/Plugins/Compute/HordeServer.Compute/Compute/ComputeService.cs b/Engine/Source/Programs/Horde/Plugins/Compute/HordeServer.Compute/Compute/ComputeService.cs
--- a/Engine/Source/Programs/Horde/Plugins/Compute/HordeServer.Compute/Compute/ComputeService.cs
+++ b/Engine/Source/Programs/Horde/Plugins/Compute/HordeServer.Compute/Compute/ComputeService.cs
@@ -129,6 +129,9 @@ namespace HordeServer.Compute
 
 		/// <inheritdoc cref="ConnectionMetadataRequest.Encryption" />
 		public ComputeEncryption Encryption { get; init; }
+		
+		/// <inheritdoc cref="ConnectionMetadataRequest.InactivityTimeoutMs" />
+		public int? InactivityTimeoutMs { get; init; }
 
 		/// <summary>
 		/// Constructor
@@ -370,6 +373,7 @@ namespace HordeServer.Compute
 			span.SetAttribute("connectionMode", arp.ConnectionMode?.ToString());
 			span.SetAttribute("usePublicIp", arp.UsePublicIp.ToString());
 			span.SetAttribute("encryption", arp.Encryption.ToString());
+			span.SetAttribute("inactivityTimeout", arp.InactivityTimeoutMs);
 			span.SetAttribute("parentLeaseId", arp.ParentLeaseId?.ToString());
 			span.SetAttribute("req.pool", arp.Requirements.Pool);
 			span.SetAttribute("req.condition", arp.Requirements.Condition?.ToString());
@@ -437,7 +441,7 @@ namespace HordeServer.Compute
 
 						using TelemetrySpan createTaskSpan = _tracer.StartActiveSpan("CreateComputeTask");
 
-						ComputeTask computeTask = CreateComputeTask(assignedResources, log?.Id, arp.Encryption, certificate, arp.ParentLeaseId, protocol);
+						ComputeTask computeTask = CreateComputeTask(assignedResources, log?.Id, arp.Encryption, certificate, arp.ParentLeaseId, protocol, arp.InactivityTimeoutMs);
 
 						CreateLeaseOptions createLeaseOptions = new CreateLeaseOptions(leaseId, arp.ParentLeaseId, "Compute task", null, null, log?.Id, assignedResources, arp.Requirements.Exclusive, computeTask);
 
@@ -850,7 +854,8 @@ namespace HordeServer.Compute
 			return relayIps[0];
 		}
 
-		static ComputeTask CreateComputeTask(Dictionary<string, int> assignedResources, LogId? logId, ComputeEncryption encryption, byte[] certificateData, LeaseId? parentLeaseId, ComputeProtocol protocol)
+		static ComputeTask CreateComputeTask(Dictionary<string, int> assignedResources, LogId? logId, ComputeEncryption encryption,
+			byte[] certificateData, LeaseId? parentLeaseId, ComputeProtocol protocol, int? inactivityTimeoutMs)
 		{
 			ComputeTask computeTask = new ComputeTask();
 			computeTask.Encryption = encryption;
@@ -861,6 +866,7 @@ namespace HordeServer.Compute
 			computeTask.LogId = logId?.ToString();
 			computeTask.ParentLeaseId = parentLeaseId?.ToString() ?? String.Empty;
 			computeTask.Protocol = (int)protocol;
+			computeTask.InactivityTimeoutMs = inactivityTimeoutMs ?? 30000;
 			return computeTask;
 		}
 
diff --git a/Engine/Source/Programs/Shared/EpicGames.Horde/Compute/ComputeMessages.cs b/Engine/Source/Programs/Shared/EpicGames.Horde/Compute/ComputeMessages.cs
--- a/Engine/Source/Programs/Shared/EpicGames.Horde/Compute/ComputeMessages.cs
+++ b/Engine/Source/Programs/Shared/EpicGames.Horde/Compute/ComputeMessages.cs
@@ -105,6 +105,12 @@ namespace EpicGames.Horde.Compute
 		/// Encryption mode to request. Server can still override.
 		/// </summary>
 		public Encryption? Encryption { get; set; }
+		
+		/// <summary>
+		/// Maximum duration (in milliseconds) that communication can be inactive between
+		/// the compute client and task-running agent before the connection is terminated.
+		/// </summary>
+		public int? InactivityTimeoutMs { get; set; }
 	}
 
 	/// <summary>
diff --git a/Engine/Source/Programs/Shared/EpicGames.Horde/Compute/Transports/IdleTimeoutTransport.cs b/Engine/Source/Programs/Shared/EpicGames.Horde/Compute/Transports/IdleTimeoutTransport.cs
--- a/Engine/Source/Programs/Shared/EpicGames.Horde/Compute/Transports/IdleTimeoutTransport.cs
+++ b/Engine/Source/Programs/Shared/EpicGames.Horde/Compute/Transports/IdleTimeoutTransport.cs
@@ -82,7 +82,7 @@ public class IdleTimeoutTransport : ComputeTransport
 			TimeSpan remainingTime = NoDataTimeout - TimeSinceActivity;
 			if (remainingTime < TimeSpan.Zero)
 			{
-				logger.LogWarning("Terminating compute transport due to timeout (last tick at {Time})", DateTime.UtcNow - TimeSinceActivity);
+				logger.LogWarning("Terminating compute transport due to timeout ({Timeout} ms with last tick at {Time})", (int)NoDataTimeout.TotalMilliseconds, DateTime.UtcNow - TimeSinceActivity);
 				await cts.CancelAsync();
 				break;
 			}
