Horde: Adding dedicated interfaces and response classes for streams.

#rnx

diff --git a/Engine/Source/Programs/Horde/Plugins/Build/HordeServer.Build/Jobs/JobTaskSource.cs b/Engine/Source/Programs/Horde/Plugins/Build/HordeServer.Build/Jobs/JobTaskSource.cs
--- a/Engine/Source/Programs/Horde/Plugins/Build/HordeServer.Build/Jobs/JobTaskSource.cs
+++ b/Engine/Source/Programs/Horde/Plugins/Build/HordeServer.Build/Jobs/JobTaskSource.cs
@@ -434,7 +434,7 @@ namespace HordeServer.Jobs
 
 					// Validate the agent type and workspace settings
 					IJobStepBatch batch = newJob.Batches[batchIdx];
-					if (!stream.Config.AgentTypes.TryGetValue(graph.Groups[batch.GroupIdx].AgentType, out AgentConfig? agentType))
+					if (!stream.AgentTypes.TryGetValue(graph.Groups[batch.GroupIdx].AgentType, out IAgentType? agentType))
 					{
 						newJob = await SkipBatchAsync(newJob, batch.Id, graph, JobStepBatchError.UnknownAgentType, cancellationToken);
 					}
@@ -446,7 +446,7 @@ namespace HordeServer.Jobs
 					{
 						newJob = await SkipBatchAsync(newJob, batch.Id, graph, JobStepBatchError.NoAgentsInPool, cancellationToken);
 					}
-					else if (!stream.Config.TryGetAgentWorkspace(agentType, out AgentWorkspaceInfo? workspace, out AutoSdkConfig? autoSdkConfig))
+					else if (!buildConfig.TryGetStream(stream.Id, out StreamConfig? streamConfig) || !streamConfig.TryGetAgentWorkspace(agentType, out AgentWorkspaceInfo? workspace, out AutoSdkConfig? autoSdkConfig))
 					{
 						newJob = await SkipBatchAsync(newJob, batch.Id, graph, JobStepBatchError.UnknownWorkspace, cancellationToken);
 					}
diff --git a/Engine/Source/Programs/Horde/Plugins/Build/HordeServer.Build/Jobs/Schedules/ScheduleService.cs b/Engine/Source/Programs/Horde/Plugins/Build/HordeServer.Build/Jobs/Schedules/ScheduleService.cs
--- a/Engine/Source/Programs/Horde/Plugins/Build/HordeServer.Build/Jobs/Schedules/ScheduleService.cs
+++ b/Engine/Source/Programs/Horde/Plugins/Build/HordeServer.Build/Jobs/Schedules/ScheduleService.cs
@@ -359,6 +359,13 @@ namespace HordeServer.Jobs.Schedules
 				return;
 			}
 
+			// Get the stream config
+			StreamConfig? streamConfig;
+			if (!_buildConfig.CurrentValue.TryGetStream(stream.Id, out streamConfig))
+			{
+				return;
+			}
+
 			// Minimum changelist number, inclusive
 			CommitIdWithOrder? minCommitId = schedule.LastTriggerCommitId;
 			bool includeMinCommitId = !schedule.RequireSubmittedChange;
@@ -388,7 +395,7 @@ namespace HordeServer.Jobs.Schedules
 			}
 
 			// Cache the Perforce history as we're iterating through changes to improve query performance
-			ICommitCollection commits = _commitService.GetCollection(stream.Config);
+			ICommitCollection commits = stream.Commits;
 			IAsyncEnumerable<ICommit> commitEnumerable = commits.FindAsync(minCommitId, includeMinCommitId, tags: schedule.Commits, cancellationToken: cancellationToken);
 			await using IAsyncEnumerator<ICommit> commitEnumerator = commitEnumerable.GetAsyncEnumerator(cancellationToken);
 
@@ -500,7 +507,7 @@ namespace HordeServer.Jobs.Schedules
 				template.GetDefaultParameters(options.Parameters, true);
 				template.GetArgumentsForParameters(options.Parameters, options.Arguments);
 
-				IJob newJob = await _jobService.CreateJobAsync(null, stream.Config, templateId, template.Hash, graph, template.Name, commitId, codeCommitId, options, cancellationToken);
+				IJob newJob = await _jobService.CreateJobAsync(null, streamConfig, templateId, template.Hash, graph, template.Name, commitId, codeCommitId, options, cancellationToken);
 				scheduleLogger.LogInformation("Started new job for {StreamId} template {TemplateId} at CL {Change} (Code CL {CodeChange}): {JobId}", stream.Id, templateId, commitId, codeCommitId, newJob.Id);
 				await stream.TryUpdateScheduleTriggerAsync(templateId, utcNow, commitId, new List<JobId> { newJob.Id }, new List<JobId>(), cancellationToken);
 			}
diff --git a/Engine/Source/Programs/Horde/Plugins/Build/HordeServer.Build/Perforce/PerforceWorkspace.cs b/Engine/Source/Programs/Horde/Plugins/Build/HordeServer.Build/Perforce/PerforceWorkspace.cs
--- a/Engine/Source/Programs/Horde/Plugins/Build/HordeServer.Build/Perforce/PerforceWorkspace.cs
+++ b/Engine/Source/Programs/Horde/Plugins/Build/HordeServer.Build/Perforce/PerforceWorkspace.cs
@@ -148,7 +148,7 @@ namespace HordeServer.Perforce
 		/// <param name="workspace">Receives the agent workspace definition</param>
 		/// <param name="autoSdkConfig">Receives the autosdk workspace config</param>
 		/// <returns>True if the agent type was valid, and an agent workspace could be created</returns>
-		public static bool TryGetAgentWorkspace(this StreamConfig streamConfig, AgentConfig agentType, [NotNullWhen(true)] out AgentWorkspaceInfo? workspace, out AutoSdkConfig? autoSdkConfig)
+		public static bool TryGetAgentWorkspace(this StreamConfig streamConfig, IAgentType agentType, [NotNullWhen(true)] out AgentWorkspaceInfo? workspace, out AutoSdkConfig? autoSdkConfig)
 		{
 			// Get the workspace settings
 			if (agentType.Workspace == null)
diff --git a/Engine/Source/Programs/Horde/Plugins/Build/HordeServer.Build/Streams/IStream.cs b/Engine/Source/Programs/Horde/Plugins/Build/HordeServer.Build/Streams/IStream.cs
--- a/Engine/Source/Programs/Horde/Plugins/Build/HordeServer.Build/Streams/IStream.cs
+++ b/Engine/Source/Programs/Horde/Plugins/Build/HordeServer.Build/Streams/IStream.cs
@@ -1,11 +1,15 @@
 // Copyright Epic Games, Inc. All Rights Reserved.
 
+using EpicGames.Horde.Agents.Pools;
 using EpicGames.Horde.Commits;
+using EpicGames.Horde.Issues;
 using EpicGames.Horde.Jobs;
 using EpicGames.Horde.Jobs.Templates;
+using EpicGames.Horde.Projects;
 using EpicGames.Horde.Streams;
 using EpicGames.Horde.Users;
 using Horde.Common.Rpc;
+using HordeServer.Commits;
 using HordeServer.Jobs.Schedules;
 
 namespace HordeServer.Streams
@@ -42,15 +46,75 @@ namespace HordeServer.Streams
 		StreamId Id { get; }
 
 		/// <summary>
-		/// Configuration settings for the stream
+		/// Project that this stream belongs to
 		/// </summary>
-		StreamConfig Config { get; }
+		ProjectId ProjectId { get; }
+
+		/// <summary>
+		/// Name of the stream
+		/// </summary>
+		string Name { get; }
+
+		/// <summary>
+		/// Path to the config file for this stream
+		/// </summary>
+		string? ConfigPath { get; }
+
+		/// <summary>
+		/// Current revision of the config file
+		/// </summary>
+		string ConfigRevision { get; }
+
+		/// <summary>
+		/// Order for this stream on the dashboard
+		/// </summary>
+		int Order { get; }
+
+		/// <summary>
+		/// Notification channel for all jobs in this stream
+		/// </summary>
+		string? NotificationChannel { get; }
+
+		/// <summary>
+		/// Notification channel filter for this template. Can be Success, Failure, or Warnings.
+		/// </summary>
+		string? NotificationChannelFilter { get; }
+
+		/// <summary>
+		/// Channel to post issue triage notifications
+		/// </summary>
+		string? TriageChannel { get; }
+
+		/// <summary>
+		/// Tabs for this stream on the dashboard
+		/// </summary>
+		IReadOnlyList<IStreamTab> Tabs { get; }
+
+		/// <summary>
+		/// Agent types configured for this stream
+		/// </summary>
+		IReadOnlyDictionary<string, IAgentType> AgentTypes { get; }
+
+		/// <summary>
+		/// Workspace types configured for this stream
+		/// </summary>
+		IReadOnlyDictionary<string, IWorkspaceType> WorkspaceTypes { get; }
 
 		/// <summary>
 		/// List of templates available for this stream
 		/// </summary>
 		IReadOnlyDictionary<TemplateId, ITemplateRef> Templates { get; }
 
+		/// <summary>
+		/// Workflows configured for this stream
+		/// </summary>
+		IReadOnlyList<IWorkflow> Workflows { get; }
+
+		/// <summary>
+		/// Default settings for preflights against this stream
+		/// </summary>
+		IDefaultPreflight? DefaultPreflight { get; }
+
 		/// <summary>
 		/// Stream is paused for builds until specified time
 		/// </summary>
@@ -61,6 +125,11 @@ namespace HordeServer.Streams
 		/// </summary>
 		string? PauseComment { get; }
 
+		/// <summary>
+		/// Commits for this stream
+		/// </summary>
+		ICommitCollection Commits { get; }
+
 		/// <summary>
 		/// Get the latest stream state
 		/// </summary>
@@ -98,6 +167,230 @@ namespace HordeServer.Streams
 		Task<IStream?> TryUpdateTemplateRefAsync(TemplateId templateRefId, List<UpdateStepStateRequest>? stepStates = null, CancellationToken cancellationToken = default);
 	}
 
+	/// <summary>
+	/// Style for rendering a tab
+	/// </summary>
+	public enum TabStyle
+	{
+		/// <summary>
+		/// Regular job list
+		/// </summary>
+		Normal,
+
+		/// <summary>
+		/// Omit job names, show condensed view
+		/// </summary>
+		Compact,
+	}
+
+	/// <summary>
+	/// Information about a page to display in the dashboard for a stream
+	/// </summary>
+	public interface IStreamTab
+	{
+		/// <summary>
+		/// Title of this page
+		/// </summary>
+		string Title { get; }
+
+		/// <summary>
+		/// Type of this tab
+		/// </summary>
+		string Type { get; }
+
+		/// <summary>
+		/// Presentation style for this page
+		/// </summary>
+		TabStyle Style { get; }
+
+		/// <summary>
+		/// Whether to show job names on this page
+		/// </summary>
+		bool ShowNames { get; }
+
+		/// <summary>
+		/// Whether to show all user preflights 
+		/// </summary>
+		bool? ShowPreflights { get; }
+
+		/// <summary>
+		/// Names of jobs to include on this page. If there is only one name specified, the name column does not need to be displayed.
+		/// </summary>
+		IReadOnlyList<string>? JobNames { get; }
+
+		/// <summary>
+		/// List of job template names to show on this page.
+		/// </summary>
+		IReadOnlyList<TemplateId>? Templates { get; }
+
+		/// <summary>
+		/// Columns to display for different types of aggregates
+		/// </summary>
+		IReadOnlyList<IStreamTabColumn>? Columns { get; }
+	}
+
+	/// <summary>
+	/// Type of a column in a jobs tab
+	/// </summary>
+	public enum TabColumnType
+	{
+		/// <summary>
+		/// Contains labels
+		/// </summary>
+		Labels,
+
+		/// <summary>
+		/// Contains parameters
+		/// </summary>
+		Parameter
+	}
+
+	/// <summary>
+	/// Describes a column to display on the jobs page
+	/// </summary>
+	public interface IStreamTabColumn
+	{
+		/// <summary>
+		/// The type of column
+		/// </summary>
+		TabColumnType Type { get; }
+
+		/// <summary>
+		/// Heading for this column
+		/// </summary>
+		string Heading { get; }
+
+		/// <summary>
+		/// Category of aggregates to display in this column. If null, includes any aggregate not matched by another column.
+		/// </summary>
+		string? Category { get; }
+
+		/// <summary>
+		/// Parameter to show in this column
+		/// </summary>
+		string? Parameter { get; }
+
+		/// <summary>
+		/// Relative width of this column.
+		/// </summary>
+		int? RelativeWidth { get; }
+	}
+
+	/// <summary>
+	/// Mapping from a BuildGraph agent type to a set of machines on the farm
+	/// </summary>
+	public interface IAgentType
+	{
+		/// <summary>
+		/// Pool of agents to use for this agent type
+		/// </summary>
+		PoolId Pool { get; }
+
+		/// <summary>
+		/// Name of the workspace to sync
+		/// </summary>
+		string? Workspace { get; }
+
+		/// <summary>
+		/// Path to the temporary storage dir
+		/// </summary>
+		string? TempStorageDir { get; }
+
+		/// <summary>
+		/// Environment variables to be set when executing the job
+		/// </summary>
+		IReadOnlyDictionary<string, string>? Environment { get; }
+	}
+
+	/// <summary>
+	/// Information about a workspace type
+	/// </summary>
+	public interface IWorkspaceType
+	{
+		/// <summary>
+		/// Name of the Perforce server cluster to use
+		/// </summary>
+		string? Cluster { get; }
+
+		/// <summary>
+		/// The Perforce server and port (eg. perforce:1666)
+		/// </summary>
+		string? ServerAndPort { get; }
+
+		/// <summary>
+		/// User to log into Perforce with (defaults to buildmachine)
+		/// </summary>
+		string? UserName { get; }
+
+		/// <summary>
+		/// Password to use to log into the workspace
+		/// </summary>
+		string? Password { get; }
+
+		/// <summary>
+		/// Identifier to distinguish this workspace from other workspaces. Defaults to the workspace type name.
+		/// </summary>
+		string? Identifier { get; }
+
+		/// <summary>
+		/// Override for the stream to sync
+		/// </summary>
+		string? Stream { get; }
+
+		/// <summary>
+		/// Custom view for the workspace
+		/// </summary>
+		IReadOnlyList<string>? View { get; }
+
+		/// <summary>
+		/// Whether to use an incrementally synced workspace
+		/// </summary>
+		bool? Incremental { get; }
+
+		/// <summary>
+		/// Whether to use the AutoSDK
+		/// </summary>
+		bool? UseAutoSdk { get; }
+
+		/// <summary>
+		/// View for the AutoSDK paths to sync. If null, the whole thing will be synced.
+		/// </summary>
+		IReadOnlyList<string>? AutoSdkView { get; }
+
+		/// <summary>
+		/// Method to use when syncing/materializing data from Perforce
+		/// </summary>
+		string? Method { get; }
+
+		/// <summary>
+		/// Minimum disk space that must be available *after* syncing this workspace (in megabytes)
+		/// If not available, the job will be aborted.
+		/// </summary>
+		long? MinScratchSpace { get; }
+
+		/// <summary>
+		/// Threshold for when to trigger an automatic conform of agent. Measured in megabytes free on disk.
+		/// Set to null or 0 to disable.
+		/// </summary>
+		long? ConformDiskFreeSpace { get; }
+	}
+
+	/// <summary>
+	/// Specifies defaults for running a preflight
+	/// </summary>
+	public interface IDefaultPreflight
+	{
+		/// <summary>
+		/// The template id to query
+		/// </summary>
+		TemplateId? TemplateId { get; }
+
+		/// <summary>
+		/// Query for the change to use
+		/// </summary>
+		IChangeQuery? Change { get; }
+	}
+
 	/// <summary>
 	/// Job template in a stream
 	/// </summary>
diff --git a/Engine/Source/Programs/Horde/Plugins/Build/HordeServer.Build/Streams/StreamCollection.cs b/Engine/Source/Programs/Horde/Plugins/Build/HordeServer.Build/Streams/StreamCollection.cs
--- a/Engine/Source/Programs/Horde/Plugins/Build/HordeServer.Build/Streams/StreamCollection.cs
+++ b/Engine/Source/Programs/Horde/Plugins/Build/HordeServer.Build/Streams/StreamCollection.cs
@@ -2,11 +2,14 @@
 
 using EpicGames.Core;
 using EpicGames.Horde.Commits;
+using EpicGames.Horde.Issues;
 using EpicGames.Horde.Jobs;
 using EpicGames.Horde.Jobs.Templates;
+using EpicGames.Horde.Projects;
 using EpicGames.Horde.Streams;
 using EpicGames.Horde.Users;
 using HordeServer.Acls;
+using HordeServer.Commits;
 using HordeServer.Jobs.Schedules;
 using HordeServer.Server;
 using HordeServer.Utilities;
@@ -36,6 +39,29 @@ namespace HordeServer.Streams
 			public DateTime? PausedUntil => _document.PausedUntil;
 			public string? PauseComment => _document.PauseComment;
 
+			ProjectId IStream.ProjectId => _config.ProjectConfig.Id;
+			string IStream.Name => _config.Name;
+			string? IStream.ConfigPath => _config.Path;
+			string IStream.ConfigRevision =>_config.Revision;
+			int IStream.Order => _config.Order;
+			string? IStream.NotificationChannel => _config.NotificationChannel;
+			string? IStream.NotificationChannelFilter => _config.NotificationChannelFilter;
+			string? IStream.TriageChannel => _config.TriageChannel;
+			IReadOnlyList<IStreamTab> IStream.Tabs => _config.Tabs;
+			IDefaultPreflight? IStream.DefaultPreflight => _config.DefaultPreflight;
+
+			IReadOnlyDictionary<string, IAgentType>? _cachedAgentTypes;
+			IReadOnlyDictionary<string, IAgentType> IStream.AgentTypes
+				=> _cachedAgentTypes ??= _config.AgentTypes.ToDictionary(x => x.Key, x => (IAgentType)x.Value);
+
+			IReadOnlyDictionary<string, IWorkspaceType>? _cachedWorkspaceTypes;
+			IReadOnlyDictionary<string, IWorkspaceType> IStream.WorkspaceTypes
+				=> _cachedWorkspaceTypes ??= _config.WorkspaceTypes.ToDictionary(x => x.Key, x => (IWorkspaceType)x.Value);
+
+			IReadOnlyList<IWorkflow> IStream.Workflows => _config.Workflows;
+
+			ICommitCollection IStream.Commits => _collection._commitService.GetCollection(_config);
+
 			public Stream(StreamCollection collection, StreamConfig config, StreamDoc document)
 			{
 				_collection = collection;
@@ -264,15 +290,17 @@ namespace HordeServer.Streams
 		}
 
 		readonly IMongoCollection<StreamDoc> _streams;
+		readonly ICommitService _commitService;
 		readonly IOptionsMonitor<BuildConfig> _buildConfig;
 		readonly IClock _clock;
 
 		/// <summary>
 		/// Constructor
 		/// </summary>
-		public StreamCollection(IMongoService mongoService, IOptionsMonitor<BuildConfig> buildConfig, IClock clock)
+		public StreamCollection(IMongoService mongoService, ICommitService commitService, IOptionsMonitor<BuildConfig> buildConfig, IClock clock)
 		{
 			_streams = mongoService.GetCollection<StreamDoc>("Streams");
+			_commitService = commitService;
 			_buildConfig = buildConfig;
 			_clock = clock;
 		}
diff --git a/Engine/Source/Programs/Horde/Plugins/Build/HordeServer.Build/Streams/StreamConfig.cs b/Engine/Source/Programs/Horde/Plugins/Build/HordeServer.Build/Streams/StreamConfig.cs
--- a/Engine/Source/Programs/Horde/Plugins/Build/HordeServer.Build/Streams/StreamConfig.cs
+++ b/Engine/Source/Programs/Horde/Plugins/Build/HordeServer.Build/Streams/StreamConfig.cs
@@ -551,46 +551,22 @@ namespace HordeServer.Streams
 		}
 	}
 
-	/// <summary>
-	/// Style for rendering a tab
-	/// </summary>
-	public enum TabStyle
-	{
-		/// <summary>
-		/// Regular job list
-		/// </summary>
-		Normal,
-
-		/// <summary>
-		/// Omit job names, show condensed view
-		/// </summary>
-		Compact,
-	}
-
 	/// <summary>
 	/// Information about a page to display in the dashboard for a stream
 	/// </summary>
-	public class TabConfig
+	public class TabConfig : IStreamTab
 	{
-		/// <summary>
-		/// Title of this page
-		/// </summary>
+		/// <inheritdoc/>
 		[Required]
 		public string Title { get; set; } = null!;
 
-		/// <summary>
-		/// Type of this tab
-		/// </summary>
+		/// <inheritdoc/>
 		public string Type { get; set; } = "Jobs";
 
-		/// <summary>
-		/// Presentation style for this page
-		/// </summary>
+		/// <inheritdoc/>
 		public TabStyle Style { get; set; }
 
-		/// <summary>
-		/// Whether to show job names on this page
-		/// </summary>
+		/// <inheritdoc/>
 		public bool ShowNames
 		{
 			get => _showNames ?? (Style != TabStyle.Compact);
@@ -598,79 +574,48 @@ namespace HordeServer.Streams
 		}
 		bool? _showNames;
 
-		/// <summary>
-		/// Whether to show all user preflights 
-		/// </summary>
+		/// <inheritdoc/>
 		public bool? ShowPreflights { get; set; }
 
-		/// <summary>
-		/// Names of jobs to include on this page. If there is only one name specified, the name column does not need to be displayed.
-		/// </summary>
+		/// <inheritdoc cref="IStreamTab.JobNames"/>
 		public List<string>? JobNames { get; set; }
+		IReadOnlyList<string>? IStreamTab.JobNames => JobNames;
 
-		/// <summary>
-		/// List of job template names to show on this page.
-		/// </summary>
+		/// <inheritdoc cref="IStreamTab.Templates"/>
 		public List<TemplateId>? Templates { get; set; }
+		IReadOnlyList<TemplateId>? IStreamTab.Templates => Templates;
 
-		/// <summary>
-		/// Columns to display for different types of aggregates
-		/// </summary>
+		/// <inheritdoc cref="IStreamTab.Columns"/>
 		public List<TabColumnConfig>? Columns { get; set; }
-	}
-
-	/// <summary>
-	/// Type of a column in a jobs tab
-	/// </summary>
-	public enum TabColumnType
-	{
-		/// <summary>
-		/// Contains labels
-		/// </summary>
-		Labels,
-
-		/// <summary>
-		/// Contains parameters
-		/// </summary>
-		Parameter
+		IReadOnlyList<IStreamTabColumn>? IStreamTab.Columns => Columns;
 	}
 
 	/// <summary>
 	/// Describes a column to display on the jobs page
 	/// </summary>
-	public class TabColumnConfig
+	public class TabColumnConfig : IStreamTabColumn
 	{
-		/// <summary>
-		/// The type of column
-		/// </summary>
+		/// <inheritdoc/>
 		public TabColumnType Type { get; set; } = TabColumnType.Labels;
 
-		/// <summary>
-		/// Heading for this column
-		/// </summary>
+		/// <inheritdoc/>
 		[Required]
 		public string Heading { get; set; } = null!;
 
-		/// <summary>
-		/// Category of aggregates to display in this column. If null, includes any aggregate not matched by another column.
-		/// </summary>
+		/// <inheritdoc/>
 		public string? Category { get; set; }
 
-		/// <summary>
-		/// Parameter to show in this column
-		/// </summary>
+		/// <inheritdoc/>
 		public string? Parameter { get; set; }
 
-		/// <summary>
-		/// Relative width of this column.
-		/// </summary>
+		/// <inheritdoc/>
 		public int? RelativeWidth { get; set; }
 	}
 
 	/// <summary>
 	/// Mapping from a BuildGraph agent type to a set of machines on the farm
 	/// </summary>
-	public class AgentConfig
+	public class AgentConfig : IAgentType
 	{
 		/// <summary>
 		/// Base agent config to inherit settings from
@@ -696,6 +641,7 @@ namespace HordeServer.Streams
 		/// Environment variables to be set when executing the job
 		/// </summary>
 		public Dictionary<string, string>? Environment { get; set; }
+		IReadOnlyDictionary<string, string>? IAgentType.Environment => Environment;
 
 		/// <summary>
 		/// Tokens to allocate for this agent type
@@ -725,87 +671,59 @@ namespace HordeServer.Streams
 	/// <summary>
 	/// Information about a workspace type
 	/// </summary>
-	public class WorkspaceConfig
+	public class WorkspaceConfig : IWorkspaceType
 	{
-		/// <summary>
-		/// Base workspace to derive from
-		/// </summary>
+		/// <inheritdoc/>
 		public string? Base { get; set; }
 
-		/// <summary>
-		/// Name of the Perforce server cluster to use
-		/// </summary>
+		/// <inheritdoc/>
 		public string? Cluster { get; set; }
 
-		/// <summary>
-		/// The Perforce server and port (eg. perforce:1666)
-		/// </summary>
+		/// <inheritdoc/>
 		public string? ServerAndPort { get; set; }
 
-		/// <summary>
-		/// User to log into Perforce with (defaults to buildmachine)
-		/// </summary>
+		/// <inheritdoc/>
 		public string? UserName { get; set; }
 
-		/// <summary>
-		/// Password to use to log into the workspace
-		/// </summary>
+		/// <inheritdoc/>
 		public string? Password { get; set; }
 
-		/// <summary>
-		/// Identifier to distinguish this workspace from other workspaces. Defaults to the workspace type name.
-		/// </summary>
+		/// <inheritdoc/>
 		public string? Identifier { get; set; }
 
-		/// <summary>
-		/// Override for the stream to sync
-		/// </summary>
+		/// <inheritdoc/>
 		public string? Stream { get; set; }
 
-		/// <summary>
-		/// Custom view for the workspace
-		/// </summary>
+		/// <inheritdoc cref="IWorkspaceType.View"/>
 		[ConfigMergeStrategy(ConfigMergeStrategy.Append)]
 		public List<string>? View { get; set; }
+		IReadOnlyList<string>? IWorkspaceType.View => View;
 
-		/// <summary>
-		/// Whether to use an incrementally synced workspace
-		/// </summary>
+		/// <inheritdoc/>
 		public bool? Incremental { get; set; }
 
-		/// <summary>
-		/// Whether to use the AutoSDK
-		/// </summary>
+		/// <inheritdoc/>
 		public bool? UseAutoSdk { get; set; }
 
-		/// <summary>
-		/// View for the AutoSDK paths to sync. If null, the whole thing will be synced.
-		/// </summary>
+		/// <inheritdoc cref="IWorkspaceType.AutoSdkView"/>
 		[ConfigMergeStrategy(ConfigMergeStrategy.Append)]
 		public List<string>? AutoSdkView { get; set; }
+		IReadOnlyList<string>? IWorkspaceType.AutoSdkView => AutoSdkView;
 
-		/// <summary>
-		/// Method to use when syncing/materializing data from Perforce
-		/// </summary>
+		/// <inheritdoc/>
 		public string? Method { get; set; } = null;
 
-		/// <summary>
-		/// Minimum disk space that must be available *after* syncing this workspace (in megabytes)
-		/// If not available, the job will be aborted.
-		/// </summary>
+		/// <inheritdoc/>
 		public long? MinScratchSpace { get; set; } = null;
 
-		/// <summary>
-		/// Threshold for when to trigger an automatic conform of agent. Measured in megabytes free on disk.
-		/// Set to null or 0 to disable.
-		/// </summary>
+		/// <inheritdoc/>
 		public long? ConformDiskFreeSpace { get; set; } = null;
 	}
 
 	/// <summary>
 	/// Specifies defaults for running a preflight
 	/// </summary>
-	public class DefaultPreflightConfig
+	public class DefaultPreflightConfig : IDefaultPreflight
 	{
 		/// <summary>
 		/// The template id to query
@@ -817,6 +735,9 @@ namespace HordeServer.Streams
 		/// </summary>
 		[Obsolete("Set on template instead")]
 		public ChangeQueryConfig? Change { get; set; }
+
+		[Obsolete("Set on template instead")]
+		IChangeQuery? IDefaultPreflight.Change => Change;
 	}
 
 	/// <summary>
diff --git a/Engine/Source/Programs/Horde/Plugins/Build/HordeServer.Build/Streams/StreamMessages.cs b/Engine/Source/Programs/Horde/Plugins/Build/HordeServer.Build/Streams/StreamMessages.cs
--- a/Engine/Source/Programs/Horde/Plugins/Build/HordeServer.Build/Streams/StreamMessages.cs
+++ b/Engine/Source/Programs/Horde/Plugins/Build/HordeServer.Build/Streams/StreamMessages.cs
@@ -1,10 +1,14 @@
 // Copyright Epic Games, Inc. All Rights Reserved.
 
 using System.ComponentModel.DataAnnotations;
+using EpicGames.Core;
+using EpicGames.Horde.Agents.Pools;
+using EpicGames.Horde.Issues;
 using EpicGames.Horde.Jobs;
 using EpicGames.Horde.Jobs.Templates;
+using EpicGames.Horde.Projects;
+using EpicGames.Horde.Streams;
 using EpicGames.Horde.Users;
-using HordeServer.Issues;
 using HordeServer.Jobs.Schedules;
 
 #pragma warning disable CA2227 // Change x to be read-only by removing the property setter
@@ -19,12 +23,12 @@ namespace HordeServer.Streams
 		/// <summary>
 		/// Unique id of the stream
 		/// </summary>
-		public string Id { get; set; }
+		public StreamId Id { get; set; }
 
 		/// <summary>
 		/// Unique id of the project containing this stream
 		/// </summary>
-		public string ProjectId { get; set; }
+		public ProjectId ProjectId { get; set; }
 
 		/// <summary>
 		/// Name of the stream
@@ -69,22 +73,22 @@ namespace HordeServer.Streams
 		/// <summary>
 		/// Default template to use for preflights
 		/// </summary>
-		public DefaultPreflightConfig? DefaultPreflight { get; set; }
+		public DefaultPreflightMessage? DefaultPreflight { get; set; }
 
 		/// <summary>
 		/// List of tabs to display for this stream
 		/// </summary>
-		public List<TabConfig> Tabs { get; set; }
+		public List<GetStreamTabResponse> Tabs { get; set; }
 
 		/// <summary>
 		/// Map of agent name to type
 		/// </summary>
-		public Dictionary<string, AgentConfig> AgentTypes { get; set; }
+		public Dictionary<string, GetAgentTypeResponse> AgentTypes { get; set; }
 
 		/// <summary>
 		/// Map of workspace name to type
 		/// </summary>
-		public Dictionary<string, WorkspaceConfig>? WorkspaceTypes { get; set; }
+		public Dictionary<string, GetWorkspaceTypeResponse>? WorkspaceTypes { get; set; }
 
 		/// <summary>
 		/// Templates for jobs in this stream
@@ -104,7 +108,7 @@ namespace HordeServer.Streams
 		/// <summary>
 		/// Workflows for this stream
 		/// </summary>
-		public List<WorkflowConfig> Workflows { get; set; }
+		public List<GetWorkflowResponse> Workflows { get; set; }
 
 		/// <summary>
 		/// Constructor
@@ -113,24 +117,150 @@ namespace HordeServer.Streams
 		/// <param name="templates">Templates for this stream</param>
 		public GetStreamResponse(IStream stream, List<GetTemplateRefResponse> templates)
 		{
-			Id = stream.Id.ToString();
-			ProjectId = stream.Config.ProjectConfig.Id.ToString();
-			Name = stream.Config.Name;
-			ConfigPath = stream.Config.Path ?? String.Empty;
-			ConfigRevision = stream.Config.Revision;
-			Order = stream.Config.Order;
-			NotificationChannel = stream.Config.NotificationChannel;
-			NotificationChannelFilter = stream.Config.NotificationChannelFilter;
-			TriageChannel = stream.Config.TriageChannel;
-			DefaultPreflightTemplate = stream.Config.DefaultPreflight?.TemplateId?.ToString();
-			DefaultPreflight = stream.Config.DefaultPreflight;
-			Tabs = stream.Config.Tabs;
-			AgentTypes = stream.Config.AgentTypes;
-			WorkspaceTypes = stream.Config.WorkspaceTypes;
+			Id = stream.Id;
+			ProjectId = stream.ProjectId;
+			Name = stream.Name;
+			ConfigPath = stream.ConfigPath ?? String.Empty;
+			ConfigRevision = stream.ConfigRevision;
+			Order = stream.Order;
+			NotificationChannel = stream.NotificationChannel;
+			NotificationChannelFilter = stream.NotificationChannelFilter;
+			TriageChannel = stream.TriageChannel;
+			DefaultPreflightTemplate = stream.DefaultPreflight?.TemplateId?.ToString();
+			DefaultPreflight = (stream.DefaultPreflight == null) ? null : new DefaultPreflightMessage(stream.DefaultPreflight);
+			Tabs = stream.Tabs.ConvertAll(x => new GetStreamTabResponse(x));
+			AgentTypes = stream.AgentTypes.ToDictionary(x => x.Key, x => new GetAgentTypeResponse(x.Value));
+			WorkspaceTypes = stream.WorkspaceTypes.ToDictionary(x => x.Key, x => new GetWorkspaceTypeResponse(x.Value));
 			Templates = templates;
+			Workflows = stream.Workflows.ConvertAll(x => new GetWorkflowResponse(x));
 			PausedUntil = stream.PausedUntil;
 			PauseComment = stream.PauseComment;
-			Workflows = stream.Config.Workflows;
+		}
+	}
+
+	/// <summary>
+	/// Information about the default preflight to run
+	/// </summary>
+	public class DefaultPreflightMessage : IDefaultPreflight
+	{
+		/// <inheritdoc/>
+		public TemplateId? TemplateId { get; set; }
+
+		/// <inheritdoc/>
+		public ChangeQueryMessage? Change { get; set; }
+		IChangeQuery? IDefaultPreflight.Change => Change;
+
+		/// <summary>
+		/// Constructor
+		/// </summary>
+		public DefaultPreflightMessage()
+		{ }
+
+		/// <summary>
+		/// Constructor
+		/// </summary>
+		public DefaultPreflightMessage(IDefaultPreflight other)
+		{
+			TemplateId = other.TemplateId;
+			Change = (other.Change == null) ? null : new ChangeQueryMessage(other.Change);
+		}
+	}
+
+	/// <summary>
+	/// Information about a page to display in the dashboard for a stream
+	/// </summary>
+	public class GetStreamTabResponse : IStreamTab
+	{
+		/// <inheritdoc/>
+		public string Title { get; set; }
+
+		/// <inheritdoc/>
+		public string Type { get; set; }
+
+		/// <inheritdoc/>
+		public TabStyle Style { get; set; }
+
+		/// <inheritdoc/>
+		public bool ShowNames { get; set; }
+
+		/// <inheritdoc/>
+		public bool? ShowPreflights { get; set; }
+
+		/// <inheritdoc cref="IStreamTab.JobNames"/>
+		public List<string>? JobNames { get; set; }
+		IReadOnlyList<string>? IStreamTab.JobNames => JobNames;
+
+		/// <inheritdoc cref="IStreamTab.Templates"/>
+		public List<TemplateId>? Templates { get; set; }
+		IReadOnlyList<TemplateId>? IStreamTab.Templates => Templates;
+
+		/// <inheritdoc cref="IStreamTab.Columns"/>
+		public List<GetStreamTabColumnResponse>? Columns { get; set; }
+		IReadOnlyList<IStreamTabColumn>? IStreamTab.Columns => Columns;
+
+		/// <summary>
+		/// Constructor
+		/// </summary>
+		public GetStreamTabResponse()
+		{
+			Title = String.Empty;
+			Type = String.Empty;
+		}
+
+		/// <summary>
+		/// Constructor
+		/// </summary>
+		public GetStreamTabResponse(IStreamTab other)
+		{
+			Title = other.Title;
+			Type = other.Type;
+			Style = other.Style;
+			ShowNames = other.ShowNames;
+			ShowPreflights = other.ShowPreflights;
+			JobNames = other.JobNames?.ToList();
+			Templates = other.Templates?.ToList();
+			Columns = other.Columns?.ConvertAll(x => new GetStreamTabColumnResponse(x));
+		}
+	}
+
+	/// <summary>
+	/// Describes a column to display on the jobs page
+	/// </summary>
+	public class GetStreamTabColumnResponse : IStreamTabColumn
+	{
+		/// <inheritdoc/>
+		public TabColumnType Type { get; set; }
+
+		/// <inheritdoc/>
+		public string Heading { get; set; }
+
+		/// <inheritdoc/>
+		public string? Category { get; set; }
+
+		/// <inheritdoc/>
+		public string? Parameter { get; set; }
+
+		/// <inheritdoc/>
+		public int? RelativeWidth { get; set; }
+
+		/// <summary>
+		/// Default constructor
+		/// </summary>
+		public GetStreamTabColumnResponse()
+		{
+			Heading = String.Empty;
+		}
+
+		/// <summary>
+		/// Constructor
+		/// </summary>
+		public GetStreamTabColumnResponse(IStreamTabColumn other)
+		{
+			Type = other.Type;
+			Heading = other.Heading;
+			Category = other.Category;
+			Parameter = other.Parameter;
+			RelativeWidth = other.RelativeWidth;
 		}
 	}
 
@@ -142,7 +272,7 @@ namespace HordeServer.Streams
 		/// <summary>
 		/// Pool of agents to use for this agent type
 		/// </summary>
-		public string Pool { get; set; }
+		public PoolId Pool { get; set; }
 
 		/// <summary>
 		/// Name of the workspace to sync
@@ -162,92 +292,90 @@ namespace HordeServer.Streams
 		/// <summary>
 		/// Constructor
 		/// </summary>
-		/// <param name="pool">Pool of agents to use for this agent type</param>
-		/// <param name="workspace">Name of the workspace to sync</param>
-		/// <param name="tempStorageDir">Path to the temp storage directory</param>
-		/// <param name="environment">Environment variables to be set when executing this job</param>
-		public GetAgentTypeResponse(string pool, string? workspace, string? tempStorageDir, Dictionary<string, string>? environment)
+		public GetAgentTypeResponse()
+		{ }
+
+		/// <summary>
+		/// Constructor
+		/// </summary>
+		public GetAgentTypeResponse(IAgentType other)
 		{
-			Pool = pool;
-			Workspace = workspace;
-			TempStorageDir = tempStorageDir;
-			Environment = environment;
+			Pool = other.Pool;
+			Workspace = other.Workspace;
+			TempStorageDir = other.TempStorageDir;
+			Environment = other.Environment?.ToDictionary();
 		}
 	}
 
 	/// <summary>
 	/// Information about a workspace type
 	/// </summary>
-	public class GetWorkspaceTypeResponse
+	public class GetWorkspaceTypeResponse : IWorkspaceType
 	{
-		/// <summary>
-		/// The Perforce server cluster
-		/// </summary>
+		/// <inheritdoc/>
 		public string? Cluster { get; set; }
 
-		/// <summary>
-		/// The Perforce server and port (eg. perforce:1666)
-		/// </summary>
+		/// <inheritdoc/>
 		public string? ServerAndPort { get; set; }
 
-		/// <summary>
-		/// User to log into Perforce with (defaults to buildmachine)
-		/// </summary>
+		/// <inheritdoc/>
 		public string? UserName { get; set; }
 
-		/// <summary>
-		/// Identifier to distinguish this workspace from other workspaces. Defaults to the workspace type name.
-		/// </summary>
+		/// <inheritdoc/>
+		string? IWorkspaceType.Password => throw new NotSupportedException();
+
+		/// <inheritdoc/>
 		public string? Identifier { get; set; }
 
-		/// <summary>
-		/// Override for the stream to sync
-		/// </summary>
+		/// <inheritdoc/>
 		public string? Stream { get; set; }
 
-		/// <summary>
-		/// Custom view for the workspace
-		/// </summary>
+		/// <inheritdoc cref="IWorkspaceType.View"/>
 		public List<string>? View { get; set; }
+		IReadOnlyList<string>? IWorkspaceType.View => View;
 
-		/// <summary>
-		/// Whether to use an incrementally synced workspace
-		/// </summary>
-		public bool Incremental { get; set; }
+		/// <inheritdoc/>
+		public bool? Incremental { get; set; }
 
-		/// <summary>
-		/// Whether to use the AutoSDK
-		/// </summary>
-		public bool UseAutoSdk { get; set; }
+		/// <inheritdoc/>
+		public bool? UseAutoSdk { get; set; }
+
+		/// <inheritdoc cref="IWorkspaceType.AutoSdkView"/>
+		public List<string>? AutoSdkView { get; set; }
+		IReadOnlyList<string>? IWorkspaceType.AutoSdkView => AutoSdkView;
+
+		/// <inheritdoc/>
+		public string? Method { get; set; } = null;
+
+		/// <inheritdoc/>
+		public long? MinScratchSpace { get; set; } = null;
+
+		/// <inheritdoc/>
+		public long? ConformDiskFreeSpace { get; set; } = null;
 
 		/// <summary>
 		/// Default constructor
 		/// </summary>
 		public GetWorkspaceTypeResponse()
-		{
-		}
+		{ }
 
 		/// <summary>
 		/// Constructor
 		/// </summary>
-		/// <param name="cluster">The server cluster</param>
-		/// <param name="serverAndPort">The perforce server</param>
-		/// <param name="userName">The perforce user name</param>
-		/// <param name="identifier">Identifier to distinguish this workspace from other workspaces. Defaults to the workspace type name.</param>
-		/// <param name="stream">Override for the stream to sync</param>
-		/// <param name="view">Custom view for the workspace</param>
-		/// <param name="incremental">Whether to use an incrementally synced workspace</param>
-		/// <param name="useAutoSdk">Whether to use the AutoSDK</param>
-		public GetWorkspaceTypeResponse(string? cluster, string? serverAndPort, string? userName, string? identifier, string? stream, List<string>? view, bool incremental, bool useAutoSdk)
+		public GetWorkspaceTypeResponse(IWorkspaceType other)
 		{
-			Cluster = cluster;
-			ServerAndPort = serverAndPort;
-			UserName = userName;
-			Identifier = identifier;
-			Stream = stream;
-			View = view;
-			Incremental = incremental;
-			UseAutoSdk = useAutoSdk;
+			Cluster = other.Cluster;
+			ServerAndPort = other.ServerAndPort;
+			UserName = other.UserName;
+			Identifier = other.Identifier;
+			Stream = other.Stream;
+			View = other.View?.ToList();
+			Incremental = other.Incremental;
+			UseAutoSdk = other.UseAutoSdk;
+			AutoSdkView = other.AutoSdkView?.ToList();
+			Method = other.Method;
+			MinScratchSpace = other.MinScratchSpace;
+			ConformDiskFreeSpace = other.ConformDiskFreeSpace;
 		}
 	}
 
