Horde: Remove some internal server methods from IStream.

#rnx

diff --git a/Engine/Source/Programs/Horde/Plugins/Build/HordeServer.Build/Streams/IStream.cs b/Engine/Source/Programs/Horde/Plugins/Build/HordeServer.Build/Streams/IStream.cs
--- a/Engine/Source/Programs/Horde/Plugins/Build/HordeServer.Build/Streams/IStream.cs
+++ b/Engine/Source/Programs/Horde/Plugins/Build/HordeServer.Build/Streams/IStream.cs
@@ -442,30 +442,6 @@ namespace HordeServer.Streams
 	/// </summary>
 	static class StreamExtensions
 	{
-		/// <summary>
-		/// Converts to a public response object
-		/// </summary>
-		/// <param name="stream">The stream object</param>
-		/// <param name="apiTemplateRefs">The template refs for this stream. Passed separately because they have their own ACL.</param>
-		/// <returns>New response instance</returns>
-		public static GetStreamResponse ToApiResponse(this IStream stream, List<GetTemplateRefResponse> apiTemplateRefs)
-		{
-			return new GetStreamResponse(stream, apiTemplateRefs);
-		}
-
-		/// <summary>
-		/// Converts to an RPC response object
-		/// </summary>
-		/// <param name="streamConfig">The stream config object</param>
-		/// <returns>New response instance</returns>
-		public static RpcGetStreamResponse ToRpcResponse(this StreamConfig streamConfig)
-		{
-			RpcGetStreamResponse response = new RpcGetStreamResponse();
-			response.Name = streamConfig.Name;
-			response.AgentTypes.Add(streamConfig.AgentTypes.ToDictionary(x => x.Key, x => x.Value.ToRpcResponse()));
-			return response;
-		}
-
 		/// <summary>
 		/// Check if stream is paused for new builds
 		/// </summary>
diff --git a/Engine/Source/Programs/Horde/Plugins/Build/HordeServer.Build/Streams/StreamsController.cs b/Engine/Source/Programs/Horde/Plugins/Build/HordeServer.Build/Streams/StreamsController.cs
--- a/Engine/Source/Programs/Horde/Plugins/Build/HordeServer.Build/Streams/StreamsController.cs
+++ b/Engine/Source/Programs/Horde/Plugins/Build/HordeServer.Build/Streams/StreamsController.cs
@@ -170,7 +170,7 @@ namespace HordeServer.Streams
 				}
 			}
 
-			return stream.ToApiResponse(apiTemplateRefs);
+			return new GetStreamResponse(stream, apiTemplateRefs);
 		}
 
 		/// <summary>
