Horde: Remove TemplateConfig from ITemplateRef interface, so we can move it into EpicGames.Horde.

#rnx

diff --git a/Engine/Source/Programs/Horde/Plugins/Build/HordeServer.Build/Jobs/Schedules/ScheduleService.cs b/Engine/Source/Programs/Horde/Plugins/Build/HordeServer.Build/Jobs/Schedules/ScheduleService.cs
--- a/Engine/Source/Programs/Horde/Plugins/Build/HordeServer.Build/Jobs/Schedules/ScheduleService.cs
+++ b/Engine/Source/Programs/Horde/Plugins/Build/HordeServer.Build/Jobs/Schedules/ScheduleService.cs
@@ -366,6 +366,12 @@ namespace HordeServer.Jobs.Schedules
 				return;
 			}
 
+			TemplateRefConfig? templateRefConfig;
+			if (!streamConfig.TryGetTemplate(templateRef.Id, out templateRefConfig))
+			{
+				return;
+			}
+
 			// Minimum changelist number, inclusive
 			CommitIdWithOrder? minCommitId = schedule.LastTriggerCommitId;
 			bool includeMinCommitId = !schedule.RequireSubmittedChange;
@@ -483,7 +489,7 @@ namespace HordeServer.Jobs.Schedules
 			}
 
 			// Get the matching template
-			ITemplate template = await _templateCollection.GetOrAddAsync(templateRef.Config);
+			ITemplate template = await _templateCollection.GetOrAddAsync(templateRefConfig);
 
 			// Register the graph for it
 			IGraph graph = await _graphs.AddAsync(template, cancellationToken);
@@ -502,7 +508,7 @@ namespace HordeServer.Jobs.Schedules
 			{
 				cancellationToken.ThrowIfCancellationRequested();
 
-				CreateJobOptions options = new CreateJobOptions(templateRef.Config);
+				CreateJobOptions options = new CreateJobOptions(templateRefConfig);
 				options.Priority = template.Priority;
 				template.GetDefaultParameters(options.Parameters, true);
 				template.GetArgumentsForParameters(options.Parameters, options.Arguments);
diff --git a/Engine/Source/Programs/Horde/Plugins/Build/HordeServer.Build/Streams/IStream.cs b/Engine/Source/Programs/Horde/Plugins/Build/HordeServer.Build/Streams/IStream.cs
--- a/Engine/Source/Programs/Horde/Plugins/Build/HordeServer.Build/Streams/IStream.cs
+++ b/Engine/Source/Programs/Horde/Plugins/Build/HordeServer.Build/Streams/IStream.cs
@@ -401,19 +401,70 @@ namespace HordeServer.Streams
 		TemplateId Id { get; }
 
 		/// <summary>
-		/// Configuration of this template ref
+		/// Whether to show badges in UGS for these jobs
 		/// </summary>
-		TemplateRefConfig Config { get; }
+		bool ShowUgsBadges { get; }
+
+		/// <summary>
+		/// Whether to show alerts in UGS for these jobs
+		/// </summary>
+		bool ShowUgsAlerts { get; }
+
+		/// <summary>
+		/// Notification channel for this template. Overrides the stream channel if set.
+		/// </summary>
+		string? NotificationChannel { get; }
+
+		/// <summary>
+		/// Notification channel filter for this template. Can be a combination of "Success", "Failure" and "Warnings" separated by pipe characters.
+		/// </summary>
+		string? NotificationChannelFilter { get; }
+
+		/// <summary>
+		/// Triage channel for this template. Overrides the stream channel if set.
+		/// </summary>
+		string? TriageChannel { get; }
 
 		/// <summary>
 		/// List of schedules for this template
 		/// </summary>
 		ISchedule? Schedule { get; }
 
+		/// <summary>
+		/// List of chained job triggers
+		/// </summary>
+		IReadOnlyList<IChainedJobTemplate>? ChainedJobs { get; }
+
 		/// <summary>
 		/// List of template step states
 		/// </summary>
 		IReadOnlyList<ITemplateStep> StepStates { get; }
+
+		/// <summary>
+		/// Default change to use for this job
+		/// </summary>
+		IReadOnlyList<IChangeQuery>? DefaultChange { get; }
+	}
+
+	/// <summary>
+	/// Trigger for another template
+	/// </summary>
+	public interface IChainedJobTemplate
+	{
+		/// <summary>
+		/// Name of the target that needs to complete before starting the other template
+		/// </summary>
+		string Trigger { get; }
+
+		/// <summary>
+		/// Id of the template to trigger
+		/// </summary>
+		TemplateId TemplateId { get; }
+
+		/// <summary>
+		/// Whether to use the default change for the template rather than the change for the parent job.
+		/// </summary>
+		bool UseDefaultChangeForTemplate { get; }
 	}
 
 	/// <summary>
diff --git a/Engine/Source/Programs/Horde/Plugins/Build/HordeServer.Build/Streams/StreamCollection.cs b/Engine/Source/Programs/Horde/Plugins/Build/HordeServer.Build/Streams/StreamCollection.cs
--- a/Engine/Source/Programs/Horde/Plugins/Build/HordeServer.Build/Streams/StreamCollection.cs
+++ b/Engine/Source/Programs/Horde/Plugins/Build/HordeServer.Build/Streams/StreamCollection.cs
@@ -10,6 +10,7 @@ using EpicGames.Horde.Streams;
 using EpicGames.Horde.Users;
 using HordeServer.Acls;
 using HordeServer.Commits;
+using HordeServer.Jobs;
 using HordeServer.Jobs.Schedules;
 using HordeServer.Server;
 using HordeServer.Utilities;
@@ -115,9 +116,18 @@ namespace HordeServer.Streams
 			}
 
 			public TemplateId Id => _config.Id;
+			public ITemplate Template => throw new NotImplementedException();
 			public TemplateRefConfig Config => _config;
 			public ISchedule? Schedule => _schedule;
 			public IReadOnlyList<ITemplateStep> StepStates => (IReadOnlyList<ITemplateStep>?)_document.StepStates ?? Array.Empty<ITemplateStep>();
+
+			bool ITemplateRef.ShowUgsBadges => _config.ShowUgsBadges;
+			bool ITemplateRef.ShowUgsAlerts => _config.ShowUgsAlerts;
+			string? ITemplateRef.NotificationChannel => _config.NotificationChannel;
+			string? ITemplateRef.NotificationChannelFilter => _config.NotificationChannelFilter;
+			string? ITemplateRef.TriageChannel => _config.TriageChannel;
+			IReadOnlyList<IChainedJobTemplate>? ITemplateRef.ChainedJobs => _config.ChainedJobs;
+			IReadOnlyList<IChangeQuery>? ITemplateRef.DefaultChange => _config.DefaultChange;
 		}
 
 		class Schedule : ISchedule
diff --git a/Engine/Source/Programs/Horde/Plugins/Build/HordeServer.Build/Streams/StreamConfig.cs b/Engine/Source/Programs/Horde/Plugins/Build/HordeServer.Build/Streams/StreamConfig.cs
--- a/Engine/Source/Programs/Horde/Plugins/Build/HordeServer.Build/Streams/StreamConfig.cs
+++ b/Engine/Source/Programs/Horde/Plugins/Build/HordeServer.Build/Streams/StreamConfig.cs
@@ -743,23 +743,17 @@ namespace HordeServer.Streams
 	/// <summary>
 	/// Trigger for another template
 	/// </summary>
-	public class ChainedJobTemplateConfig
+	public class ChainedJobTemplateConfig : IChainedJobTemplate
 	{
-		/// <summary>
-		/// Name of the target that needs to complete before starting the other template
-		/// </summary>
+		/// <inheritdoc/>
 		[Required]
 		public string Trigger { get; set; } = String.Empty;
 
-		/// <summary>
-		/// Id of the template to trigger
-		/// </summary>
+		/// <inheritdoc/>
 		[Required]
 		public TemplateId TemplateId { get; set; }
 
-		/// <summary>
-		/// Whether to use the default change for the template rather than the change for the parent job.
-		/// </summary>
+		/// <inheritdoc/>
 		public bool UseDefaultChangeForTemplate { get; set; }
 	}
 
diff --git a/Engine/Source/Programs/Horde/Plugins/Build/HordeServer.Build/Streams/StreamMessages.cs b/Engine/Source/Programs/Horde/Plugins/Build/HordeServer.Build/Streams/StreamMessages.cs
--- a/Engine/Source/Programs/Horde/Plugins/Build/HordeServer.Build/Streams/StreamMessages.cs
+++ b/Engine/Source/Programs/Horde/Plugins/Build/HordeServer.Build/Streams/StreamMessages.cs
@@ -465,7 +465,7 @@ namespace HordeServer.Streams
 		/// <summary>
 		/// List of templates to trigger
 		/// </summary>
-		public List<ChainedJobTemplateConfig>? ChainedJobs { get; set; }
+		public List<GetChainedJobTemplateResponse>? ChainedJobs { get; set; }
 
 		/// <summary>
 		/// List of step states
@@ -496,18 +496,51 @@ namespace HordeServer.Streams
 		{
 			Id = id.ToString();
 			Hash = template.Hash.ToString();
-			ShowUgsBadges = templateRef.Config.ShowUgsBadges;
-			ShowUgsAlerts = templateRef.Config.ShowUgsAlerts;
-			NotificationChannel = templateRef.Config.NotificationChannel;
-			NotificationChannelFilter = templateRef.Config.NotificationChannelFilter;
+			ShowUgsBadges = templateRef.ShowUgsBadges;
+			ShowUgsAlerts = templateRef.ShowUgsAlerts;
+			NotificationChannel = templateRef.NotificationChannel;
+			NotificationChannelFilter = templateRef.NotificationChannelFilter;
 			Schedule = (templateRef.Schedule != null) ? new GetScheduleResponse(templateRef.Schedule, schedulerTimeZone) : null;
-			ChainedJobs = (templateRef.Config.ChainedJobs != null && templateRef.Config.ChainedJobs.Count > 0) ? templateRef.Config.ChainedJobs : null;
+			ChainedJobs = (templateRef.ChainedJobs != null && templateRef.ChainedJobs.Count > 0) ? templateRef.ChainedJobs.ConvertAll(x => new GetChainedJobTemplateResponse(x)) : null;
 			StepStates = stepStates;
-			DefaultChange = templateRef.Config.DefaultChange?.ConvertAll(x => new ChangeQueryMessage(x));
+			DefaultChange = templateRef.DefaultChange?.ConvertAll(x => new ChangeQueryMessage(x));
 			CanRun = canRun;
 		}
 	}
 
+	/// <summary>
+	/// Trigger for another template
+	/// </summary>
+	public class GetChainedJobTemplateResponse : IChainedJobTemplate
+	{
+		/// <inheritdoc/>
+		public string Trigger { get; set; }
+
+		/// <inheritdoc/>
+		public TemplateId TemplateId { get; set; }
+
+		/// <inheritdoc/>
+		public bool UseDefaultChangeForTemplate { get; set; }
+
+		/// <summary>
+		/// Constructor
+		/// </summary>
+		public GetChainedJobTemplateResponse()
+		{
+			Trigger = String.Empty;
+		}
+
+		/// <summary>
+		/// Constructor
+		/// </summary>
+		public GetChainedJobTemplateResponse(IChainedJobTemplate other)
+		{
+			Trigger = other.Trigger;
+			TemplateId = other.TemplateId;
+			UseDefaultChangeForTemplate = other.UseDefaultChangeForTemplate;
+		}
+	}
+
 	/// <summary>
 	/// Step state update request
 	/// </summary>
diff --git a/Engine/Source/Programs/Horde/Plugins/Build/HordeServer.Build/Streams/StreamsController.cs b/Engine/Source/Programs/Horde/Plugins/Build/HordeServer.Build/Streams/StreamsController.cs
--- a/Engine/Source/Programs/Horde/Plugins/Build/HordeServer.Build/Streams/StreamsController.cs
+++ b/Engine/Source/Programs/Horde/Plugins/Build/HordeServer.Build/Streams/StreamsController.cs
@@ -71,7 +71,7 @@ namespace HordeServer.Streams
 		{
 			ProjectId[] projectIdValues = Array.ConvertAll(projectIds, x => new ProjectId(x));
 
-			List<StreamId> streamIds = new List<StreamId>();
+			List<StreamConfig> streamConfigs = new List<StreamConfig>();
 			foreach (ProjectConfig projectConfig in _buildConfig.Value.Projects)
 			{
 				if (projectIdValues.Length == 0 || projectIdValues.Contains(projectConfig.Id))
@@ -80,18 +80,19 @@ namespace HordeServer.Streams
 					{
 						if (streamConfig.Authorize(StreamAclAction.ViewStream, User))
 						{
-							streamIds.Add(streamConfig.Id);
+							streamConfigs.Add(streamConfig);
 						}
 					}
 				}
 			}
 
-			IReadOnlyList<IStream> streams = await _streamCollection.FindAsync(streamIds, cancellationToken);
+			IReadOnlyList<IStream> streams = await _streamCollection.FindAsync(streamConfigs.ConvertAll(x => x.Id), cancellationToken);
 
 			List<GetStreamResponse> responses = new List<GetStreamResponse>();
 			foreach (IStream stream in streams)
 			{
-				GetStreamResponse response = await CreateGetStreamResponseAsync(stream, cancellationToken);
+				StreamConfig streamConfig = streamConfigs.First(x => x.Id == stream.Id);
+				GetStreamResponse response = await CreateGetStreamResponseAsync(stream, streamConfig, cancellationToken);
 				responses.Add(response);
 			}
 
@@ -126,16 +127,17 @@ namespace HordeServer.Streams
 				return NotFound(streamId);
 			}
 
-			return PropertyFilter.Apply(await CreateGetStreamResponseAsync(stream, cancellationToken), filter);
+			return PropertyFilter.Apply(await CreateGetStreamResponseAsync(stream, streamConfig, cancellationToken), filter);
 		}
 
 		/// <summary>
 		/// Create a stream response object, including all the templates
 		/// </summary>
 		/// <param name="stream">Stream to create response for</param>
+		/// <param name="streamConfig">Config object for this stream</param>
 		/// <param name="cancellationToken">Cancellation token for the operation</param>
 		/// <returns>Response object</returns>
-		async Task<GetStreamResponse> CreateGetStreamResponseAsync(IStream stream, CancellationToken cancellationToken)
+		async Task<GetStreamResponse> CreateGetStreamResponseAsync(IStream stream, StreamConfig streamConfig, CancellationToken cancellationToken)
 		{
 			using TelemetrySpan span = _tracer.StartActiveSpan($"{nameof(StreamsController)}.{nameof(CreateGetStreamResponseAsync)}");
 			span.SetAttribute("streamId", stream.Id);
@@ -143,11 +145,17 @@ namespace HordeServer.Streams
 			List<GetTemplateRefResponse> apiTemplateRefs = new List<GetTemplateRefResponse>();
 			foreach (KeyValuePair<TemplateId, ITemplateRef> pair in stream.Templates)
 			{
+				TemplateRefConfig? templateRefConfig;
+				if (!streamConfig.TryGetTemplate(pair.Key, out templateRefConfig))
+				{
+					continue;
+				}
+
 				using TelemetrySpan templateScope = _tracer.StartActiveSpan($"{nameof(StreamsController)}.{nameof(CreateGetStreamResponseAsync)}.Template");
-				templateScope.SetAttribute("templateName", pair.Value.Config.Name);
+				templateScope.SetAttribute("templateName", templateRefConfig.Name);
 
 				ITemplateRef templateRef = pair.Value;
-				if (templateRef.Config.Authorize(StreamAclAction.ViewTemplate, User))
+				if (templateRefConfig.Authorize(StreamAclAction.ViewTemplate, User))
 				{
 					List<GetTemplateStepStateResponse>? stepStates = null;
 					if (templateRef.StepStates != null)
@@ -163,9 +171,9 @@ namespace HordeServer.Streams
 						}
 					}
 
-					bool canRun = templateRef.Config.Authorize(JobAclAction.CreateJob, User);
+					bool canRun = templateRefConfig.Authorize(JobAclAction.CreateJob, User);
 
-					ITemplate? template = await _templateCollection.GetOrAddAsync(pair.Value.Config);
+					ITemplate? template = await _templateCollection.GetOrAddAsync(templateRefConfig);
 					apiTemplateRefs.Add(new GetTemplateRefResponse(pair.Key, pair.Value, template, stepStates, _timeZone, canRun));
 				}
 			}
diff --git a/Engine/Source/Programs/Shared/EpicGames.Horde/Jobs/Templates/ITemplate.cs b/Engine/Source/Programs/Shared/EpicGames.Horde/Jobs/Templates/ITemplate.cs
--- a/Engine/Source/Programs/Shared/EpicGames.Horde/Jobs/Templates/ITemplate.cs
+++ b/Engine/Source/Programs/Shared/EpicGames.Horde/Jobs/Templates/ITemplate.cs
@@ -18,62 +18,62 @@ namespace EpicGames.Horde.Jobs.Templates
 		/// <summary>
 		/// Hash of this template
 		/// </summary>
-		public ContentHash Hash { get; }
+		ContentHash Hash { get; }
 
 		/// <summary>
 		/// Name of the template.
 		/// </summary>
-		public string Name { get; }
+		string Name { get; }
 
 		/// <summary>
 		/// Description for the template
 		/// </summary>
-		public string? Description { get; }
+		string? Description { get; }
 
 		/// <summary>
 		/// Priority of this job
 		/// </summary>
-		public Priority? Priority { get; }
+		Priority? Priority { get; }
 
 		/// <summary>
 		/// Whether to allow preflights for this job type
 		/// </summary>
-		public bool AllowPreflights { get; }
+		bool AllowPreflights { get; }
 
 		/// <summary>
 		/// Whether to always issues for jobs using this template
 		/// </summary>
-		public bool UpdateIssues { get; }
+		bool UpdateIssues { get; }
 
 		/// <summary>
 		/// Whether to promote issues by default for jobs using this template
 		/// </summary>
-		public bool PromoteIssuesByDefault { get; }
+		bool PromoteIssuesByDefault { get; }
 
 		/// <summary>
 		/// Agent type to use for parsing the job state
 		/// </summary>
-		public string? InitialAgentType { get; }
+		string? InitialAgentType { get; }
 
 		/// <summary>
 		/// Path to a file within the stream to submit to generate a new changelist for jobs
 		/// </summary>
-		public string? SubmitNewChange { get; }
+		string? SubmitNewChange { get; }
 
 		/// <summary>
 		/// Description for new changelists
 		/// </summary>
-		public string? SubmitDescription { get; }
+		string? SubmitDescription { get; }
 
 		/// <summary>
 		/// Optional predefined user-defined properties for this job
 		/// </summary>
-		public IReadOnlyList<string> Arguments { get; }
+		IReadOnlyList<string> Arguments { get; }
 
 		/// <summary>
 		/// Parameters for this template
 		/// </summary>
-		public IReadOnlyList<ITemplateParameter> Parameters { get; }
+		IReadOnlyList<ITemplateParameter> Parameters { get; }
 	}
 
 	/// <summary>
