Horde: Fix links in documentation.

#rnx

diff --git a/Engine/Source/Programs/Horde/Docs/Config/BuildAutomation.md b/Engine/Source/Programs/Horde/Docs/Config/BuildAutomation.md
--- a/Engine/Source/Programs/Horde/Docs/Config/BuildAutomation.md
+++ b/Engine/Source/Programs/Horde/Docs/Config/BuildAutomation.md
@@ -119,7 +119,7 @@ Allows entering arbitrary text for an argument, with an optional regex for valid
         "toolTip": "Tooltip for text parameter"
     }
 
-See [TextParameterData](Schema/Streams.md#textparameterdata) for valid properties.
+See [TemplateTextParameterConfig](Schema/Streams.md#templatetextparameterconfig) for valid properties.
 
 ### List Parameters
 
@@ -148,7 +148,7 @@ Allows the user to select one or more options from a predefined list.
         ]
     }
 
-See [ListParameterData](Schema/Streams.md#listparameterdata) for valid properties.
+See [TemplateListParameterConfig](Schema/Streams.md#templatelistparameterconfig) for valid properties.
 
 ### Multi-List Parameters
 
@@ -189,7 +189,7 @@ to group items underneath.
         ]
     },
 
-See [ListParameterData](Schema/Streams.md#listparameterdata) for valid properties.
+See [TemplateListParameterConfig](Schema/Streams.md#templatelistparameterconfig) for valid properties.
 
 ### Tag-Picker Parameters
 
@@ -239,7 +239,7 @@ Allows toggling whether to enable an option or not.
         "argumentIfDisabled": "-set:BoolParameter=false"
     }
 
-See [BoolParameterData](Schema/Streams.md#boolparameterdata) for valid properties.
+See [TemplateBoolParameterConfig](Schema/Streams.md#templateboolparameterconfig) for valid properties.
 
 ## Artifacts
 
