Horde: Move some internal server implementation details out of IJob.

#rnx

diff --git a/Engine/Source/Programs/Horde/Plugins/Build/HordeServer.Build/BuildConfig.cs b/Engine/Source/Programs/Horde/Plugins/Build/HordeServer.Build/BuildConfig.cs
--- a/Engine/Source/Programs/Horde/Plugins/Build/HordeServer.Build/BuildConfig.cs
+++ b/Engine/Source/Programs/Horde/Plugins/Build/HordeServer.Build/BuildConfig.cs
@@ -19,6 +19,7 @@ using HordeServer.Agents.Pools;
 using HordeServer.Artifacts;
 using HordeServer.Commits;
 using HordeServer.Devices;
+using HordeServer.Jobs;
 using HordeServer.Perforce;
 using HordeServer.Plugins;
 using HordeServer.Projects;
@@ -181,6 +182,25 @@ namespace HordeServer
 			}
 		}
 
+		/// <summary>
+		/// Determines if the user is authorized to perform an action on a particular stream
+		/// </summary>
+		/// <param name="job">Job to authorize for</param>
+		/// <param name="action">The action being performed</param>
+		/// <param name="user">The principal to authorize</param>
+		/// <returns>True if the action is authorized</returns>
+		public bool Authorize(IJob? job, AclAction action, ClaimsPrincipal user)
+		{
+			if (job != null && TryGetStream(job.StreamId, out StreamConfig? streamConfig))
+			{
+				return streamConfig.Authorize(action, user);
+			}
+			else
+			{
+				return Authorize(action, user);
+			}
+		}
+
 		/// <summary>
 		/// Authorize an action to perform on an artifact
 		/// </summary>
diff --git a/Engine/Source/Programs/Horde/Plugins/Build/HordeServer.Build/Jobs/IJob.cs b/Engine/Source/Programs/Horde/Plugins/Build/HordeServer.Build/Jobs/IJob.cs
--- a/Engine/Source/Programs/Horde/Plugins/Build/HordeServer.Build/Jobs/IJob.cs
+++ b/Engine/Source/Programs/Horde/Plugins/Build/HordeServer.Build/Jobs/IJob.cs
@@ -17,7 +17,6 @@ using EpicGames.Horde.Logs;
 using EpicGames.Horde.Streams;
 using EpicGames.Horde.Ugs;
 using EpicGames.Horde.Users;
-using Horde.Common.Rpc;
 using HordeServer.Acls;
 using HordeServer.Jobs.Timing;
 using HordeServer.Streams;
@@ -397,26 +396,6 @@ namespace HordeServer.Jobs
 	/// </summary>
 	public static class JobExtensions
 	{
-		/// <summary>
-		/// Determines if the user is authorized to perform an action on a particular stream
-		/// </summary>
-		/// <param name="buildConfig">Current global config instance</param>
-		/// <param name="job">Job to authorize for</param>
-		/// <param name="action">The action being performed</param>
-		/// <param name="user">The principal to authorize</param>
-		/// <returns>True if the action is authorized</returns>
-		public static bool Authorize(this BuildConfig buildConfig, IJob? job, AclAction action, ClaimsPrincipal user)
-		{
-			if (job != null && buildConfig.TryGetStream(job.StreamId, out StreamConfig? streamConfig))
-			{
-				return streamConfig.Authorize(action, user);
-			}
-			else
-			{
-				return buildConfig.Authorize(action, user);
-			}
-		}
-
 		/// <summary>
 		/// Gets the current job state
 		/// </summary>
@@ -1166,22 +1145,6 @@ namespace HordeServer.Jobs
 			}
 		}
 
-		/// <summary>
-		/// Creates an RPC response object
-		/// </summary>
-		/// <param name="job">The job document</param>
-		/// <returns></returns>
-		public static RpcGetJobResponse ToRpcResponse(this IJob job)
-		{
-			RpcGetJobResponse response = new RpcGetJobResponse();
-			response.StreamId = job.StreamId.ToString();
-			response.Change = job.CommitId.GetPerforceChange();
-			response.CodeChange = job.CodeCommitId?.GetPerforceChange() ?? 0;
-			response.PreflightChange = job.PreflightCommitId?.GetPerforceChange() ?? 0;
-			response.Arguments.Add(job.Arguments);
-			return response;
-		}
-
 		/// <summary>
 		/// Gets a key attached to all artifacts produced for a job
 		/// </summary>
diff --git a/Engine/Source/Programs/Horde/Plugins/Build/HordeServer.Build/Jobs/JobRpcService.cs b/Engine/Source/Programs/Horde/Plugins/Build/HordeServer.Build/Jobs/JobRpcService.cs
--- a/Engine/Source/Programs/Horde/Plugins/Build/HordeServer.Build/Jobs/JobRpcService.cs
+++ b/Engine/Source/Programs/Horde/Plugins/Build/HordeServer.Build/Jobs/JobRpcService.cs
@@ -162,7 +162,13 @@ namespace HordeServer.Jobs
 				throw new StructuredRpcException(StatusCode.PermissionDenied, "Not authenticated to access job {JobId}", request.JobId);
 			}
 
-			return job.ToRpcResponse();
+			RpcGetJobResponse response = new RpcGetJobResponse();
+			response.StreamId = job.StreamId.ToString();
+			response.Change = job.CommitId.GetPerforceChange();
+			response.CodeChange = job.CodeCommitId?.GetPerforceChange() ?? 0;
+			response.PreflightChange = job.PreflightCommitId?.GetPerforceChange() ?? 0;
+			response.Arguments.Add(job.Arguments);
+			return response;
 		}
 
 		/// <summary>
