Horde: Add an identifier type for notification triggers.

#rnx

diff --git a/Engine/Source/Programs/Horde/HordeServer.Shared/Acls/AclConfig.cs b/Engine/Source/Programs/Horde/HordeServer.Shared/Acls/AclConfig.cs
--- a/Engine/Source/Programs/Horde/HordeServer.Shared/Acls/AclConfig.cs
+++ b/Engine/Source/Programs/Horde/HordeServer.Shared/Acls/AclConfig.cs
@@ -417,7 +417,7 @@ namespace HordeServer.Acls
 	/// <summary>
 	/// New claim to create
 	/// </summary>
-	public class AclClaimConfig
+	public class AclClaimConfig : IAclClaim
 	{
 		/// <summary>
 		/// The claim type
diff --git a/Engine/Source/Programs/Horde/Plugins/Build/HordeServer.Build/Jobs/IJob.cs b/Engine/Source/Programs/Horde/Plugins/Build/HordeServer.Build/Jobs/IJob.cs
--- a/Engine/Source/Programs/Horde/Plugins/Build/HordeServer.Build/Jobs/IJob.cs
+++ b/Engine/Source/Programs/Horde/Plugins/Build/HordeServer.Build/Jobs/IJob.cs
@@ -14,13 +14,11 @@ using EpicGames.Horde.Jobs.Bisect;
 using EpicGames.Horde.Jobs.Graphs;
 using EpicGames.Horde.Jobs.Templates;
 using EpicGames.Horde.Logs;
+using EpicGames.Horde.Notifications;
 using EpicGames.Horde.Streams;
 using EpicGames.Horde.Ugs;
 using EpicGames.Horde.Users;
-using HordeServer.Acls;
 using HordeServer.Jobs.Timing;
-using HordeServer.Streams;
-//using HordeServer.Ugs;
 using Microsoft.Extensions.Logging;
 using MongoDB.Bson;
 using MongoDB.Driver;
@@ -162,7 +160,7 @@ namespace HordeServer.Jobs
 		/// <summary>
 		/// Claims inherited from the user that started this job
 		/// </summary>
-		public IReadOnlyList<AclClaimConfig> Claims { get; }
+		public IReadOnlyList<IAclClaim> Claims { get; }
 
 		/// <summary>
 		/// Array of jobstep runs
@@ -202,7 +200,7 @@ namespace HordeServer.Jobs
 		/// <summary>
 		/// Unique id for notifications
 		/// </summary>
-		public ObjectId? NotificationTriggerId { get; }
+		public NotificationTriggerId? NotificationTriggerId { get; }
 
 		/// <summary>
 		/// Whether to show badges in UGS for this job
@@ -227,7 +225,7 @@ namespace HordeServer.Jobs
 		/// <summary>
 		/// Mapping of label ids to notification trigger ids for notifications
 		/// </summary>
-		public IReadOnlyDictionary<int, ObjectId> LabelIdxToTriggerId { get; }
+		public IReadOnlyDictionary<int, NotificationTriggerId> LabelIdxToTriggerId { get; }
 
 		/// <summary>
 		/// List of reports for this step
@@ -301,7 +299,7 @@ namespace HordeServer.Jobs
 		/// <param name="jobTrigger">New downstream job id</param>
 		/// <param name="cancellationReason">Optional reason why the job was canceled</param>
 		/// <param name="cancellationToken">Cancellation token for the operation</param>
-		Task<IJob?> TryUpdateJobAsync(string? name = null, Priority? priority = null, bool? autoSubmit = null, int? autoSubmitChange = null, string? autoSubmitMessage = null, UserId? abortedByUserId = null, ObjectId? notificationTriggerId = null, List<JobReport>? reports = null, List<string>? arguments = null, KeyValuePair<int, ObjectId>? labelIdxToTriggerId = null, KeyValuePair<TemplateId, JobId>? jobTrigger = null, string? cancellationReason = null, CancellationToken cancellationToken = default);
+		Task<IJob?> TryUpdateJobAsync(string? name = null, Priority? priority = null, bool? autoSubmit = null, int? autoSubmitChange = null, string? autoSubmitMessage = null, UserId? abortedByUserId = null, NotificationTriggerId? notificationTriggerId = null, List<JobReport>? reports = null, List<string>? arguments = null, KeyValuePair<int, NotificationTriggerId>? labelIdxToTriggerId = null, KeyValuePair<TemplateId, JobId>? jobTrigger = null, string? cancellationReason = null, CancellationToken cancellationToken = default);
 
 		/// <summary>
 		/// Updates the state of a batch
@@ -333,7 +331,7 @@ namespace HordeServer.Jobs
 		/// <param name="newCancellationReason">The reason the job step was canceled</param>
 		/// <param name="cancellationToken">Cancellation token for the operation</param>
 		/// <returns>True if the job was updated, false if it was deleted in the meantime</returns>
-		Task<IJob?> TryUpdateStepAsync(JobStepBatchId batchId, JobStepId stepId, JobStepState newState = default, JobStepOutcome newOutcome = default, JobStepError? newError = null, bool? newAbortRequested = null, UserId? newAbortByUserId = null, LogId? newLogId = null, ObjectId? newNotificationTriggerId = null, UserId? newRetryByUserId = null, Priority? newPriority = null, List<JobReport>? newReports = null, Dictionary<string, string?>? newProperties = null, string? newCancellationReason = null, CancellationToken cancellationToken = default);
+		Task<IJob?> TryUpdateStepAsync(JobStepBatchId batchId, JobStepId stepId, JobStepState newState = default, JobStepOutcome newOutcome = default, JobStepError? newError = null, bool? newAbortRequested = null, UserId? newAbortByUserId = null, LogId? newLogId = null, NotificationTriggerId? newNotificationTriggerId = null, UserId? newRetryByUserId = null, Priority? newPriority = null, List<JobReport>? newReports = null, Dictionary<string, string?>? newProperties = null, string? newCancellationReason = null, CancellationToken cancellationToken = default);
 
 		/// <summary>
 		/// Attempts to update the node groups to be executed for a job. Fails if another write happens in the meantime.
@@ -1504,7 +1502,7 @@ namespace HordeServer.Jobs
 		/// <summary>
 		/// Unique id for notifications
 		/// </summary>
-		public ObjectId? NotificationTriggerId { get; }
+		public NotificationTriggerId? NotificationTriggerId { get; }
 
 		/// <summary>
 		/// Time at which the batch transitioned to the ready state (UTC).
diff --git a/Engine/Source/Programs/Horde/Plugins/Build/HordeServer.Build/Jobs/JobCollection.cs b/Engine/Source/Programs/Horde/Plugins/Build/HordeServer.Build/Jobs/JobCollection.cs
--- a/Engine/Source/Programs/Horde/Plugins/Build/HordeServer.Build/Jobs/JobCollection.cs
+++ b/Engine/Source/Programs/Horde/Plugins/Build/HordeServer.Build/Jobs/JobCollection.cs
@@ -5,6 +5,7 @@ using System.Diagnostics.CodeAnalysis;
 using System.Runtime.CompilerServices;
 using System.Text;
 using EpicGames.Core;
+using EpicGames.Horde.Acls;
 using EpicGames.Horde.Agents;
 using EpicGames.Horde.Agents.Leases;
 using EpicGames.Horde.Agents.Pools;
@@ -15,6 +16,7 @@ using EpicGames.Horde.Jobs.Bisect;
 using EpicGames.Horde.Jobs.Graphs;
 using EpicGames.Horde.Jobs.Templates;
 using EpicGames.Horde.Logs;
+using EpicGames.Horde.Notifications;
 using EpicGames.Horde.Streams;
 using EpicGames.Horde.Users;
 using HordeServer.Acls;
@@ -76,7 +78,7 @@ namespace HordeServer.Jobs
 			bool IJob.PromoteIssuesByDefault => _document.PromoteIssuesByDefault;
 			DateTime IJob.CreateTimeUtc => _document.GetCreateTimeOrDefault();
 			JobOptions? IJob.JobOptions => _document.JobOptions;
-			IReadOnlyList<AclClaimConfig> IJob.Claims => _document.Claims;
+			IReadOnlyList<IAclClaim> IJob.Claims => _document.Claims;
 			IReadOnlyList<IJobStepBatch> IJob.Batches => Batches;
 			IReadOnlyDictionary<ParameterId, string> IJob.Parameters => _document.Parameters;
 			IReadOnlyList<string> IJob.Arguments => _document.Arguments;
@@ -84,12 +86,12 @@ namespace HordeServer.Jobs
 			IReadOnlyList<string> IJob.AdditionalArguments => _document.AdditionalArguments;
 			IReadOnlyDictionary<string, string> IJob.Environment => _document.Environment;
 			IReadOnlyList<int> IJob.Issues => _document.ReferencedByIssues;
-			ObjectId? IJob.NotificationTriggerId => _document.NotificationTriggerId;
+			NotificationTriggerId? IJob.NotificationTriggerId => _document.NotificationTriggerId;
 			bool IJob.ShowUgsBadges => _document.ShowUgsBadges;
 			bool IJob.ShowUgsAlerts => _document.ShowUgsAlerts;
 			string? IJob.NotificationChannel => _document.NotificationChannel;
 			string? IJob.NotificationChannelFilter => _document.NotificationChannelFilter;
-			IReadOnlyDictionary<int, ObjectId> IJob.LabelIdxToTriggerId => _document._labelNotifications.ToDictionary(x => x._labelIdx, x => x._triggerId);
+			IReadOnlyDictionary<int, NotificationTriggerId> IJob.LabelIdxToTriggerId => _document._labelNotifications.ToDictionary(x => x._labelIdx, x => x._triggerId);
 			IReadOnlyList<IJobReport>? IJob.Reports => _document.Reports;
 			IReadOnlyList<IChainedJob> IJob.ChainedJobs => _document.ChainedJobs;
 			DateTime IJob.UpdateTimeUtc => _document.UpdateTimeUtc ?? _document.UpdateTime?.UtcDateTime ?? DateTime.UnixEpoch;
@@ -220,13 +222,13 @@ namespace HordeServer.Jobs
 				return _collection.CreateJobObject(newDocument, newGraph);
 			}
 
-			public async Task<IJob?> TryUpdateJobAsync(string? name = null, Priority? priority = null, bool? autoSubmit = null, int? autoSubmitChange = null, string? autoSubmitMessage = null, UserId? abortedByUserId = null, ObjectId? notificationTriggerId = null, List<JobReport>? reports = null, List<string>? arguments = null, KeyValuePair<int, ObjectId>? labelIdxToTriggerId = null, KeyValuePair<TemplateId, JobId>? jobTrigger = null, string? cancellationReason = null, CancellationToken cancellationToken = default)
+			public async Task<IJob?> TryUpdateJobAsync(string? name = null, Priority? priority = null, bool? autoSubmit = null, int? autoSubmitChange = null, string? autoSubmitMessage = null, UserId? abortedByUserId = null, NotificationTriggerId? notificationTriggerId = null, List<JobReport>? reports = null, List<string>? arguments = null, KeyValuePair<int, NotificationTriggerId>? labelIdxToTriggerId = null, KeyValuePair<TemplateId, JobId>? jobTrigger = null, string? cancellationReason = null, CancellationToken cancellationToken = default)
 			{
 				JobDocument? newDocument = await _collection.TryUpdateJobAsync(Document, Graph, name, priority, autoSubmit, autoSubmitChange, autoSubmitMessage, abortedByUserId, notificationTriggerId, reports, arguments, labelIdxToTriggerId, jobTrigger, cancellationReason, cancellationToken);
 				return _collection.CreateJobObject(newDocument, Graph);
 			}
 
-			public async Task<IJob?> TryUpdateStepAsync(JobStepBatchId batchId, JobStepId stepId, JobStepState newState = JobStepState.Unspecified, JobStepOutcome newOutcome = JobStepOutcome.Unspecified, JobStepError? newError = null, bool? newAbortRequested = null, UserId? newAbortByUserId = null, LogId? newLogId = null, ObjectId? newNotificationTriggerId = null, UserId? newRetryByUserId = null, Priority? newPriority = null, List<JobReport>? newReports = null, Dictionary<string, string?>? newProperties = null, string? newCancellationReason = null, CancellationToken cancellationToken = default)
+			public async Task<IJob?> TryUpdateStepAsync(JobStepBatchId batchId, JobStepId stepId, JobStepState newState = JobStepState.Unspecified, JobStepOutcome newOutcome = JobStepOutcome.Unspecified, JobStepError? newError = null, bool? newAbortRequested = null, UserId? newAbortByUserId = null, LogId? newLogId = null, NotificationTriggerId? newNotificationTriggerId = null, UserId? newRetryByUserId = null, Priority? newPriority = null, List<JobReport>? newReports = null, Dictionary<string, string?>? newProperties = null, string? newCancellationReason = null, CancellationToken cancellationToken = default)
 			{
 				JobDocument? newDocument = await _collection.TryUpdateStepAsync(_document, _graph, batchId, stepId, newState, newOutcome, newError, newAbortRequested, newAbortByUserId, newLogId, newNotificationTriggerId, newRetryByUserId, newPriority, newReports, newProperties, newCancellationReason, cancellationToken);
 				return _collection.CreateJobObject(newDocument, Graph);
@@ -310,7 +312,7 @@ namespace HordeServer.Jobs
 			JobStepOutcome IJobStep.Outcome => Document.Outcome;
 			JobStepError IJobStep.Error => Document.Error;
 			LogId? IJobStep.LogId => Document.LogId;
-			ObjectId? IJobStep.NotificationTriggerId => Document.NotificationTriggerId;
+			NotificationTriggerId? IJobStep.NotificationTriggerId => Document.NotificationTriggerId;
 			DateTime? IJobStep.ReadyTimeUtc => Document.ReadyTime?.UtcDateTime;
 			DateTime? IJobStep.StartTimeUtc => Document.StartTime?.UtcDateTime;
 			DateTime? IJobStep.FinishTimeUtc => Document.FinishTime?.UtcDateTime;
@@ -500,7 +502,7 @@ namespace HordeServer.Jobs
 			public Dictionary<string, string> Environment { get; set; } = new Dictionary<string, string>();
 
 			public List<int> ReferencedByIssues { get; set; } = new List<int>();
-			public ObjectId? NotificationTriggerId { get; set; }
+			public NotificationTriggerId? NotificationTriggerId { get; set; }
 			public bool ShowUgsBadges { get; set; }
 			public bool ShowUgsAlerts { get; set; }
 			public string? NotificationChannel { get; set; }
@@ -656,7 +658,7 @@ namespace HordeServer.Jobs
 			public LogId? LogId { get; set; }
 
 			[BsonIgnoreIfNull]
-			public ObjectId? NotificationTriggerId { get; set; }
+			public NotificationTriggerId? NotificationTriggerId { get; set; }
 
 			[BsonIgnoreIfNull]
 			public DateTimeOffset? ReadyTime { get; set; }
@@ -745,7 +747,7 @@ namespace HordeServer.Jobs
 		class LabelNotificationDocument
 		{
 			public int _labelIdx;
-			public ObjectId _triggerId;
+			public NotificationTriggerId _triggerId;
 		}
 
 		/// <summary>
@@ -1167,7 +1169,7 @@ namespace HordeServer.Jobs
 		}
 
 		/// <inheritdoc/>
-		async Task<JobDocument?> TryUpdateJobAsync(JobDocument jobDocument, IGraph graph, string? name, Priority? priority, bool? autoSubmit, int? autoSubmitChange, string? autoSubmitMessage, UserId? abortedByUserId, ObjectId? notificationTriggerId, List<JobReport>? reports, List<string>? arguments, KeyValuePair<int, ObjectId>? labelIdxToTriggerId, KeyValuePair<TemplateId, JobId>? jobTrigger, string? cancellationReason, CancellationToken cancellationToken)
+		async Task<JobDocument?> TryUpdateJobAsync(JobDocument jobDocument, IGraph graph, string? name, Priority? priority, bool? autoSubmit, int? autoSubmitChange, string? autoSubmitMessage, UserId? abortedByUserId, NotificationTriggerId? notificationTriggerId, List<JobReport>? reports, List<string>? arguments, KeyValuePair<int, NotificationTriggerId>? labelIdxToTriggerId, KeyValuePair<TemplateId, JobId>? jobTrigger, string? cancellationReason, CancellationToken cancellationToken)
 		{
 			// Create the update 
 			UpdateDefinitionBuilder<JobDocument> updateBuilder = Builders<JobDocument>.Update;
@@ -1412,7 +1414,7 @@ namespace HordeServer.Jobs
 		}
 
 		/// <inheritdoc/>
-		async Task<JobDocument?> TryUpdateStepAsync(JobDocument jobDocument, IGraph graph, JobStepBatchId batchId, JobStepId stepId, JobStepState newState, JobStepOutcome newOutcome, JobStepError? newError, bool? newAbortRequested, UserId? newAbortByUserId, LogId? newLogId, ObjectId? newNotificationTriggerId, UserId? newRetryByUserId, Priority? newPriority, List<JobReport>? newReports, Dictionary<string, string?>? newProperties, string? newCancellationReason, CancellationToken cancellationToken)
+		async Task<JobDocument?> TryUpdateStepAsync(JobDocument jobDocument, IGraph graph, JobStepBatchId batchId, JobStepId stepId, JobStepState newState, JobStepOutcome newOutcome, JobStepError? newError, bool? newAbortRequested, UserId? newAbortByUserId, LogId? newLogId, NotificationTriggerId? newNotificationTriggerId, UserId? newRetryByUserId, Priority? newPriority, List<JobReport>? newReports, Dictionary<string, string?>? newProperties, string? newCancellationReason, CancellationToken cancellationToken)
 		{
 			jobDocument = Clone(jobDocument);
 
diff --git a/Engine/Source/Programs/Horde/Plugins/Build/HordeServer.Build/Jobs/JobService.cs b/Engine/Source/Programs/Horde/Plugins/Build/HordeServer.Build/Jobs/JobService.cs
--- a/Engine/Source/Programs/Horde/Plugins/Build/HordeServer.Build/Jobs/JobService.cs
+++ b/Engine/Source/Programs/Horde/Plugins/Build/HordeServer.Build/Jobs/JobService.cs
@@ -12,6 +12,7 @@ using EpicGames.Horde.Jobs;
 using EpicGames.Horde.Jobs.Graphs;
 using EpicGames.Horde.Jobs.Templates;
 using EpicGames.Horde.Logs;
+using EpicGames.Horde.Notifications;
 using EpicGames.Horde.Streams;
 using EpicGames.Horde.Users;
 using HordeServer.Agents;
@@ -298,7 +299,7 @@ namespace HordeServer.Jobs
 			}
 
 			// Remove all the triggers from it
-			List<ObjectId> triggerIds = new List<ObjectId>();
+			List<NotificationTriggerId> triggerIds = new List<NotificationTriggerId>();
 			if (job.NotificationTriggerId != null)
 			{
 				triggerIds.Add(job.NotificationTriggerId.Value);
@@ -340,7 +341,7 @@ namespace HordeServer.Jobs
 		/// <param name="labelIdxToTriggerId">New trigger ID for a label in the job</param>
 		/// <param name="cancellationReason">Reason the job was canceled</param>
 		/// <param name="cancellationToken">Cancellation token for the operation</param>
-		public async Task<IJob?> UpdateJobAsync(IJob job, string? name = null, Priority? priority = null, bool? autoSubmit = null, UserId? abortedByUserId = null, ObjectId? onCompleteTriggerId = null, List<JobReport>? reports = null, List<string>? arguments = null, KeyValuePair<int, ObjectId>? labelIdxToTriggerId = null, string? cancellationReason = null, CancellationToken cancellationToken = default)
+		public async Task<IJob?> UpdateJobAsync(IJob job, string? name = null, Priority? priority = null, bool? autoSubmit = null, UserId? abortedByUserId = null, NotificationTriggerId? onCompleteTriggerId = null, List<JobReport>? reports = null, List<string>? arguments = null, KeyValuePair<int, NotificationTriggerId>? labelIdxToTriggerId = null, string? cancellationReason = null, CancellationToken cancellationToken = default)
 		{
 			using TelemetrySpan span = _tracer.StartActiveSpan($"{nameof(JobService)}.{nameof(UpdateJobAsync)}");
 			span.SetAttribute("JobId", job.Id.ToString());
@@ -849,7 +850,7 @@ namespace HordeServer.Jobs
 		/// <param name="newCancellationReason">The reason the step was canceled</param>
 		/// <param name="cancellationToken">Cancellation token for the operation</param>
 		/// <returns>True if the job was updated, false if it was deleted in the meantime</returns>
-		public async Task<IJob?> UpdateStepAsync(IJob job, JobStepBatchId batchId, JobStepId stepId, StreamConfig streamConfig, JobStepState newState = JobStepState.Unspecified, JobStepOutcome newOutcome = JobStepOutcome.Unspecified, JobStepError? newError = null, bool? newAbortRequested = null, UserId? newAbortByUserId = null, LogId? newLogId = null, ObjectId? newNotificationTriggerId = null, UserId? newRetryByUserId = null, Priority? newPriority = null, List<JobReport>? newReports = null, Dictionary<string, string?>? newProperties = null, string? newCancellationReason = null, CancellationToken cancellationToken = default)
+		public async Task<IJob?> UpdateStepAsync(IJob job, JobStepBatchId batchId, JobStepId stepId, StreamConfig streamConfig, JobStepState newState = JobStepState.Unspecified, JobStepOutcome newOutcome = JobStepOutcome.Unspecified, JobStepError? newError = null, bool? newAbortRequested = null, UserId? newAbortByUserId = null, LogId? newLogId = null, NotificationTriggerId? newNotificationTriggerId = null, UserId? newRetryByUserId = null, Priority? newPriority = null, List<JobReport>? newReports = null, Dictionary<string, string?>? newProperties = null, string? newCancellationReason = null, CancellationToken cancellationToken = default)
 		{
 			using TelemetrySpan span = _tracer.StartActiveSpan($"{nameof(JobService)}.{nameof(UpdateStepAsync)}");
 			span.SetAttribute("Job", job.Id.ToString());
@@ -896,7 +897,7 @@ namespace HordeServer.Jobs
 		/// <param name="newCancellationReason">Property changes. Any properties with a null value will be removed.</param>
 		/// <param name="cancellationToken">Cancellation token for the operation</param>
 		/// <returns>True if the job was updated, false if it was deleted in the meantime</returns>
-		public async Task<IJob?> TryUpdateStepAsync(IJob job, JobStepBatchId batchId, JobStepId stepId, StreamConfig streamConfig, JobStepState newState = JobStepState.Unspecified, JobStepOutcome newOutcome = JobStepOutcome.Unspecified, JobStepError? newError = null, bool? newAbortRequested = null, UserId? newAbortByUserId = null, LogId? newLogId = null, ObjectId? newTriggerId = null, UserId? newRetryByUserId = null, Priority? newPriority = null, List<JobReport>? newReports = null, Dictionary<string, string?>? newProperties = null, string? newCancellationReason = null, CancellationToken cancellationToken = default)
+		public async Task<IJob?> TryUpdateStepAsync(IJob job, JobStepBatchId batchId, JobStepId stepId, StreamConfig streamConfig, JobStepState newState = JobStepState.Unspecified, JobStepOutcome newOutcome = JobStepOutcome.Unspecified, JobStepError? newError = null, bool? newAbortRequested = null, UserId? newAbortByUserId = null, LogId? newLogId = null, NotificationTriggerId? newTriggerId = null, UserId? newRetryByUserId = null, Priority? newPriority = null, List<JobReport>? newReports = null, Dictionary<string, string?>? newProperties = null, string? newCancellationReason = null, CancellationToken cancellationToken = default)
 		{
 			using TelemetrySpan span = _tracer.StartActiveSpan($"{nameof(JobService)}.{nameof(TryUpdateStepAsync)}");
 			span.SetAttribute("Job", job.Id.ToString());
diff --git a/Engine/Source/Programs/Horde/Plugins/Build/HordeServer.Build/Jobs/JobTaskSource.cs b/Engine/Source/Programs/Horde/Plugins/Build/HordeServer.Build/Jobs/JobTaskSource.cs
--- a/Engine/Source/Programs/Horde/Plugins/Build/HordeServer.Build/Jobs/JobTaskSource.cs
+++ b/Engine/Source/Programs/Horde/Plugins/Build/HordeServer.Build/Jobs/JobTaskSource.cs
@@ -890,7 +890,7 @@ namespace HordeServer.Jobs
 			claims.Add(new AclClaimConfig(HordeClaimTypes.LeaseStream, streamConfig.Id.ToString()));
 			claims.Add(new AclClaimConfig(HordeClaimTypes.LeaseProject, streamConfig.ProjectConfig.Id.ToString()));
 			claims.Add(new AclClaimConfig(HordeClaimTypes.LeaseTemplate, job.TemplateId.ToString()));
-			claims.AddRange(job.Claims);
+			claims.AddRange(job.Claims.ConvertAll(x => new AclClaimConfig(x.Type, x.Value)));
 
 			// Encode the payload
 			ExecuteJobTask task = new ExecuteJobTask();
diff --git a/Engine/Source/Programs/Horde/Plugins/Build/HordeServer.Build/Jobs/JobsController.cs b/Engine/Source/Programs/Horde/Plugins/Build/HordeServer.Build/Jobs/JobsController.cs
--- a/Engine/Source/Programs/Horde/Plugins/Build/HordeServer.Build/Jobs/JobsController.cs
+++ b/Engine/Source/Programs/Horde/Plugins/Build/HordeServer.Build/Jobs/JobsController.cs
@@ -7,6 +7,7 @@ using EpicGames.Horde.Jobs;
 using EpicGames.Horde.Jobs.Graphs;
 using EpicGames.Horde.Jobs.Templates;
 using EpicGames.Horde.Logs;
+using EpicGames.Horde.Notifications;
 using EpicGames.Horde.Streams;
 using EpicGames.Horde.Users;
 using HordeServer.Acls;
@@ -382,7 +383,7 @@ namespace HordeServer.Jobs
 				return Forbid(NotificationAclAction.CreateSubscription, jobId);
 			}
 
-			ObjectId triggerId = job.NotificationTriggerId ?? ObjectId.GenerateNewId();
+			NotificationTriggerId triggerId = job.NotificationTriggerId ?? new NotificationTriggerId(BinaryIdUtils.CreateNew());
 
 			job = await _jobService.UpdateJobAsync(job, null, null, null, null, triggerId, null, null, cancellationToken: cancellationToken);
 			if (job == null)
@@ -1713,10 +1714,10 @@ namespace HordeServer.Jobs
 				return NotFound(jobId, batchId, stepId);
 			}
 
-			ObjectId? triggerId = step.NotificationTriggerId;
+			NotificationTriggerId? triggerId = step.NotificationTriggerId;
 			if (triggerId == null)
 			{
-				triggerId = ObjectId.GenerateNewId();
+				triggerId = new NotificationTriggerId(BinaryIdUtils.CreateNew());
 				if (await _jobService.UpdateStepAsync(job, batchId, stepId, streamConfig, JobStepState.Unspecified, JobStepOutcome.Unspecified, newNotificationTriggerId: triggerId, cancellationToken: cancellationToken) == null)
 				{
 					return NotFound(jobId, batchId, stepId);
@@ -1788,7 +1789,7 @@ namespace HordeServer.Jobs
 		[Route("/api/v1/jobs/{jobId}/labels/{labelIndex}/notifications")]
 		public async Task<ActionResult> UpdateLabelNotificationsAsync(JobId jobId, int labelIndex, [FromBody] UpdateNotificationsRequest request, CancellationToken cancellationToken)
 		{
-			ObjectId triggerId;
+			NotificationTriggerId triggerId;
 			for (; ; )
 			{
 				IJob? job = await _jobService.GetJobAsync(jobId, cancellationToken);
@@ -1807,16 +1808,16 @@ namespace HordeServer.Jobs
 					return Forbid(NotificationAclAction.CreateSubscription, jobId);
 				}
 
-				ObjectId newTriggerId;
+				NotificationTriggerId newTriggerId;
 				if (job.LabelIdxToTriggerId.TryGetValue(labelIndex, out newTriggerId))
 				{
 					triggerId = newTriggerId;
 					break;
 				}
 
-				newTriggerId = ObjectId.GenerateNewId();
+				newTriggerId = new NotificationTriggerId(BinaryIdUtils.CreateNew());
 
-				IJob? newJob = await _jobService.UpdateJobAsync(job, labelIdxToTriggerId: new KeyValuePair<int, ObjectId>(labelIndex, newTriggerId), cancellationToken: cancellationToken);
+				IJob? newJob = await _jobService.UpdateJobAsync(job, labelIdxToTriggerId: new KeyValuePair<int, NotificationTriggerId>(labelIndex, newTriggerId), cancellationToken: cancellationToken);
 				if (newJob != null)
 				{
 					triggerId = newTriggerId;
diff --git a/Engine/Source/Programs/Horde/Plugins/Build/HordeServer.Build/Notifications/INotificationService.cs b/Engine/Source/Programs/Horde/Plugins/Build/HordeServer.Build/Notifications/INotificationService.cs
--- a/Engine/Source/Programs/Horde/Plugins/Build/HordeServer.Build/Notifications/INotificationService.cs
+++ b/Engine/Source/Programs/Horde/Plugins/Build/HordeServer.Build/Notifications/INotificationService.cs
@@ -4,6 +4,7 @@ using System.Diagnostics.CodeAnalysis;
 using System.Security.Claims;
 using EpicGames.Horde.Jobs;
 using EpicGames.Horde.Jobs.Graphs;
+using EpicGames.Horde.Notifications;
 using HordeServer.Agents;
 using HordeServer.Configuration;
 using HordeServer.Devices;
@@ -11,7 +12,6 @@ using HordeServer.Issues;
 using HordeServer.Jobs;
 using HordeServer.Streams;
 using HordeServer.Users;
-using MongoDB.Bson;
 
 namespace HordeServer.Notifications
 {
@@ -102,7 +102,7 @@ namespace HordeServer.Notifications
 		/// <param name="slack">Whether to receive Slack notifications</param>
 		/// <param name="cancellationToken">Cancellation token for the operation</param>
 		/// <returns></returns>
-		Task<bool> UpdateSubscriptionsAsync(ObjectId triggerId, ClaimsPrincipal user, bool? email, bool? slack, CancellationToken cancellationToken);
+		Task<bool> UpdateSubscriptionsAsync(NotificationTriggerId triggerId, ClaimsPrincipal user, bool? email, bool? slack, CancellationToken cancellationToken);
 
 		/// <summary>
 		/// Gets the current subscriptions for a user
@@ -111,7 +111,7 @@ namespace HordeServer.Notifications
 		/// <param name="user"></param>
 		/// <param name="cancellationToken">Cancellation token for the operation</param>
 		/// <returns>Subscriptions for that user</returns>
-		Task<INotificationSubscription?> GetSubscriptionsAsync(ObjectId triggerId, ClaimsPrincipal user, CancellationToken cancellationToken);
+		Task<INotificationSubscription?> GetSubscriptionsAsync(NotificationTriggerId triggerId, ClaimsPrincipal user, CancellationToken cancellationToken);
 
 		/// <summary>
 		/// Notify all subscribers that a job step has finished
diff --git a/Engine/Source/Programs/Horde/Plugins/Build/HordeServer.Build/Notifications/INotificationTrigger.cs b/Engine/Source/Programs/Horde/Plugins/Build/HordeServer.Build/Notifications/INotificationTrigger.cs
--- a/Engine/Source/Programs/Horde/Plugins/Build/HordeServer.Build/Notifications/INotificationTrigger.cs
+++ b/Engine/Source/Programs/Horde/Plugins/Build/HordeServer.Build/Notifications/INotificationTrigger.cs
@@ -1,5 +1,6 @@
 // Copyright Epic Games, Inc. All Rights Reserved.
 
+using EpicGames.Horde.Notifications;
 using EpicGames.Horde.Users;
 using MongoDB.Bson;
 
@@ -34,7 +35,7 @@ namespace HordeServer.Notifications
 		/// <summary>
 		/// Unique id for this subscription list
 		/// </summary>
-		ObjectId Id { get; }
+		NotificationTriggerId Id { get; }
 
 		/// <summary>
 		/// Whether this trigger has been fired
diff --git a/Engine/Source/Programs/Horde/Plugins/Build/HordeServer.Build/Notifications/INotificationTriggerCollection.cs b/Engine/Source/Programs/Horde/Plugins/Build/HordeServer.Build/Notifications/INotificationTriggerCollection.cs
--- a/Engine/Source/Programs/Horde/Plugins/Build/HordeServer.Build/Notifications/INotificationTriggerCollection.cs
+++ b/Engine/Source/Programs/Horde/Plugins/Build/HordeServer.Build/Notifications/INotificationTriggerCollection.cs
@@ -1,7 +1,7 @@
 // Copyright Epic Games, Inc. All Rights Reserved.
 
+using EpicGames.Horde.Notifications;
 using EpicGames.Horde.Users;
-using MongoDB.Bson;
 
 #pragma warning disable CA1030 // Conside making FireAsync an event
 
@@ -18,7 +18,7 @@ namespace HordeServer.Notifications
 		/// <param name="triggerId">The trigger id</param>
 		/// <param name="cancellationToken">Cancellation token for the operation</param>
 		/// <returns>New trigger document</returns>
-		Task<INotificationTrigger> FindOrAddAsync(ObjectId triggerId, CancellationToken cancellationToken);
+		Task<INotificationTrigger> FindOrAddAsync(NotificationTriggerId triggerId, CancellationToken cancellationToken);
 
 		/// <summary>
 		/// Finds an existing trigger with the given id, or adds one if it does not exist
@@ -26,7 +26,7 @@ namespace HordeServer.Notifications
 		/// <param name="triggerId"></param>
 		/// <param name="cancellationToken">Cancellation token for the operation</param>
 		/// <returns></returns>
-		Task<INotificationTrigger?> GetAsync(ObjectId triggerId, CancellationToken cancellationToken);
+		Task<INotificationTrigger?> GetAsync(NotificationTriggerId triggerId, CancellationToken cancellationToken);
 
 		/// <summary>
 		/// Deletes a trigger
@@ -34,7 +34,7 @@ namespace HordeServer.Notifications
 		/// <param name="triggerId">The unique trigger id</param>
 		/// <param name="cancellationToken">Cancellation token for the operation</param>
 		/// <returns>Async task</returns>
-		Task DeleteAsync(ObjectId triggerId, CancellationToken cancellationToken);
+		Task DeleteAsync(NotificationTriggerId triggerId, CancellationToken cancellationToken);
 
 		/// <summary>
 		/// Deletes a trigger
@@ -42,7 +42,7 @@ namespace HordeServer.Notifications
 		/// <param name="triggerIds">The unique trigger id</param>
 		/// <param name="cancellationToken">Cancellation token for the operation</param>
 		/// <returns>Async task</returns>
-		Task DeleteAsync(List<ObjectId> triggerIds, CancellationToken cancellationToken);
+		Task DeleteAsync(List<NotificationTriggerId> triggerIds, CancellationToken cancellationToken);
 
 		/// <summary>
 		/// Fires the trigger, and marks notifications has having been sent
diff --git a/Engine/Source/Programs/Horde/Plugins/Build/HordeServer.Build/Notifications/NotificationService.cs b/Engine/Source/Programs/Horde/Plugins/Build/HordeServer.Build/Notifications/NotificationService.cs
--- a/Engine/Source/Programs/Horde/Plugins/Build/HordeServer.Build/Notifications/NotificationService.cs
+++ b/Engine/Source/Programs/Horde/Plugins/Build/HordeServer.Build/Notifications/NotificationService.cs
@@ -6,6 +6,7 @@ using System.Security.Claims;
 using EpicGames.Core;
 using EpicGames.Horde.Jobs;
 using EpicGames.Horde.Jobs.Graphs;
+using EpicGames.Horde.Notifications;
 using EpicGames.Horde.Users;
 using EpicGames.Redis;
 using HordeServer.Agents;
@@ -202,7 +203,7 @@ namespace HordeServer.Notifications
 		}
 
 		/// <inheritdoc/>
-		public async Task<bool> UpdateSubscriptionsAsync(ObjectId triggerId, ClaimsPrincipal user, bool? email, bool? slack, CancellationToken cancellationToken)
+		public async Task<bool> UpdateSubscriptionsAsync(NotificationTriggerId triggerId, ClaimsPrincipal user, bool? email, bool? slack, CancellationToken cancellationToken)
 		{
 			UserId? userId = user.GetUserId();
 			if (userId == null)
@@ -217,7 +218,7 @@ namespace HordeServer.Notifications
 		}
 
 		/// <inheritdoc/>
-		public async Task<INotificationSubscription?> GetSubscriptionsAsync(ObjectId triggerId, ClaimsPrincipal user, CancellationToken cancellationToken)
+		public async Task<INotificationSubscription?> GetSubscriptionsAsync(NotificationTriggerId triggerId, ClaimsPrincipal user, CancellationToken cancellationToken)
 		{
 			UserId? userId = user.GetUserId();
 			if (userId == null)
@@ -486,7 +487,7 @@ namespace HordeServer.Notifications
 		/// <param name="fireTrigger">If true, the trigger is fired and cannot be reused</param>
 		/// <param name="cancellationToken">Cancellation token for the operation</param>
 		/// <returns></returns>
-		private async Task<INotificationTrigger?> GetNotificationTriggerAsync(ObjectId? triggerId, bool fireTrigger, CancellationToken cancellationToken)
+		private async Task<INotificationTrigger?> GetNotificationTriggerAsync(NotificationTriggerId? triggerId, bool fireTrigger, CancellationToken cancellationToken)
 		{
 			if (triggerId == null)
 			{
@@ -581,7 +582,7 @@ namespace HordeServer.Notifications
 			return Task.CompletedTask;
 		}
 
-		private async Task<IReadOnlyList<IUser>> GetUsersToNotifyAsync(EventRecord? eventRecord, ObjectId? notificationTriggerId, bool fireTrigger, CancellationToken cancellationToken)
+		private async Task<IReadOnlyList<IUser>> GetUsersToNotifyAsync(EventRecord? eventRecord, NotificationTriggerId? notificationTriggerId, bool fireTrigger, CancellationToken cancellationToken)
 		{
 			List<UserId> userIds = new List<UserId>();
 
@@ -747,8 +748,8 @@ namespace HordeServer.Notifications
 						eventId = new LabelCompleteEventRecord(job.StreamId, job.TemplateId, label.DashboardCategory, label.DashboardName, newLabel.Outcome);
 					}
 
-					ObjectId? triggerId;
-					if (job.LabelIdxToTriggerId.TryGetValue(labelIdx, out ObjectId newTriggerId))
+					NotificationTriggerId? triggerId;
+					if (job.LabelIdxToTriggerId.TryGetValue(labelIdx, out NotificationTriggerId newTriggerId))
 					{
 						triggerId = newTriggerId;
 					}
diff --git a/Engine/Source/Programs/Horde/Plugins/Build/HordeServer.Build/Notifications/NotificationTriggerCollection.cs b/Engine/Source/Programs/Horde/Plugins/Build/HordeServer.Build/Notifications/NotificationTriggerCollection.cs
--- a/Engine/Source/Programs/Horde/Plugins/Build/HordeServer.Build/Notifications/NotificationTriggerCollection.cs
+++ b/Engine/Source/Programs/Horde/Plugins/Build/HordeServer.Build/Notifications/NotificationTriggerCollection.cs
@@ -1,5 +1,6 @@
 // Copyright Epic Games, Inc. All Rights Reserved.
 
+using EpicGames.Horde.Notifications;
 using EpicGames.Horde.Users;
 using HordeServer.Server;
 using HordeServer.Users;
@@ -52,6 +53,7 @@ namespace HordeServer.Notifications
 			public List<SubscriptionDocument> Subscriptions { get; set; } = new List<SubscriptionDocument>();
 			public int UpdateIndex { get; set; }
 
+			NotificationTriggerId INotificationTrigger.Id => new NotificationTriggerId(BinaryIdUtils.FromObjectId(Id));
 			IReadOnlyList<INotificationSubscription> INotificationTrigger.Subscriptions => Subscriptions;
 		}
 
@@ -70,9 +72,11 @@ namespace HordeServer.Notifications
 		}
 
 		/// <inheritdoc/>
-		public async Task<INotificationTrigger?> GetAsync(ObjectId triggerId, CancellationToken cancellationToken)
+		public async Task<INotificationTrigger?> GetAsync(NotificationTriggerId triggerId, CancellationToken cancellationToken)
 		{
-			TriggerDocument? trigger = await _triggers.Find(x => x.Id == triggerId).FirstOrDefaultAsync(cancellationToken);
+			ObjectId triggerObjectId = BinaryIdUtils.ToObjectId(triggerId.Id);
+
+			TriggerDocument? trigger = await _triggers.Find(x => x.Id == triggerObjectId).FirstOrDefaultAsync(cancellationToken);
 			if (trigger != null)
 			{
 				for (int idx = 0; idx < trigger.Subscriptions.Count; idx++)
@@ -98,7 +102,7 @@ namespace HordeServer.Notifications
 		}
 
 		/// <inheritdoc/>
-		public async Task<INotificationTrigger> FindOrAddAsync(ObjectId triggerId, CancellationToken cancellationToken)
+		public async Task<INotificationTrigger> FindOrAddAsync(NotificationTriggerId triggerId, CancellationToken cancellationToken)
 		{
 			for (; ; )
 			{
@@ -113,7 +117,7 @@ namespace HordeServer.Notifications
 				try
 				{
 					TriggerDocument newDocument = new TriggerDocument();
-					newDocument.Id = triggerId;
+					newDocument.Id = BinaryIdUtils.ToObjectId(triggerId.Id);
 					await _triggers.InsertOneAsync(newDocument, (InsertOneOptions?)null, cancellationToken);
 					return newDocument;
 				}
@@ -139,7 +143,8 @@ namespace HordeServer.Notifications
 			TriggerDocument document = (TriggerDocument)trigger;
 			int nextUpdateIndex = document.UpdateIndex + 1;
 
-			FilterDefinition<TriggerDocument> filter = Builders<TriggerDocument>.Filter.Expr(x => x.Id == trigger.Id && x.UpdateIndex == document.UpdateIndex);
+			ObjectId triggerObjectId = BinaryIdUtils.ToObjectId(trigger.Id.Id);
+			FilterDefinition<TriggerDocument> filter = Builders<TriggerDocument>.Filter.Expr(x => x.Id == triggerObjectId && x.UpdateIndex == document.UpdateIndex);
 			UpdateDefinition<TriggerDocument> update = transaction.ToUpdateDefinition().Set(x => x.UpdateIndex, nextUpdateIndex);
 
 			UpdateResult result = await _triggers.UpdateOneAsync(filter, update, null, cancellationToken);
@@ -154,15 +159,17 @@ namespace HordeServer.Notifications
 		}
 
 		/// <inheritdoc/>
-		public async Task DeleteAsync(ObjectId triggerId, CancellationToken cancellationToken)
+		public async Task DeleteAsync(NotificationTriggerId triggerId, CancellationToken cancellationToken)
 		{
-			await _triggers.DeleteOneAsync(x => x.Id == triggerId, cancellationToken);
+			ObjectId triggerObjectId = BinaryIdUtils.ToObjectId(triggerId.Id);
+			await _triggers.DeleteOneAsync(x => x.Id == triggerObjectId, cancellationToken);
 		}
 
 		/// <inheritdoc/>
-		public async Task DeleteAsync(List<ObjectId> triggerIds, CancellationToken cancellationToken)
+		public async Task DeleteAsync(List<NotificationTriggerId> triggerIds, CancellationToken cancellationToken)
 		{
-			FilterDefinition<TriggerDocument> filter = Builders<TriggerDocument>.Filter.In(x => x.Id, triggerIds);
+			List<ObjectId> triggerObjectIds = triggerIds.ConvertAll(x => BinaryIdUtils.ToObjectId(x.Id));
+			FilterDefinition<TriggerDocument> filter = Builders<TriggerDocument>.Filter.In(x => x.Id, triggerObjectIds);
 			await _triggers.DeleteManyAsync(filter, cancellationToken);
 		}
 
diff --git a/Engine/Source/Programs/Shared/EpicGames.Horde/Acls/IAclClaim.cs b/Engine/Source/Programs/Shared/EpicGames.Horde/Acls/IAclClaim.cs
new file mode 100644
--- /dev/null
+++ b/Engine/Source/Programs/Shared/EpicGames.Horde/Acls/IAclClaim.cs
@@ -0,0 +1,26 @@
+// Copyright Epic Games, Inc. All Rights Reserved.
+
+using System;
+using System.Collections.Generic;
+using System.Linq;
+using System.Text;
+using System.Threading.Tasks;
+
+namespace EpicGames.Horde.Acls
+{
+	/// <summary>
+	/// Claim for an ACL
+	/// </summary>
+	public interface IAclClaim
+	{
+		/// <summary>
+		/// Type of the claim
+		/// </summary>
+		string Type { get; }
+
+		/// <summary>
+		/// Value for the claim
+		/// </summary>
+		string Value { get; }
+	}
+}
diff --git a/Engine/Source/Programs/Shared/EpicGames.Horde/Notifications/NotificationTriggerId.cs b/Engine/Source/Programs/Shared/EpicGames.Horde/Notifications/NotificationTriggerId.cs
new file mode 100644
--- /dev/null
+++ b/Engine/Source/Programs/Shared/EpicGames.Horde/Notifications/NotificationTriggerId.cs
@@ -0,0 +1,37 @@
+// Copyright Epic Games, Inc. All Rights Reserved.
+
+using System.ComponentModel;
+using EpicGames.Core;
+using EpicGames.Horde;
+
+namespace EpicGames.Horde.Notifications
+{
+	/// <summary>
+	/// Unique id for a notification trigger id
+	/// </summary>
+	/// <param name="Id">Identifier for the notification trigger</param>
+	[LogValueType]
+	[JsonSchemaString]
+	[TypeConverter(typeof(BinaryIdTypeConverter<NotificationTriggerId, NotificationTriggerIdConverter>))]
+	[BinaryIdConverter(typeof(NotificationTriggerIdConverter))]
+	public readonly record struct NotificationTriggerId(BinaryId Id)
+	{
+		/// <inheritdoc cref="BinaryId.Parse(System.String)"/>
+		public static NotificationTriggerId Parse(string text) => new NotificationTriggerId(BinaryId.Parse(text));
+
+		/// <inheritdoc/>
+		public override readonly string ToString() => Id.ToString();
+	}
+
+	/// <summary>
+	/// Converter class to and from ObjectId values
+	/// </summary>
+	class NotificationTriggerIdConverter : BinaryIdConverter<NotificationTriggerId>
+	{
+		/// <inheritdoc/>
+		public override NotificationTriggerId FromBinaryId(BinaryId id) => new NotificationTriggerId(id);
+
+		/// <inheritdoc/>
+		public override BinaryId ToBinaryId(NotificationTriggerId value) => value.Id;
+	}
+}
