Horde: Updating schema documentation

#rnx

diff --git a/Engine/Source/Programs/Horde/Docs/Config/Schema/Projects.md b/Engine/Source/Programs/Horde/Docs/Config/Schema/Projects.md
--- a/Engine/Source/Programs/Horde/Docs/Config/Schema/Projects.md
+++ b/Engine/Source/Programs/Horde/Docs/Config/Schema/Projects.md
@@ -199,7 +199,7 @@ Information about a workspace type
 
 Name | Description
 ---- | -----------
-`base` | `string`<br>Base workspace to derive from
+`base` | `string`<br>
 `cluster` | `string`<br>Name of the Perforce server cluster to use
 `serverAndPort` | `string`<br>The Perforce server and port (eg. perforce:1666)
 `userName` | `string`<br>User to log into Perforce with (defaults to buildmachine)
diff --git a/Engine/Source/Programs/Horde/Docs/Config/Schema/Streams.md b/Engine/Source/Programs/Horde/Docs/Config/Schema/Streams.md
--- a/Engine/Source/Programs/Horde/Docs/Config/Schema/Streams.md
+++ b/Engine/Source/Programs/Horde/Docs/Config/Schema/Streams.md
@@ -197,7 +197,7 @@ Information about a workspace type
 
 Name | Description
 ---- | -----------
-`base` | `string`<br>Base workspace to derive from
+`base` | `string`<br>
 `cluster` | `string`<br>Name of the Perforce server cluster to use
 `serverAndPort` | `string`<br>The Perforce server and port (eg. perforce:1666)
 `userName` | `string`<br>User to log into Perforce with (defaults to buildmachine)
@@ -256,7 +256,7 @@ Name | Description
 `requireSubmittedChange` | `boolean`<br>Whether the build requires a change to be submitted
 `gate` | [ScheduleGateConfig](#schedulegateconfig)<br>Gate allowing the schedule to trigger
 `commits` | `string[]`<br>Commit tags for this schedule
-`claims` | [AclClaimConfig](#aclclaimconfig)`[]`<br>Roles to impersonate for this schedule
+`claims` | [ScheduleClaimConfig](#scheduleclaimconfig)`[]`<br>Roles to impersonate for this schedule
 `filter` | [ChangeContentFlags](#changecontentflags-enum)`[]`<br>The types of changes to run for
 `files` | `string[]`<br>Files that should cause the job to trigger
 `templateParameters` | `string` `->` `string`<br>Parameters for the template
@@ -271,9 +271,9 @@ Name | Description
 `templateId` | `string`<br>The template containing the dependency
 `target` | `string`<br>Target to wait for
 
-## AclClaimConfig
+## ScheduleClaimConfig
 
-New claim to create
+Claim to grant to leases running a scheduled job
 
 Name | Description
 ---- | -----------
@@ -343,6 +343,15 @@ Name | Description
 `actions` | `string[]`<br>Array of actions to allow
 `profiles` | `string[]`<br>List of profiles to grant
 
+## AclClaimConfig
+
+New claim to create
+
+Name | Description
+---- | -----------
+`type` | `string`<br>The claim type
+`value` | `string`<br>The claim value
+
 ## AclProfileConfig
 
 Configuration for an ACL profile. This defines a preset group of actions which can be given to a user via an ACL entry.
