Horde: Static analysis fix.

#rnx

diff --git a/Engine/Source/Programs/Horde/Plugins/Experimental/HordeServer.Experimental/Notifications/Sinks/ExperimentalSlackNotificationSink.cs b/Engine/Source/Programs/Horde/Plugins/Experimental/HordeServer.Experimental/Notifications/Sinks/ExperimentalSlackNotificationSink.cs
--- a/Engine/Source/Programs/Horde/Plugins/Experimental/HordeServer.Experimental/Notifications/Sinks/ExperimentalSlackNotificationSink.cs
+++ b/Engine/Source/Programs/Horde/Plugins/Experimental/HordeServer.Experimental/Notifications/Sinks/ExperimentalSlackNotificationSink.cs
@@ -7,7 +7,6 @@ using HordeServer.Agents;
 using HordeServer.Configuration;
 using HordeServer.Devices;
 using HordeServer.Issues;
-using HordeServer.Jobs;
 using HordeServer.Logs;
 using HordeServer.Streams;
 using HordeServer.Users;
