Horde: Add stats for storage service.

#rnx

diff --git a/Engine/Source/Programs/Horde/Plugins/Storage/HordeServer.Storage.Tests/Storage/StorageServiceTests.cs b/Engine/Source/Programs/Horde/Plugins/Storage/HordeServer.Storage.Tests/Storage/StorageServiceTests.cs
--- a/Engine/Source/Programs/Horde/Plugins/Storage/HordeServer.Storage.Tests/Storage/StorageServiceTests.cs
+++ b/Engine/Source/Programs/Horde/Plugins/Storage/HordeServer.Storage.Tests/Storage/StorageServiceTests.cs
@@ -2,8 +2,10 @@
 
 using EpicGames.Core;
 using EpicGames.Horde.Storage;
+using HordeServer.Server;
 using HordeServer.Storage;
 using Microsoft.Extensions.DependencyInjection;
+using MongoDB.Driver;
 
 namespace HordeServer.Tests.Storage
 {
@@ -17,6 +19,34 @@ namespace HordeServer.Tests.Storage
 			AddPlugin<StoragePlugin>();
 		}
 
+		[TestMethod]
+		public async Task StatsTestAsync()
+		{
+			StorageService storageService = ServiceProvider.GetRequiredService<StorageService>();
+			IStorageNamespace client = storageService.GetNamespace(new NamespaceId("memory"));
+
+			await using (IBlobWriter writer = client.CreateBlobWriter())
+			{
+				BlobType type1 = new BlobType(Guid.Parse("{11C2D886-4164-3349-D1E9-6F943D2ED10B}"), 0);
+				for (int idx = 0; idx < 1000; idx++)
+				{
+					writer.WriteFixedLengthBytes(new byte[] { 1, 2, 3 });
+					await writer.CompleteAsync(type1);
+					await writer.FlushAsync();
+				}
+			}
+
+			await Clock.AdvanceAsync(TimeSpan.FromDays(7.0));
+
+			await storageService.TickBlobsAsync(CancellationToken.None);
+			await storageService.TickLengthsAsync(CancellationToken.None);
+			await storageService.TickStatsAsync(CancellationToken.None);
+
+			IReadOnlyList<IStorageStats> stats = await storageService.FindStatsAsync();
+			Assert.AreEqual(1, stats.Count);
+			Assert.IsTrue(stats[0].Namespaces.First().Value.Size > 3 * 1000);
+		}
+
 		[TestMethod]
 		public async Task BlobCollectionTestAsync()
 		{
diff --git a/Engine/Source/Programs/Horde/Plugins/Storage/HordeServer.Storage/Storage/StorageService.cs b/Engine/Source/Programs/Horde/Plugins/Storage/HordeServer.Storage/Storage/StorageService.cs
--- a/Engine/Source/Programs/Horde/Plugins/Storage/HordeServer.Storage/Storage/StorageService.cs
+++ b/Engine/Source/Programs/Horde/Plugins/Storage/HordeServer.Storage/Storage/StorageService.cs
@@ -37,6 +37,43 @@ namespace HordeServer.Storage
 	{
 	}
 
+	/// <summary>
+	/// Stats across the storage system
+	/// </summary>
+	public interface IStorageStats
+	{
+		/// <summary>
+		/// Time that the stats apply to
+		/// </summary>
+		public DateTime Time { get; }
+
+		/// <summary>
+		/// Amount of time that it took to gather these stats
+		/// </summary>
+		public long ScanTimeSecs { get; }
+
+		/// <summary>
+		/// Per-namespace stats
+		/// </summary>
+		public IReadOnlyDictionary<NamespaceId, IStorageStatsForNamespace> Namespaces { get; }
+	}
+
+	/// <summary>
+	/// Stats for a particular namespace
+	/// </summary>
+	public interface IStorageStatsForNamespace
+	{
+		/// <summary>
+		/// Number of blobs
+		/// </summary>
+		long Count { get; }
+
+		/// <summary>
+		/// Total size of the namespace
+		/// </summary>
+		long Size { get; }
+	}
+
 	/// <summary>
 	/// Functionality related to the storage service
 	/// </summary>
@@ -403,6 +440,44 @@ namespace HordeServer.Storage
 			}
 		}
 
+		class Stats : IStorageStats
+		{
+			[BsonElement("tm")]
+			public DateTime Time { get; set; }
+
+			[BsonElement("st")]
+			public long ScanTimeSecs { get; set; }
+
+			[BsonElement("ns")]
+			public Dictionary<NamespaceId, StatsForNamespace> Namespaces { get; set; } = new Dictionary<NamespaceId, StatsForNamespace>();
+
+			IReadOnlyDictionary<NamespaceId, IStorageStatsForNamespace>? _cachedNamespaces;
+			IReadOnlyDictionary<NamespaceId, IStorageStatsForNamespace> IStorageStats.Namespaces
+				=> _cachedNamespaces ??= Namespaces.ToDictionary(x => x.Key, x => (IStorageStatsForNamespace)x.Value);
+		}
+
+		class StatsForNamespace : IStorageStatsForNamespace
+		{
+			[BsonElement("ct")]
+			public long Count { get; set; }
+
+			[BsonElement("sz")]
+			public long Size { get; set; }
+		}
+
+		[SingletonDocument("stats-state")]
+		class StatsState : SingletonBase
+		{
+			[BsonElement("st")]
+			public DateTime? StartTime { get; set; }
+
+			[BsonElement("blob")]
+			public ObjectId LastBlobId { get; set; }
+
+			[BsonElement("ns")]
+			public Dictionary<NamespaceId, StatsForNamespace> Namespaces { get; set; } = new Dictionary<NamespaceId, StatsForNamespace>();
+		}
+
 		[SingletonDocument("gc-state")]
 		class GcState : SingletonBase
 		{
@@ -459,10 +534,14 @@ namespace HordeServer.Storage
 
 		readonly IMongoCollection<BlobInfo> _blobCollection;
 		readonly IMongoCollection<RefInfo> _refCollection;
+		readonly IMongoCollection<Stats> _statsCollection;
 
 		readonly ITicker _blobTicker;
 		readonly ITicker _refTicker;
 
+		readonly SingletonDocument<StatsState> _statsState;
+		readonly ITicker _statsTicker;
+
 		readonly SingletonDocument<GcState> _gcState;
 		readonly ITicker _gcTicker;
 
@@ -520,9 +599,16 @@ namespace HordeServer.Storage
 			refIndexes.Add(keys => keys.Descending(x => x.ExpiresAtUtc), sparse: true);
 			_refCollection = mongoService.GetCollection<RefInfo>("Storage.Refs", refIndexes);
 
+			List<MongoIndex<Stats>> statsIndexes = new List<MongoIndex<Stats>>();
+			statsIndexes.Add(keys => keys.Ascending(x => x.Time));
+			_statsCollection = mongoService.GetCollection<Stats>("Storage.Stats");
+
 			_blobTicker = clock.AddSharedTicker("Storage:Blobs", TimeSpan.FromMinutes(5.0), TickBlobsAsync, _logger);
 			_refTicker = clock.AddSharedTicker("Storage:Refs", TimeSpan.FromMinutes(5.0), TickRefsAsync, _logger);
 
+			_statsState = new SingletonDocument<StatsState>(mongoService);
+			_statsTicker = clock.AddSharedTicker("Storage:Stats", TimeSpan.FromMinutes(10.0), TickStatsAsync, _logger);
+
 			_gcState = new SingletonDocument<GcState>(mongoService);
 			_gcTicker = clock.AddTicker("Storage:GC", TimeSpan.FromMinutes(5.0), TickGcAsync, logger);
 
@@ -764,7 +850,7 @@ namespace HordeServer.Storage
 		/// Finds blobs at least 30 minutes old and computes import metadata for them. Done with a delay to allow write redirects.
 		/// </summary>
 		/// <param name="cancellationToken">Cancellation token for the operation</param>
-		async ValueTask TickBlobsAsync(CancellationToken cancellationToken)
+		internal async ValueTask TickBlobsAsync(CancellationToken cancellationToken)
 		{
 			using TelemetrySpan span = _tracer.StartActiveSpan($"{nameof(StorageService)}.{nameof(TickBlobsAsync)}");
 
@@ -1156,6 +1242,106 @@ namespace HordeServer.Storage
 
 		#endregion
 
+		#region Stats
+
+		/// <summary>
+		/// Find stats for a range of time
+		/// </summary>
+		/// <param name="startTime">Start time for the search</param>
+		/// <param name="finishTime">Finish time for the search</param>
+		/// <param name="count">Number of elements to return</param>
+		/// <param name="cancellationToken">Cancellation token for the operation</param>
+		/// <returns>Stats objects</returns>
+		public async Task<IReadOnlyList<IStorageStats>> FindStatsAsync(DateTime? startTime = null, DateTime? finishTime = null, int count = 10, CancellationToken cancellationToken = default)
+		{
+			List<FilterDefinition<Stats>> filters = new List<FilterDefinition<Stats>>();
+			if (startTime != null)
+			{
+				filters.Add(Builders<Stats>.Filter.Gte(x => x.Time, startTime.Value));
+			}
+			if (finishTime != null)
+			{
+				filters.Add(Builders<Stats>.Filter.Lte(x => x.Time, finishTime.Value));
+			}
+
+			FilterDefinition<Stats> filter = (filters.Count == 0)? FilterDefinition<Stats>.Empty : Builders<Stats>.Filter.And(filters);
+			return await _statsCollection.Find(filter).SortByDescending(x => x.Time).Limit(count).ToListAsync(cancellationToken);
+		}
+
+		internal async ValueTask TickStatsAsync(CancellationToken cancellationToken)
+		{
+			using TelemetrySpan span = _tracer.StartActiveSpan($"{nameof(StorageService)}.{nameof(TickStatsAsync)}");
+
+			StatsState? statsState = null;
+			for (; ; )
+			{
+				// Make sure we have a current state
+				statsState ??= await _statsState.GetAsync(cancellationToken);
+
+				// Reset the start time
+				if (statsState.StartTime == null)
+				{
+					statsState.StartTime = _clock.UtcNow;
+					statsState.LastBlobId = ObjectId.Empty;
+				}
+
+				// If we're not meant to start yet, skip this iteration
+				if (_clock.UtcNow < statsState.StartTime)
+				{
+					_logger.LogDebug("Next stats update will begin at {StartTime}", statsState.StartTime);
+					break;
+				}
+
+				// Fetch up to the blob referenced by the start time
+				ObjectId latestInfoId = ObjectId.GenerateNewId(statsState.StartTime.Value);
+
+				// Fetch the next batch of blobs
+				List<BlobInfo> current = await _blobCollection.Find(x => x.Id > statsState.LastBlobId && x.Id < latestInfoId).SortBy(x => x.Id).Limit(500).ToListAsync(cancellationToken);
+				if (current.Count == 0)
+				{
+					_logger.LogInformation("Publishing storage stats for {Time}", statsState.StartTime.Value);
+
+					// Publish the current stats
+					Stats stats = new Stats();
+					stats.Time = statsState.StartTime.Value;
+					stats.ScanTimeSecs = (long)(_clock.UtcNow - stats.Time).TotalSeconds;
+					stats.Namespaces = statsState.Namespaces;
+					await _statsCollection.InsertOneAsync(stats, null, cancellationToken);
+
+					// Reset the current stats
+					DateTime nextStartTime = new DateTime(DateOnly.FromDateTime(statsState.StartTime.Value).AddDays(1), new TimeOnly());
+					statsState = new StatsState { StartTime = nextStartTime, Revision = statsState.Revision };
+					await _statsState.TryUpdateAsync(statsState, cancellationToken);
+					break;
+				}
+				statsState.LastBlobId = current[^1].Id;
+
+				// Add a check record for each new blob
+				_logger.LogDebug("Adding {NumBlobs} blobs for stats ({FirstId} to {LastId})", current.Count, current[0].Id, current[^1].Id);
+				foreach (BlobInfo blobInfo in current)
+				{
+					StatsForNamespace? namespaceStats;
+					if (!statsState.Namespaces.TryGetValue(blobInfo.NamespaceId, out namespaceStats))
+					{
+						namespaceStats = new StatsForNamespace();
+						statsState.Namespaces.Add(blobInfo.NamespaceId, namespaceStats);
+					}
+
+					namespaceStats.Count++;
+					namespaceStats.Size += blobInfo.Length;
+				}
+
+				// Update the last imported blob id
+				if (!await _statsState.TryUpdateAsync(statsState, cancellationToken))
+				{
+					_logger.LogInformation("Unable to update stats; resetting scan.");
+					statsState = null;
+				}
+			}
+		}
+
+		#endregion
+
 		#region GC
 
 		uint GetGcTimestamp() => GetGcTimestamp(_clock.UtcNow);
@@ -1465,7 +1651,7 @@ namespace HordeServer.Storage
 		/// Scan new blobs for their sizes
 		/// </summary>
 		/// <param name="cancellationToken">Cancellation token for the operation</param>
-		async ValueTask TickLengthsAsync(CancellationToken cancellationToken)
+		internal async ValueTask TickLengthsAsync(CancellationToken cancellationToken)
 		{
 			using TelemetrySpan span = _tracer.StartActiveSpan($"{nameof(StorageService)}.{nameof(TickLengthsAsync)}");
 
