Horde: Fix static analysis warning.

#rnx

diff --git a/Engine/Source/Programs/Horde/Plugins/Storage/HordeServer.Storage/Storage/StorageService.cs b/Engine/Source/Programs/Horde/Plugins/Storage/HordeServer.Storage/Storage/StorageService.cs
--- a/Engine/Source/Programs/Horde/Plugins/Storage/HordeServer.Storage/Storage/StorageService.cs
+++ b/Engine/Source/Programs/Horde/Plugins/Storage/HordeServer.Storage/Storage/StorageService.cs
@@ -628,6 +628,7 @@ namespace HordeServer.Storage
 		{
 			await _blobTicker.DisposeAsync();
 			await _refTicker.DisposeAsync();
+			await _statsTicker.DisposeAsync();
 			await _gcTicker.DisposeAsync();
 			await _lengthScanTicker.DisposeAsync();
 		}
