Horde: Allow leases running in a stream to access artifacts in that stream by default, if not overridden elsewhere.

diff --git a/Engine/Source/Programs/Horde/HordeServer.Shared/Acls/AclConfig.cs b/Engine/Source/Programs/Horde/HordeServer.Shared/Acls/AclConfig.cs
--- a/Engine/Source/Programs/Horde/HordeServer.Shared/Acls/AclConfig.cs
+++ b/Engine/Source/Programs/Horde/HordeServer.Shared/Acls/AclConfig.cs
@@ -75,6 +75,15 @@ namespace HordeServer.Acls
 		/// <param name="user">The principal to authorize</param>
 		/// <returns>True/false if the action is allowed or denied, null if there is no specific setting for this user</returns>
 		public bool Authorize(AclAction action, ClaimsPrincipal user)
+			=> TryAuthorize(action, user) ?? false;
+
+		/// <summary>
+		/// Tests whether a user is authorized to perform the given actions
+		/// </summary>
+		/// <param name="action">Action that is being performed. This should be a single flag.</param>
+		/// <param name="user">The principal to authorize</param>
+		/// <returns>True/false if the action is allowed or denied, null if there is no specific setting for this user</returns>
+		public bool? TryAuthorize(AclAction action, ClaimsPrincipal user)
 		{
 			if (user.HasAdminClaim())
 			{
@@ -90,7 +99,7 @@ namespace HordeServer.Acls
 				}
 			}
 
-			return false;
+			return null;
 		}
 
 		/// <summary>
diff --git a/Engine/Source/Programs/Horde/Plugins/Build/HordeServer.Build/BuildConfig.cs b/Engine/Source/Programs/Horde/Plugins/Build/HordeServer.Build/BuildConfig.cs
--- a/Engine/Source/Programs/Horde/Plugins/Build/HordeServer.Build/BuildConfig.cs
+++ b/Engine/Source/Programs/Horde/Plugins/Build/HordeServer.Build/BuildConfig.cs
@@ -224,7 +224,7 @@ namespace HordeServer
 			}
 
 			auth ??= AuthorizeArtifactType(ArtifactTypes, type, action, user);
-			return auth ?? Acl.Authorize(action, user);
+			return auth ?? Acl.TryAuthorize(action, user) ?? user.HasClaim(HordeClaimTypes.LeaseStream, streamId.ToString());
 		}
 
 		static bool? AuthorizeArtifactType(IReadOnlyList<ArtifactTypeConfig> artifactTypes, ArtifactType artifactType, AclAction action, ClaimsPrincipal user)
