Horde: Fix exception sorting stream responses.

#rnx

diff --git a/Engine/Source/Programs/Horde/Plugins/Build/HordeServer.Build/Streams/StreamsController.cs b/Engine/Source/Programs/Horde/Plugins/Build/HordeServer.Build/Streams/StreamsController.cs
--- a/Engine/Source/Programs/Horde/Plugins/Build/HordeServer.Build/Streams/StreamsController.cs
+++ b/Engine/Source/Programs/Horde/Plugins/Build/HordeServer.Build/Streams/StreamsController.cs
@@ -96,7 +96,7 @@ namespace HordeServer.Streams
 				responses.Add(response);
 			}
 
-			return responses.OrderBy(x => x.Id).Select(x => PropertyFilter.Apply(x, filter)).ToList();
+			return responses.OrderBy(x => x.Id.Id.Text, Utf8StringComparer.Ordinal).Select(x => PropertyFilter.Apply(x, filter)).ToList();
 		}
 
 		/// <summary>
