Horde: Add explicit clear selection to template chooser
#skipci

diff --git a/Engine/Source/Programs/Horde/HordeDashboard/src/components/JobSearchSimple.tsx b/Engine/Source/Programs/Horde/HordeDashboard/src/components/JobSearchSimple.tsx
--- a/Engine/Source/Programs/Horde/HordeDashboard/src/components/JobSearchSimple.tsx
+++ b/Engine/Source/Programs/Horde/HordeDashboard/src/components/JobSearchSimple.tsx
@@ -99,6 +99,8 @@ export const JobSearchSimpleModal: React.FC<{ streamIdIn: string, changeListIn?:
       })
    })
 
+   const allTemplates = new Set<string>();
+
    Array.from(sorted.keys()).sort((a, b) => a < b ? -1 : 1).forEach(cat => {
 
       const templates = sorted.get(cat);
@@ -107,6 +109,8 @@ export const JobSearchSimpleModal: React.FC<{ streamIdIn: string, changeListIn?:
       }
 
       const subItems: IContextualMenuItem[] = templates.sort((a, b) => a.name < b.name ? -1 : 1).map(t => {
+
+         allTemplates.add(t.id);
          return {
             itemType: ContextualMenuItemType.Normal,
             key: t.id,
@@ -134,11 +138,21 @@ export const JobSearchSimpleModal: React.FC<{ streamIdIn: string, changeListIn?:
 
    })
 
+   templateOptions.push({
+      key: `show_all_templates`, text: "Clear Selection", onClick: (ev) => {
+         ev?.stopPropagation();
+         ev?.preventDefault();
+         searchState.templates = new Set();
+         setSearchState({ ...searchState, containsParameter: undefined })
+         return false;
+      }
+   });
+
    templateOptions.push({
       key: `show_all_templates`, text: "All Templates", onClick: (ev) => {
          ev?.stopPropagation();
          ev?.preventDefault();
-         searchState.templates.clear();
+         searchState.templates = allTemplates;
          setSearchState({ ...searchState, containsParameter: undefined })
          return false;
       }
@@ -571,7 +585,7 @@ export const JobSearchSimpleModal: React.FC<{ streamIdIn: string, changeListIn?:
       title = `Find Jobs in ${nstream?.fullname ?? searchState.streamIdOverride}`;
    }
 
-   let templateText = "All Templates";
+   let templateText = "None Selected";
    if (searchState.templates.size) {
       templateText = ctemplates.map(t => t.name).join(", ");
       if (templateText.length > 84) {
