Horde: Log encryption used for compute tasks in lease log

diff --git a/Engine/Source/Programs/Horde/HordeAgent/Leases/Handlers/ComputeHandler.cs b/Engine/Source/Programs/Horde/HordeAgent/Leases/Handlers/ComputeHandler.cs
--- a/Engine/Source/Programs/Horde/HordeAgent/Leases/Handlers/ComputeHandler.cs
+++ b/Engine/Source/Programs/Horde/HordeAgent/Leases/Handlers/ComputeHandler.cs
@@ -61,52 +61,48 @@ namespace HordeAgent.Leases.Handlers
 				}
 
 				logger.LogInformation("Matched connection for {Nonce}", StringUtils.FormatHexString(computeTask.Nonce.Span));
-
-				using (CancellationTokenSource cts = CancellationTokenSource.CreateLinkedTokenSource(cancellationToken))
+				
+				using CancellationTokenSource cts = CancellationTokenSource.CreateLinkedTokenSource(cancellationToken);
+				await using ComputeTransport innerTransport = await CreateTransportAsync(computeTask, tcpClient.Client, cts.Token);
+				await using IdleTimeoutTransport idleTimeoutTransport = new(innerTransport, timeout);
+				
+				await using BackgroundTask timeoutTask = BackgroundTask.StartNew(ctx => idleTimeoutTransport.StartWatchdogTimerAsync(cts, logger, ctx));
+				try
 				{
-					await using ComputeTransport innerTransport = await CreateTransportAsync(computeTask, tcpClient.Client, cts.Token);
-					await using IdleTimeoutTransport idleTimeoutTransport = new(innerTransport, timeout);
-
-					await using BackgroundTask timeoutTask = BackgroundTask.StartNew(ctx => idleTimeoutTransport.StartWatchdogTimerAsync(cts, logger, ctx));
+					ComputeProtocol protocol = (ComputeProtocol)computeTask.Protocol;
+					logger.LogInformation("Using compute protocol version {Version} with encryption {Encryption}", (int)protocol, computeTask.Encryption);
+					
+					await using RemoteComputeSocket socket = new (idleTimeoutTransport, protocol, logger);
+					DirectoryReference sandboxDir = DirectoryReference.Combine(session.WorkingDir, "Sandbox", leaseId.ToString());
 					try
 					{
-						ComputeProtocol protocol = (ComputeProtocol)computeTask.Protocol;
-						logger.LogInformation("Using compute protocol version {Version}", (int)protocol);
-
-						await using (RemoteComputeSocket socket = new RemoteComputeSocket(idleTimeoutTransport, protocol, logger))
+						DirectoryReference.CreateDirectory(sandboxDir);
+						DirectoryReference sharedDir = DirectoryReference.Combine(session.WorkingDir, "Saved");
+						DirectoryReference.CreateDirectory(sharedDir);
+						
+						Dictionary<string, string?> newEnvVars = new ()
 						{
-							DirectoryReference sandboxDir = DirectoryReference.Combine(session.WorkingDir, "Sandbox", leaseId.ToString());
-							try
-							{
-								DirectoryReference.CreateDirectory(sandboxDir);
-								DirectoryReference sharedDir = DirectoryReference.Combine(session.WorkingDir, "Saved");
-								DirectoryReference.CreateDirectory(sharedDir);
-
-								Dictionary<string, string?> newEnvVars = new ()
-								{
-									["UE_HORDE_SHARED_DIR"] = sharedDir.FullName,
-									["UE_HORDE_TERMINATION_SIGNAL_FILE"] = _settings.CurrentValue.GetTerminationSignalFile().FullName,
-									["UE_HORDE_CPU_COUNT"] = Convert.ToString(_settings.CurrentValue.CpuCount),
-									["UE_HORDE_CPU_MULTIPLIER"] = Convert.ToString(_settings.CurrentValue.CpuMultiplier, CultureInfo.InvariantCulture),
-								};
-								
-								AgentMessageHandler worker = new AgentMessageHandler(sandboxDir, newEnvVars, false, _settings.CurrentValue.WineExecutablePath, _settings.CurrentValue.ContainerEngineExecutablePath, logger);
-								await worker.RunAsync(socket, cts.Token);
-								await socket.CloseAsync(cts.Token);
-								return LeaseResult.Success;
-							}
-							finally
-							{
-								FileUtils.ForceDeleteDirectory(sandboxDir);
-							}
-						}
+							["UE_HORDE_SHARED_DIR"] = sharedDir.FullName,
+							["UE_HORDE_TERMINATION_SIGNAL_FILE"] = _settings.CurrentValue.GetTerminationSignalFile().FullName,
+							["UE_HORDE_CPU_COUNT"] = Convert.ToString(_settings.CurrentValue.CpuCount),
+							["UE_HORDE_CPU_MULTIPLIER"] = Convert.ToString(_settings.CurrentValue.CpuMultiplier, CultureInfo.InvariantCulture),
+						};
+						
+						AgentMessageHandler worker = new AgentMessageHandler(sandboxDir, newEnvVars, false, _settings.CurrentValue.WineExecutablePath, _settings.CurrentValue.ContainerEngineExecutablePath, logger);
+						await worker.RunAsync(socket, cts.Token);
+						await socket.CloseAsync(cts.Token);
+						return LeaseResult.Success;
 					}
-					catch (OperationCanceledException ex) when (cts.IsCancellationRequested && idleTimeoutTransport.TimeSinceActivity > idleTimeoutTransport.NoDataTimeout)
+					finally
 					{
-						logger.LogError(ex, "Lease was terminated due to no data being received for {Time} seconds", (int)idleTimeoutTransport.NoDataTimeout.TotalSeconds);
-						return LeaseResult.Failed;
+						FileUtils.ForceDeleteDirectory(sandboxDir);
 					}
 				}
+				catch (OperationCanceledException ex) when (cts.IsCancellationRequested && idleTimeoutTransport.TimeSinceActivity > idleTimeoutTransport.NoDataTimeout)
+				{
+					logger.LogError(ex, "Lease was terminated due to no data being received for {Time} seconds", (int)idleTimeoutTransport.NoDataTimeout.TotalSeconds);
+					return LeaseResult.Failed;
+				}
 			}
 			catch (OperationCanceledException)
 			{
