Horde: Enable background storage stats ticker.

#rnx

diff --git a/Engine/Source/Programs/Horde/Plugins/Storage/HordeServer.Storage/Storage/StorageService.cs b/Engine/Source/Programs/Horde/Plugins/Storage/HordeServer.Storage/Storage/StorageService.cs
--- a/Engine/Source/Programs/Horde/Plugins/Storage/HordeServer.Storage/Storage/StorageService.cs
+++ b/Engine/Source/Programs/Horde/Plugins/Storage/HordeServer.Storage/Storage/StorageService.cs
@@ -661,6 +661,7 @@ namespace HordeServer.Storage
 		{
 			await _blobTicker.StartAsync();
 			await _refTicker.StartAsync();
+			await _statsTicker.StartAsync();
 			await _gcTicker.StartAsync();
 			await _lengthScanTicker.StartAsync();
 		}
@@ -669,6 +670,7 @@ namespace HordeServer.Storage
 		public async Task StopAsync(CancellationToken cancellationToken)
 		{
 			await _lengthScanTicker.StopAsync();
+			await _statsTicker.StopAsync();
 			await _gcTicker.StopAsync();
 			await _refTicker.StopAsync();
 			await _blobTicker.StopAsync();
