Horde: Update release notes.

#rnx

diff --git a/Engine/Source/Programs/Horde/Docs/ReleaseNotes.md b/Engine/Source/Programs/Horde/Docs/ReleaseNotes.md
--- a/Engine/Source/Programs/Horde/Docs/ReleaseNotes.md
+++ b/Engine/Source/Programs/Horde/Docs/ReleaseNotes.md
@@ -1,5 +1,26 @@
 # Release Notes
 
+## 2024-11-14
+
+* Handle EndOfStreamException in RecvOptionalAsync. Treat EndOfStreamException same as zero bytes read, returning false to indicate no data available rather than propagating the exception. (38077833)
+* Forward completion event to inner in IdleTimeoutTransport (38077823)
+* Change back to 16 read tasks for batch reads from storage. (38057408)
+* Set maximum size of channels used for buffering data read from storage. If part of this pipeline is particularly slow (eg. I/O, network, CPU), is can cause large backlogs of buffered data which can kill performance or cause machine instability. Applying some back pressure resolves the issue without any loss in performance. (38051395)
+* Log encryption used for compute tasks in lease log (38031944)
+* Make port allocation for tests thread-safe (38016954)
+* Add explicit clear selection to template chooser (38007851)
+* Enable parallel test execution for EpicGames.Horde.Tests. Fixed race conditions by adding unique suffixes to shared test directories. (38002324)
+* Allow leases running in a stream to access artifacts in that stream by default, if not overridden elsewhere. (37996421)
+* Replace AesTransport with a new implementation. Uses random nonces for each message instead of a base nonce. Add TCP transport integration tests (37992625)
+* Add parameter for controlling inactivity timeout for compute executions. For UBA-based compute tasks this value needs to be increased as CPU or I/O can at times get starved, leading to a timeout. (37960155)
+* Job search improvements, link to related jobs from operations bar, step callout optimizations (37920079)
+* Improve symbol indexing time; perform hash calculations in parallel, and upload aliases in smaller batches. (37898457)
+* Hide agent install console window and skip uninstall registration in Windows Registry (37874147)
+* Fix agent uninstall directory path when registry path is not used (37874091)
+* Fix CLI tool upload using named blobs (37873422)
+* Show critical log levels as errors in the log view (37845567)
+* Add an IStorageWriter interface, which allows buffering and batching metadata update requests. Use this to upload artifact aliases in a single batch. (37843288)
+
 ## 2024-11-06
 
 * Fix cached commit source decrementing number of remaining results to return even if commits are filtered out of the resulting list by commit tags. Prevents issues with new jobs not being able to identify code changelist to run at. (37833532)
