Horde: Add discriminator names for template parameter classes.

#rnx

diff --git a/Engine/Source/Programs/Horde/Plugins/Build/HordeServer.Build/Jobs/Templates/TemplateCollection.cs b/Engine/Source/Programs/Horde/Plugins/Build/HordeServer.Build/Jobs/Templates/TemplateCollection.cs
--- a/Engine/Source/Programs/Horde/Plugins/Build/HordeServer.Build/Jobs/Templates/TemplateCollection.cs
+++ b/Engine/Source/Programs/Horde/Plugins/Build/HordeServer.Build/Jobs/Templates/TemplateCollection.cs
@@ -99,6 +99,7 @@ namespace HordeServer.Jobs.Templates
 			public abstract void GetDefaultParameters(Dictionary<ParameterId, string> parameters, bool scheduledBuild);
 		}
 
+		[BsonDiscriminator("BoolParameter")]
 		sealed class BoolParameterDocument : ParameterDocument, ITemplateBoolParameter
 		{
 			public ParameterId Id { get; set; }
@@ -202,6 +203,7 @@ namespace HordeServer.Jobs.Templates
 			}
 		}
 
+		[BsonDiscriminator("TextParameter")]
 		sealed class TextParameterDocument : ParameterDocument, ITemplateTextParameter
 		{
 			public ParameterId Id { get; set; }
@@ -259,6 +261,7 @@ namespace HordeServer.Jobs.Templates
 			}
 		}
 
+		[BsonDiscriminator("ListParameter")]
 		sealed class ListParameterDocument : ParameterDocument, ITemplateListParameter
 		{
 			public string Label { get; set; }
