Add LLM Memory Tracking Metircs to default Horde Analytics view

diff --git a/Engine/Source/Programs/Horde/HordeServer/Defaults/default-metrics.telemetry.json b/Engine/Source/Programs/Horde/HordeServer/Defaults/default-metrics.telemetry.json
--- a/Engine/Source/Programs/Horde/HordeServer/Defaults/default-metrics.telemetry.json
+++ b/Engine/Source/Programs/Horde/HordeServer/Defaults/default-metrics.telemetry.json
@@ -1,50 +1,19 @@
 {
 	"metrics": [
 		{
-			"id": "editor-performance-cache-warning",
-			"filter": "$[?(@.Payload.EventName == 'Editor.Performance.Warning' && @.Payload.KPI_Category=='Cache' && @.Payload.KPI_DisplayType=='Percent')]",
-			"groupBy": "$.Payload.Build_BranchName,$.Payload.ProjectName,$.Payload.KPI_Name",
-			"property": "$.Payload.KPI_CurrentValue",
-			"function": "Average",
-			"interval": "00:15:00"
-		},
-		{
-			"id": "editor-performance-editor-warning",
-			"filter": "$[?(@.Payload.EventName == 'Editor.Performance.Warning' && @.Payload.KPI_Category=='Editor')]",
-			"groupBy": "$.Payload.Build_BranchName,$.Payload.ProjectName,$.Payload.MapName,$.Payload.KPI_Name",
-			"property": "$.Payload.KPI_CurrentValue",
-			"function": "Average",
-			"interval": "00:15:00"
-		},
-		{
-			"id": "editor-performance-pie-warning",
-			"filter": "$[?(@.Payload.EventName == 'Editor.Performance.Warning' && @.Payload.KPI_Category=='PIE')]",
-			"groupBy": "$.Payload.Build_BranchName,$.Payload.ProjectName,$.Payload.MapName,$.Payload.KPI_Name",
-			"property": "$.Payload.KPI_CurrentValue",
-			"function": "Average",
-			"interval": "00:15:00"
-		},
-		{
-			"id": "editor-performance-warning-count",
-			"filter": "$[?(@.Payload.EventName == 'Editor.Performance.Warning')]",
-			"groupBy": "$.Payload.Build_BranchName,$.Payload.ProjectName,$.Payload.KPI_Category,$.Payload.KPI_Name",
-			"function": "Count",
-			"interval": "00:15:00"
-		},
-		{
 			"id": "core-loading-user-loadingtime",
 			"filter": "$[?(@.Payload.EventName == 'Core.Loading' && @.Payload.Config_IsBuildMachine==false && @.Payload.Config_IsRunningCommandlet==false && @.Payload.LoadingSeconds>5 && @.Payload.LoadingSeconds<3600)]",
 			"property": "$.Payload.LoadingSeconds",
 			"groupBy": "$.Payload.Build_BranchName,$.Payload.ProjectName,$.Payload.LoadingName",
 			"function": "Average",
-			"interval": "00:15:00"
+			"interval": "03:00:00"
 		},
 		{
 			"id": "core-loading-user-loadingcount",
 			"filter": "$[?(@.Payload.EventName == 'Core.Loading' &&@.Payload.Config_IsBuildMachine==false )]",
 			"groupBy": "$.Payload.Build_BranchName,$.Payload.ProjectName,$.Payload.LoadingName",
 			"function": "Count",
-			"interval": "00:15:00"
+			"interval": "03:00:00"
 		},
 		{
 			"id": "core-loading-user-editorboottime",
@@ -52,7 +21,7 @@
 			"property": "$.Payload.LoadingSeconds",
 			"groupBy": "$.Payload.Build_BranchName,$.Payload.ProjectName,$.Payload.ProjectName",
 			"function": "Average",
-			"interval": "00:15:00"
+			"interval": "03:00:00"
 		},
 		{
 			"id": "core-loading-user-editorstartuptime",
@@ -60,7 +29,7 @@
 			"property": "$.Payload.LoadingSeconds",
 			"groupBy": "$.Payload.Build_BranchName,$.Payload.ProjectName,$.Payload.ProjectName",
 			"function": "Average",
-			"interval": "00:15:00"
+			"interval": "03:00:00"
 		},
 		{
 			"id": "core-loading-user-loadmaptime",
@@ -68,14 +37,14 @@
 			"property": "$.Payload.LoadingSeconds",
 			"groupBy": "$.Payload.Build_BranchName,$.Payload.ProjectName,$.Payload.MapName",
 			"function": "Average",
-			"interval": "00:15:00"
+			"interval": "03:00:00"
 		},
 		{
 			"id": "core-loading-user-loadmap-count",
 			"filter": "$[?(@.Payload.EventName == 'Core.Loading' && @.Payload.LoadingName=='LoadMap' && @.Payload.Config_IsBuildMachine==false && @.Payload.Config_IsRunningCommandlet==false && @.Payload.LoadingSeconds>5 && @.Payload.LoadingSeconds<3600)]",
 			"groupBy": "$.Payload.Build_BranchName,$.Payload.ProjectName,$.Payload.MapName",
 			"function": "Count",
-			"interval": "00:15:00"
+			"interval": "03:00:00"
 		},
 		{
 			"id": "core-loading-user-pieloadmaptime",
@@ -83,7 +52,7 @@
 			"property": "$.Payload.LoadingSeconds",
 			"groupBy": "$.Payload.Build_BranchName,$.Payload.ProjectName,$.Payload.MapName",
 			"function": "Average",
-			"interval": "00:15:00"
+			"interval": "03:00:00"
 		},
 		{
 			"id": "core-loading-user-timetopie",
@@ -91,7 +60,7 @@
 			"property": "$.Payload.LoadingSeconds",
 			"groupBy": "$.Payload.Build_BranchName,$.Payload.ProjectName,$.Payload.MapName",
 			"function": "Average",
-			"interval": "00:15:00"
+			"interval": "03:00:00"
 		},
 		{
 			"id": "core-loading-user-piestartuptime",
@@ -99,7 +68,7 @@
 			"property": "$.Payload.LoadingSeconds",
 			"groupBy": "$.Payload.Build_BranchName,$.Payload.ProjectName,$.Payload.MapName",
 			"function": "Average",
-			"interval": "00:15:00"
+			"interval": "03:00:00"
 		},
 		{
 			"id": "core-loading-user-editorboottime-kpi",
@@ -107,7 +76,7 @@
 			"property": "$.Payload.LoadingSeconds",
 			"groupBy": "$.Payload.Build_BranchName,$.Payload.ProjectName",
 			"function": "Average",
-			"interval": "01:00:00"
+			"interval": "03:00:00"
 		},
 		{
 			"id": "core-loading-user-editorstartuptime-kpi",
@@ -115,7 +84,7 @@
 			"property": "$.Payload.LoadingSeconds",
 			"groupBy": "$.Payload.Build_BranchName,$.Payload.ProjectName",
 			"function": "Average",
-			"interval": "01:00:00"
+			"interval": "03:00:00"
 		},
 		{
 			"id": "core-loading-user-loadmaptime-kpi",
@@ -123,7 +92,7 @@
 			"property": "$.Payload.LoadingSeconds",
 			"groupBy": "$.Payload.Build_BranchName,$.Payload.ProjectName",
 			"function": "Average",
-			"interval": "01:00:00"
+			"interval": "03:00:00"
 		},
 		{
 			"id": "core-loading-user-piestartuptime-kpi",
@@ -131,7 +100,7 @@
 			"property": "$.Payload.LoadingSeconds",
 			"groupBy": "$.Payload.Build_BranchName,$.Payload.ProjectName",
 			"function": "Average",
-			"interval": "01:00:00"
+			"interval": "03:00:00"
 		},
 		{
 			"id": "core-loading-user-timetopie-kpi",
@@ -139,7 +108,7 @@
 			"property": "$.Payload.LoadingSeconds",
 			"groupBy": "$.Payload.Build_BranchName,$.Payload.ProjectName",
 			"function": "Average",
-			"interval": "01:00:00"
+			"interval": "03:00:00"
 		},
 		{
 			"id": "core-loading-user-cloudlocalddchitratio",
@@ -147,7 +116,7 @@
 			"property": "$.Payload.DDC_Summary_CloudGetHitPct",
 			"groupBy": "$.Payload.Build_BranchName,$.Payload.ProjectName",
 			"function": "Average",
-			"interval": "01:00:00"
+			"interval": "03:00:00"
 		},
 		{
 			"id": "core-loading-user-zenlocalddchitratio",
@@ -155,7 +124,7 @@
 			"property": "$.Payload.DDC_Summary_ZenLocalGetHitPct",
 			"groupBy": "$.Payload.Build_BranchName,$.Payload.ProjectName",
 			"function": "Average",
-			"interval": "01:00:00"
+			"interval": "03:00:00"
 		},
 		{
 			"id": "core-loading-user-sharedddchitratio",
@@ -163,7 +132,7 @@
 			"property": "$.Payload.DDC_Summary_SharedGetHitPct",
 			"groupBy": "$.Payload.Build_BranchName,$.Payload.ProjectName",
 			"function": "Average",
-			"interval": "01:00:00"
+			"interval": "03:00:00"
 		},
 		{
 			"id": "core-loading-user-totalddchitratio",
@@ -171,7 +140,7 @@
 			"property": "$.Payload.DDC_Summary_TotalGetHitPct",
 			"groupBy": "$.Payload.Build_BranchName,$.Payload.ProjectName",
 			"function": "Average",
-			"interval": "01:00:00"
+			"interval": "03:00:00"
 		},
 		{
 			"id": "core-loading-user-ddc-asset-efficiency",
@@ -179,7 +148,7 @@
 			"property": "$.Payload.Efficiency",
 			"groupBy": "$.Payload.Build_BranchName,$.Payload.ProjectName,$.Payload.AssetType",
 			"function": "Average",
-			"interval": "01:00:00"
+			"interval": "03:00:00"
 		},
 		{
 			"id": "core-loading-user-timetopie-percentile",
@@ -188,7 +157,7 @@
 			"property": "$.Payload.LoadingSeconds",
 			"function": "Percentile",
 			"percentile": 75,
-			"interval": "00:15:00"
+			"interval": "03:00:00"
 		},
 		{
 			"id": "core-loading-user-editorboottime-percentile",
@@ -197,7 +166,7 @@
 			"property": "$.Payload.LoadingSeconds",
 			"function": "Percentile",
 			"percentile": 75,
-			"interval": "00:15:00"
+			"interval": "03:00:00"
 		},
 		{
 			"id": "core-loading-user-editorstartuptime-percentile",
@@ -206,7 +175,7 @@
 			"property": "$.Payload.LoadingSeconds",
 			"function": "Percentile",
 			"percentile": 75,
-			"interval": "00:15:00"
+			"interval": "03:00:00"
 		},
 		{
 			"id": "core-loading-user-loadmaptime-percentile",
@@ -215,7 +184,7 @@
 			"property": "$.Payload.LoadingSeconds",
 			"function": "Percentile",
 			"percentile": 75,
-			"interval": "00:15:00"
+			"interval": "03:00:00"
 		},
 		{
 			"id": "core-loading-user-piestartuptime-percentile",
@@ -224,7 +193,7 @@
 			"property": "$.Payload.LoadingSeconds",
 			"function": "Percentile",
 			"percentile": 75,
-			"interval": "00:15:00"
+			"interval": "03:00:00"
 		},
 		{
 			"id": "core-cooking-standard-cook-time",
@@ -427,5 +396,29 @@
 			"interval": "01:00:00"
 		},
 		{
+			"id": "core-memory-llm-systems",
+			"filter": "$[?(@.Payload.EventName == 'Core.Memory.LLM' && @.Payload.TagSet=='None')]",
+			"groupBy": "$.Payload.Build_BranchName,$.Payload.ProjectName,$.Payload.Hardware_Platform,$.Payload.Context,$.Payload.Name",
+			"property": "$.Payload.Size",
+			"function": "Average",
+			"interval": "00:15:00"
+		},
+		{
+			"id": "core-memory-llm-assetclasses",
+			"filter": "$[?(@.Payload.EventName == 'Core.Memory.LLM' && @.Payload.TagSet=='AssetClasses')]",
+			"groupBy": "$.Payload.Build_BranchName,$.Payload.ProjectName,$.Payload.Hardware_Platform,$.Payload.Context,$.Payload.Name",
+			"property": "$.Payload.Size",
+			"function": "Average",
+			"interval": "00:15:00"
+		},
+		{
+			"id": "core-memory-llm-assets",
+			"filter": "$[?(@.Payload.EventName == 'Core.Memory.LLM' && @.Payload.TagSet=='Assets')]",
+			"groupBy": "$.Payload.Build_BranchName,$.Payload.ProjectName,$.Payload.Hardware_Platform,$.Payload.Context,$.Payload.Name",
+			"property": "$.Payload.Size",
+			"function": "Average",
+			"interval": "00:15:00"
+		},
+		{
 			"id": "horde-incremental-build-jobs",
 			"filter": "$[?(@.Payload.EventName == 'State.JobStepRef' && @.Payload.JobName=='Incremental Build' && search(@.Payload.StepName, \"Incremental\") && @.Payload.Duration>0 && @.Payload.State=='Completed')]",
@@ -563,14 +556,6 @@
 			"interval": "00:30:00"
 		},
 		{
-			"id": "ubt-uba-local-usage",
-			"filter": "$[?(@.Payload.EventName == 'Core.UBT.Executor' && search(@.Payload.Executor, \"Unreal Build Accelerator\") )]",
-			"groupBy": "$.Payload.Build_BranchName,$.Payload.Target,$.Payload.Mode,$.Payload.Configuration,$.Payload.Platform",
-			"property": "$.Payload.LocalUsage",
-			"function": "Average",
-			"interval": "00:30:00"
-		},
-		{
 			"id": "ubt-uba-remote-actions",
 			"filter": "$[?(@.Payload.EventName == 'Core.UBT.Executor' && search(@.Payload.Executor, \"Unreal Build Accelerator\") )]",
 			"groupBy": "$.Payload.Build_BranchName,$.Payload.Target,$.Payload.Mode,$.Payload.Configuration,$.Payload.Platform",
@@ -579,3 +564,11 @@
 			"interval": "00:30:00"
 		},
 		{
+			"id": "ubt-uba-local-usage",
+			"filter": "$[?(@.Payload.EventName == 'Core.UBT.Executor' && search(@.Payload.Executor, \"Unreal Build Accelerator\") )]",
+			"groupBy": "$.Payload.Build_BranchName,$.Payload.Target,$.Payload.Mode,$.Payload.Configuration,$.Payload.Platform",
+			"property": "$.Payload.LocalUsage",
+			"function": "Average",
+			"interval": "00:30:00"
+		},
+		{
@@ -627,13 +620,5 @@
 			"interval": "00:30:00"
 		},
 		{
-			"id": "ubt-uba-cache-rate",
-			"filter": "$[?(@.Payload.EventName == 'Core.UBT.Executor' && search(@.Payload.Executor, \"Unreal Build Accelerator\") )]",
-			"groupBy": "$.Payload.Build_BranchName,$.Payload.Target,$.Payload.Mode,$.Payload.Configuration,$.Payload.Platform",
-			"property": "$.Payload.CacheRate",
-			"function": "Average",
-			"interval": "00:30:00"
-		},
-		{
 			"id": "ubt-uba-remote-usage-all",
 			"filter": "$[?(@.Payload.EventName == 'Core.UBT.Executor' && search(@.Payload.Executor, \"Unreal Build Accelerator\") )]",
@@ -643,14 +628,6 @@
 			"interval": "00:30:00"
 		},
 		{
-			"id": "ubt-uba-remote-actions-all",
-			"filter": "$[?(@.Payload.EventName == 'Core.UBT.Executor' && search(@.Payload.Executor, \"Unreal Build Accelerator\") )]",
-			"groupBy": "$.Payload.Build_BranchName",
-			"property": "$.Payload.RemoteActions",
-			"function": "Sum",
-			"interval": "00:30:00"
-		},
-		{
 			"id": "ubt-uba-total-actions-all",
 			"filter": "$[?(@.Payload.EventName == 'Core.UBT.Executor' && search(@.Payload.Executor, \"Unreal Build Accelerator\") )]",
 			"groupBy": "$.Payload.Build_BranchName",
@@ -658,3 +635,11 @@
 			"function": "Sum",
 			"interval": "00:30:00"
 		},
+		{
+			"id": "ubt-uba-remote-actions-all",
+			"filter": "$[?(@.Payload.EventName == 'Core.UBT.Executor' && search(@.Payload.Executor, \"Unreal Build Accelerator\") )]",
+			"groupBy": "$.Payload.Build_BranchName",
+			"property": "$.Payload.RemoteActions",
+			"function": "Sum",
+			"interval": "00:30:00"
+		},
@@ -667,13 +652,5 @@
 			"interval": "00:30:00"
 		},
 		{
-			"id": "ubt-uba-cache-rate-all",
-			"filter": "$[?(@.Payload.EventName == 'Core.UBT.Executor' && search(@.Payload.Executor, \"Unreal Build Accelerator\") )]",
-			"groupBy": "$.Payload.Build_BranchName",
-			"property": "$.Payload.CacheRate",
-			"function": "Average",
-			"interval": "00:30:00"
-		},
-		{
 			"id": "ubt-completed-duration",
 			"filter": "$[?(@.Payload.EventName == 'Core.UBT.Completed' )]",
diff --git a/Engine/Source/Programs/Horde/HordeServer/Defaults/default-views.telemetry.json b/Engine/Source/Programs/Horde/HordeServer/Defaults/default-views.telemetry.json
--- a/Engine/Source/Programs/Horde/HordeServer/Defaults/default-views.telemetry.json
+++ b/Engine/Source/Programs/Horde/HordeServer/Defaults/default-views.telemetry.json
@@ -78,6 +78,14 @@
 					"defaults": [
 						"Development"
 					]
+				},
+				{
+					"name": "Project",
+					"group": "ProjectName"
+				},
+				{
+					"name": "Hardware_Platform",
+					"group": "Hardware_Platform"
 				}
 			],
 			"categories": [
@@ -220,6 +228,38 @@
 						}
 					]
 				},
+				{
+					"name": "Memory",
+					"charts": [
+						{
+							"name": "Memory By System",
+							"display": "Value",
+							"metrics": [
+								{
+									"id": "core-memory-llm-systems"
+								}
+							]
+						},
+						{
+							"name": "Memory By AssetClass",
+							"display": "Value",
+							"metrics": [
+								{
+									"id": "core-memory-llm-assetclasses"
+								}
+							]
+						},
+						{
+							"name": "Memory By Asset",
+							"display": "Value",
+							"metrics": [
+								{
+									"id": "core-memory-llm-assets"
+								}
+							]
+						}
+					]
+				},
 				{
 					"name": "Cook",
 					"charts": [
