Horde: Ensure agent stopping state is reported to server

Previously:
- Agent session could terminate without reporting status to server
- Led to agents appearing as "ok" in server despite terminated
- Caused repeated UBA agent reassignment due to stale status
- Required server-side cleanup to resolve stale agent states

Now:
- Agent status is consistently reported back during termination
- Prevents dangling agent states in server

diff --git a/Engine/Source/Programs/Horde/HordeAgent.Tests/FakeHordeServer.cs b/Engine/Source/Programs/Horde/HordeAgent.Tests/FakeHordeServer.cs
--- a/Engine/Source/Programs/Horde/HordeAgent.Tests/FakeHordeServer.cs
+++ b/Engine/Source/Programs/Horde/HordeAgent.Tests/FakeHordeServer.cs
@@ -66,12 +66,23 @@ internal class FakeHordeClient(FakeHordeRpcServer server) : IHordeClient
 /// <summary>
 /// Fake implementation of a HordeRpc gRPC server.
 /// Provides a corresponding gRPC client class that can be used with to test client-server interactions inside a single process.
+/// Can only handle a single agent per instance.
 /// </summary>
 internal class FakeHordeRpcServer
 {
 	public readonly TaskCompletionSource<bool> CreateSessionReceived = new();
 	public readonly TaskCompletionSource<bool> UpdateSessionReceived = new();
-
+	
+	/// <summary>
+	/// Last status reported by the agent during a session update
+	/// </summary>
+	public RpcAgentStatus? LastReportedStatus { get; private set; } = null;
+	
+	/// <summary>
+	/// Agent status decided by server. Returned back to agent in session update. 
+	/// </summary>
+	public RpcAgentStatus AgentStatus { get; private set; } = RpcAgentStatus.Ok;
+	
 	private readonly Dictionary<LeaseId, RpcLease> _leases = new();
 	private readonly FakeHordeClient _hordeClient;
 	private readonly ILogger<FakeHordeRpcServer> _logger;
@@ -100,6 +111,11 @@ internal class FakeHordeRpcServer
 		_leases[leaseId.Value] = new RpcLease { Id = leaseId.Value, State = RpcLeaseState.Pending, Payload = Any.Pack(testTask) };
 	}
 	
+	public void SetAgentStatus(RpcAgentStatus newStatus)
+	{
+		AgentStatus = newStatus;
+	}
+	
 	public RpcLease GetLease(LeaseId leaseId)
 	{
 		return _leases[leaseId];
@@ -132,6 +148,15 @@ internal class FakeHordeRpcServer
 		async Task OnRequest(RpcUpdateSessionRequest request)
 		{
 			UpdateSessionReceived.TrySetResult(true);
+			LastReportedStatus = request.Status;
+			
+			switch (request.Status)
+			{
+				case RpcAgentStatus.Ok: break;
+				case RpcAgentStatus.Stopping: SetAgentStatus(RpcAgentStatus.Stopped); break;
+				case RpcAgentStatus.Stopped: break;
+				default: throw new Exception($"Unhandled agent status {request.Status}");
+			}
 			
 			foreach (RpcLease agentLease in request.Leases)
 			{
@@ -143,7 +168,12 @@ internal class FakeHordeRpcServer
 			
 			_logger.LogInformation("OnUpdateSessionRequest: {AgentId} {SessionId} {Status}", request.AgentId, request.SessionId, request.Status);
 			await Task.Delay(100, cancellationToken);
-			RpcUpdateSessionResponse response = new() { ExpiryTime = Timestamp.FromDateTime(DateTime.UtcNow + TimeSpan.FromMinutes(120)) };
+			
+			RpcUpdateSessionResponse response = new()
+			{
+				Status = AgentStatus,
+				ExpiryTime = Timestamp.FromDateTime(DateTime.UtcNow + TimeSpan.FromMinutes(120))
+			};
 			response.Leases.AddRange(_leases.Values.Where(x => x.State != RpcLeaseState.Completed));
 			await responseStream.Write(response);
 		}
diff --git a/Engine/Source/Programs/Horde/HordeAgent.Tests/Leases/LeaseManagerTests.cs b/Engine/Source/Programs/Horde/HordeAgent.Tests/Leases/LeaseManagerTests.cs
--- a/Engine/Source/Programs/Horde/HordeAgent.Tests/Leases/LeaseManagerTests.cs
+++ b/Engine/Source/Programs/Horde/HordeAgent.Tests/Leases/LeaseManagerTests.cs
@@ -14,6 +14,7 @@ using EpicGames.Horde.Agents.Sessions;
 using HordeAgent.Leases;
 using HordeAgent.Services;
 using HordeAgent.Tests.Services;
+using HordeCommon.Rpc;
 using HordeCommon.Rpc.Messages;
 using HordeCommon.Rpc.Tasks;
 using Microsoft.Extensions.Logging;
@@ -78,6 +79,21 @@ public sealed class LeaseManagerTests
 		_leaseManager = CreateLeaseManager(_hordeServer.GetHordeClient());
 	}
 	
+	[TestMethod]
+	public async Task Run_SingleLease_FinishesSuccessfullyAsync()
+	{
+		using CancellationTokenSource cts = new(5000);
+		_hordeServer.ScheduleTestLease();
+		Task<SessionResult> runTask = _leaseManager.RunAsync(cts.Token);
+		_leaseManager.OnLeaseFinished += (lease, result) =>
+		{
+			Assert.AreEqual(RpcAgentStatus.Ok, _hordeServer.LastReportedStatus);
+			_hordeServer.SetAgentStatus(RpcAgentStatus.Stopped);
+		};
+		
+		Assert.AreEqual(new SessionResult(SessionOutcome.BackOff, SessionReason.Completed), await runTask);
+	}
+	
 	[TestMethod]
 	public async Task TerminateGracefully_LeaseIsActive_Async()
 	{
@@ -85,6 +101,7 @@ public sealed class LeaseManagerTests
 		_leaseManager.TerminateSessionAfterLease = true;
 		_hordeServer.ScheduleTestLease();
 		Assert.AreEqual(new SessionResult(SessionOutcome.Terminate, SessionReason.Completed), await _leaseManager.RunAsync(cts.Token));
+		Assert.AreEqual(RpcAgentStatus.Stopped, _hordeServer.AgentStatus);
 	}
 	
 	[TestMethod]
@@ -93,6 +110,7 @@ public sealed class LeaseManagerTests
 		using CancellationTokenSource cts = new(5000);
 		_leaseManager.TerminateSessionAfterLease = true;
 		Assert.AreEqual(new SessionResult(SessionOutcome.Terminate, SessionReason.Completed), await _leaseManager.RunAsync(cts.Token));
+		Assert.AreEqual(RpcAgentStatus.Stopped, _hordeServer.AgentStatus);
 	}
 	
 	/// <summary>
diff --git a/Engine/Source/Programs/Horde/HordeAgent/Leases/LeaseManager.cs b/Engine/Source/Programs/Horde/HordeAgent/Leases/LeaseManager.cs
--- a/Engine/Source/Programs/Horde/HordeAgent/Leases/LeaseManager.cs
+++ b/Engine/Source/Programs/Horde/HordeAgent/Leases/LeaseManager.cs
@@ -29,12 +29,17 @@ namespace HordeAgent.Leases
 		/// <summary>
 		/// Delegate for lease active events
 		/// </summary>
-		public delegate void LeaseActiveEvent(RpcLease lease);
+		public delegate void LeaseEvent(RpcLease lease, LeaseResult? result = null);
 		
 		/// <summary>
-		/// Event triggered when a lease is accepted and set to active
+		/// Event triggered when a lease is accepted and started
 		/// </summary>
-		public event LeaseActiveEvent? OnLeaseActive;
+		public event LeaseEvent? OnLeaseStarted;
+		
+		/// <summary>
+		/// Event triggered when a lease is finished with an outcome
+		/// </summary>
+		public event LeaseEvent? OnLeaseFinished;
 		
 		/// <summary>
 		/// List of pool IDs this session is a member of
@@ -346,7 +351,7 @@ namespace HordeAgent.Leases
 								bool atLeastOneLeaseFinished = _activeLeases.Any(x => x.RpcLease.State is RpcLeaseState.Completed or RpcLeaseState.Cancelled);
 								if (atLeastOneLeaseFinished || _activeLeases.Count == 0)
 								{
-									_logger.LogInformation("At least one lease executed or no lease is active. Terminating session...");
+									_logger.LogInformation("Session termination requested. At least one lease executed or no lease is active, proceeding...");
 									_sessionResult = new SessionResult(SessionOutcome.Terminate, SessionReason.Completed);
 								}
 							}
@@ -385,11 +390,12 @@ namespace HordeAgent.Leases
 											_logger.LogInformation("Lease {LeaseId} is setting session result to {Result}", leaseHandler.Id, result.SessionResult.Outcome);
 											_sessionResult = result.SessionResult;
 										}
+										OnLeaseFinished?.Invoke(serverLease, task.Result);
 										_updateLeasesEvent.Set();
 									}, TaskScheduler.Default);
 
 									_activeLeases.Add(leaseHandler);
-									OnLeaseActive?.Invoke(serverLease);
+									OnLeaseStarted?.Invoke(serverLease);
 								}
 							}
 
@@ -397,17 +403,10 @@ namespace HordeAgent.Leases
 							if (updateSessionResponse.Status == RpcAgentStatus.Stopped)
 							{
 								SessionResult result = _sessionResult ?? new SessionResult(SessionOutcome.BackOff, SessionReason.Completed);
-								_logger.LogInformation("Agent status is stopped; returning from session update loop with result {Result}", result.Outcome);
+								_logger.LogInformation("Agent status is stopped; returning from session update loop with result {Result}", result);
 								return result;
 							}
 						}
-
-						// If there's nothing still running and cancellation was requested, exit
-						if (_activeLeases.Count == 0 && _sessionResult != null)
-						{
-							_logger.LogInformation("No leases are active. Agent is stopping. Session {SessionResult}", _sessionResult); // TODO: Should not really hit this any more; server should report stopping state in lease update above.
-							return _sessionResult;
-						}
 					}
 
 					// Update the current status
