Horde: Use properly typed artifact ids in temp storage log messages.

#rnx

diff --git a/Engine/Source/Programs/Horde/Drivers/JobDriver/Utility/TempStorage.cs b/Engine/Source/Programs/Horde/Drivers/JobDriver/Utility/TempStorage.cs
--- a/Engine/Source/Programs/Horde/Drivers/JobDriver/Utility/TempStorage.cs
+++ b/Engine/Source/Programs/Horde/Drivers/JobDriver/Utility/TempStorage.cs
@@ -630,7 +630,7 @@ namespace JobDriver.Utility
 				NamespaceId namespaceId = new NamespaceId(artifact.NamespaceId);
 				RefName refName = new RefName(artifact.RefName);
 
-				logger.LogInformation("Reading tag \"{NodeName}\":\"{TagName}\" from temp storage (artifact: {ArtifactId} '{ArtifactName}' ({ArtifactType}), ns: {NamespaceId}, ref: {RefName}, localFile: {LocalFile})", nodeName, tagName, artifact.Id, artifactName, artifactType, namespaceId, refName, localFileListLocation);
+				logger.LogInformation("Reading tag \"{NodeName}\":\"{TagName}\" from temp storage (artifact: {ArtifactId} '{ArtifactName}' ({ArtifactType}), ns: {NamespaceId}, ref: {RefName}, localFile: {LocalFile})", nodeName, tagName, ArtifactId.Parse(artifact.Id), artifactName, artifactType, namespaceId, refName, localFileListLocation);
 
 				IStorageNamespace storageNamespace = hordeClient.GetStorageNamespace(namespaceId, artifact.Token);
 				DirectoryNode node = await storageNamespace.ReadRefTargetAsync<DirectoryNode>(artifact.RefName, cancellationToken: cancellationToken);
@@ -768,7 +768,7 @@ namespace JobDriver.Utility
 				NamespaceId namespaceId = new NamespaceId(artifact.NamespaceId);
 				RefName refName = new RefName(artifact.RefName);
 
-				logger.LogInformation("Reading block \"{NodeName}\":\"{BlockName}\" from temp storage (artifact: {ArtifactId} '{ArtifactName}' ({ArtifactType}), ns: {NamespaceId}, ref: {RefName}, local: {LocalFile}, blockdir: {BlockDir})", nodeName, blockName, artifact.Id, artifactName, artifactType, namespaceId, refName, localManifestFile, blockDirectoryName);
+				logger.LogInformation("Reading block \"{NodeName}\":\"{BlockName}\" from temp storage (artifact: {ArtifactId} '{ArtifactName}' ({ArtifactType}), ns: {NamespaceId}, ref: {RefName}, local: {LocalFile}, blockdir: {BlockDir})", nodeName, blockName, ArtifactId.Parse(artifact.Id), artifactName, artifactType, namespaceId, refName, localManifestFile, blockDirectoryName);
 
 				IStorageNamespace storageNamespace = hordeClient.GetStorageNamespace(namespaceId, artifact.Token);
 				DirectoryNode node = await storageNamespace.ReadRefTargetAsync<DirectoryNode>(refName, cancellationToken: cancellationToken);
