Horde: Preserve type URLs in Protobuf messages for upgrades and conforms during serialization

When serializing Protobuf messages for upgrade and conform operations, pack them into Any to preserve type URL information. Without this, attempt counters weren't updating correctly.

Added test coverage for upgrade attempt counting to prevent regressions.

diff --git a/Engine/Source/Programs/Horde/Plugins/Compute/HordeServer.Compute.Tests/Agents/AgentCollectionTests.cs b/Engine/Source/Programs/Horde/Plugins/Compute/HordeServer.Compute.Tests/Agents/AgentCollectionTests.cs
--- a/Engine/Source/Programs/Horde/Plugins/Compute/HordeServer.Compute.Tests/Agents/AgentCollectionTests.cs
+++ b/Engine/Source/Programs/Horde/Plugins/Compute/HordeServer.Compute.Tests/Agents/AgentCollectionTests.cs
@@ -1,12 +1,17 @@
 // Copyright Epic Games, Inc. All Rights Reserved.
 
+extern alias HordeAgent;
+using EpicGames.Horde;
 using EpicGames.Horde.Agents;
 using EpicGames.Horde.Agents.Leases;
 using EpicGames.Horde.Agents.Pools;
 using Google.Protobuf.WellKnownTypes;
+using HordeAgent::HordeCommon.Rpc.Tasks;
 using HordeCommon.Rpc;
 using HordeCommon.Rpc.Messages;
 using HordeServer.Agents;
+using System.Security.Cryptography;
+using Google.Protobuf;
 
 namespace HordeServer.Tests.Agents;
 
@@ -167,7 +172,50 @@ public class AgentCollectionTests : ComputeTestSetup
 		Assert.AreEqual(1, leases.Length);
 		Assert.IsTrue(leases.Contains(_lease1.Id));
 	}
+	
+	[TestMethod]
+	public async Task UpgradeAttemptCount_ShouldTrackUpgradeHistory_Async()
+	{
+		const string CurrentVersion = "v1";
+		const string NewVersion = "v2";
+		
+		// First failing upgrade attempt to v2
+		IAgent? agent = await CreateSessionAsync(_agent, CurrentVersion);
+		agent = await CreateLeaseAsync(agent!, new UpgradeTask { SoftwareId = NewVersion });
+		Assert.AreEqual(CurrentVersion, agent!.Version);
+		Assert.AreEqual(NewVersion, agent.LastUpgradeVersion);
+		Assert.AreEqual(1, agent.UpgradeAttemptCount);
+		agent = await agent.TryTerminateSessionAsync();
+		Assert.IsNotNull(agent);
+		
+		// Second failing upgrade attempt to v2
+		agent = await CreateSessionAsync(agent, CurrentVersion);
+		agent = await CreateLeaseAsync(agent!, new UpgradeTask { SoftwareId = NewVersion });
+		Assert.AreEqual(CurrentVersion, agent!.Version);
+		Assert.AreEqual(NewVersion, agent.LastUpgradeVersion);
+		Assert.AreEqual(2, agent.UpgradeAttemptCount);
+		agent = await agent.TryTerminateSessionAsync();
+		Assert.IsNotNull(agent);
 
+		// Successful upgrade
+		agent = await CreateSessionAsync(agent, NewVersion);
+		Assert.AreEqual(NewVersion, agent!.Version);
+		Assert.AreEqual(NewVersion, agent.LastUpgradeVersion);
+		Assert.IsNull(agent.UpgradeAttemptCount);
+	}
+	
+	private static Task<IAgent?> CreateSessionAsync(IAgent agent, string? version = null)
+	{
+		return agent.TryCreateSessionAsync(new CreateSessionOptions(new RpcAgentCapabilities(), Array.Empty<PoolId>(), version));
+	}
+	
+	private static Task<IAgent?> CreateLeaseAsync(IAgent agent, IMessage? payload = null)
+	{
+		LeaseId leaseId = new (new BinaryId(RandomNumberGenerator.GetBytes(12)));
+		CreateLeaseOptions lease = new(leaseId, null, "lease", null, null, null, null, false, payload ?? new Empty());
+		return agent.TryCreateLeaseAsync(lease);
+	}
+	
 	[TestMethod]
 	public async Task CancelLeaseAsync()
 	{
diff --git a/Engine/Source/Programs/Horde/Plugins/Compute/HordeServer.Compute/Agents/AgentCollection.cs b/Engine/Source/Programs/Horde/Plugins/Compute/HordeServer.Compute/Agents/AgentCollection.cs
--- a/Engine/Source/Programs/Horde/Plugins/Compute/HordeServer.Compute/Agents/AgentCollection.cs
+++ b/Engine/Source/Programs/Horde/Plugins/Compute/HordeServer.Compute/Agents/AgentCollection.cs
@@ -1179,7 +1179,7 @@ namespace HordeServer.Agents
 			while (document != null)
 			{
 				List<UpdateDefinition<AgentDocument>> updates = new List<UpdateDefinition<AgentDocument>>();
-				GetNewLeaseUpdates(document, options.Payload.ToByteArray(), updates);
+				GetNewLeaseUpdates(document, Any.Pack(options.Payload).ToByteArray(), updates);
 
 				if (updates.Count == 0)
 				{
