Horde: Fix an issue where blobs uploaded out of order can cause references to be dropped. Now creates a "shadow" blob info record for any referenced blobs that don't exist, and converts that to a regular record when data is sent.

diff --git a/Engine/Source/Programs/Horde/Plugins/Storage/HordeServer.Storage.Tests/Storage/StorageServiceTests.cs b/Engine/Source/Programs/Horde/Plugins/Storage/HordeServer.Storage.Tests/Storage/StorageServiceTests.cs
--- a/Engine/Source/Programs/Horde/Plugins/Storage/HordeServer.Storage.Tests/Storage/StorageServiceTests.cs
+++ b/Engine/Source/Programs/Horde/Plugins/Storage/HordeServer.Storage.Tests/Storage/StorageServiceTests.cs
@@ -87,6 +87,28 @@ namespace HordeServer.Tests.Storage
 			Assert.AreEqual(1, aliases.Length);
 			Assert.AreEqual(handle2.GetLocator(), aliases[0].Target.GetLocator());
 		}
+
+		[TestMethod]
+		public async Task ForwardDeclareAsync()
+		{
+			StorageService storageService = ServiceProvider.GetRequiredService<StorageService>();
+			NamespaceId namespaceId = new NamespaceId("memory");
+
+			StorageService.BlobInfo? blobInfo = await storageService.FindBlobAsync(namespaceId, new BlobLocator("foo"));
+			Assert.IsNull(blobInfo);
+
+			blobInfo = await storageService.AddBlobAsync(namespaceId, new BlobLocator("foo"), [new BlobLocator("bar")]);
+			Assert.AreEqual(1, blobInfo.Imports?.Count ?? 0);
+			Assert.IsFalse(blobInfo.Shadow);
+
+			StorageService.BlobInfo? importInfo = await storageService.FindBlobAsync(namespaceId, new BlobLocator("bar"));
+			Assert.IsNotNull(importInfo);
+			Assert.IsTrue(importInfo.Shadow);
+
+			StorageService.BlobInfo? importInfo2 = await storageService.AddBlobAsync(namespaceId, new BlobLocator("bar"), []);
+			Assert.AreEqual(importInfo.Id, importInfo2.Id);
+			Assert.IsFalse(blobInfo.Shadow);
+		}
 	}
 }
 
diff --git a/Engine/Source/Programs/Horde/Plugins/Storage/HordeServer.Storage/Storage/StorageService.cs b/Engine/Source/Programs/Horde/Plugins/Storage/HordeServer.Storage/Storage/StorageService.cs
--- a/Engine/Source/Programs/Horde/Plugins/Storage/HordeServer.Storage/Storage/StorageService.cs
+++ b/Engine/Source/Programs/Horde/Plugins/Storage/HordeServer.Storage/Storage/StorageService.cs
@@ -350,6 +350,9 @@ namespace HordeServer.Storage
 			[BsonElement("ali"), BsonIgnoreIfNull]
 			public List<AliasInfo>? Aliases { get; set; }
 
+			[BsonElement("shd"), BsonIgnoreIfDefault]
+			public bool Shadow { get; set; } // Blob has not been added yet, but is referenced by something else
+
 			[BsonElement("del"), BsonIgnoreIfDefault]
 			public int GcVersion { get; set; }
 
@@ -781,30 +784,61 @@ namespace HordeServer.Storage
 		#region Blobs
 
 		/// <inheritdoc/>
-		async Task AddBlobAsync(NamespaceId namespaceId, BlobLocator locator, IReadOnlyCollection<BlobLocator> imports, List<AliasInfo>? exports = null, CancellationToken cancellationToken = default)
+		internal async Task<BlobInfo> AddBlobAsync(NamespaceId namespaceId, BlobLocator locator, IReadOnlyCollection<BlobLocator> imports, List<AliasInfo>? exports = null, CancellationToken cancellationToken = default)
 		{
-			ObjectId id = ObjectId.GenerateNewId(_clock.UtcNow);
-			BlobInfo blobInfo = new BlobInfo(id, namespaceId, locator);
-			blobInfo.Aliases = exports;
+			List<ObjectId> importIds = new List<ObjectId>();
+			if (imports.Count > 0)
+			{
+				foreach (BlobLocator import in imports)
+				{
+					ObjectId importId = await FindOrAddShadowBlobAsync(namespaceId, import, cancellationToken);
+					importIds.Add(importId);
+				}
+			}
 
-			if (imports.Count == 0)
+			FilterDefinition<BlobInfo> insertFilter =
+				Builders<BlobInfo>.Filter.Eq(x => x.NamespaceId, namespaceId) &
+				Builders<BlobInfo>.Filter.Eq(x => x.Path, locator.ToString()) &
+				Builders<BlobInfo>.Filter.Eq(x => x.Shadow, true);
+
+			UpdateDefinition<BlobInfo> insertUpdate = Builders<BlobInfo>.Update
+				.Unset(x => x.Shadow);
+
+			if (importIds.Count > 0)
 			{
-				blobInfo.Imports = new List<ObjectId>();
+				insertUpdate = insertUpdate.Set(x => x.Imports, importIds);
 			}
-			else
+			if (exports != null && exports.Count > 0)
 			{
-				List<string> paths = imports.Select(x => x.BaseLocator.ToString()).ToList();
+				insertUpdate = insertUpdate.Set(x => x.Aliases, exports);
+			}
+
+			BlobInfo blobInfo = await _blobCollection.FindOneAndUpdateAsync(insertFilter, insertUpdate, new FindOneAndUpdateOptions<BlobInfo, BlobInfo> { IsUpsert = true, ReturnDocument = ReturnDocument.After }, cancellationToken);
+			_logger.LogDebug("Created blob {BlobId} at {Path} ({NumImports} imports)", blobInfo.Id, blobInfo.Path, imports.Count);
+			return blobInfo;
+		}
 
-				FilterDefinition<BlobInfo> filter =
-					Builders<BlobInfo>.Filter.Eq(x => x.NamespaceId, namespaceId) &
-					Builders<BlobInfo>.Filter.In(x => x.Path, paths);
+		async Task<ObjectId> FindOrAddShadowBlobAsync(NamespaceId namespaceId, BlobLocator locator, CancellationToken cancellationToken = default)
+		{
+			FilterDefinition<BlobInfo> filter =
+				Builders<BlobInfo>.Filter.Eq(x => x.NamespaceId, namespaceId) &
+				Builders<BlobInfo>.Filter.Eq(x => x.Path, locator.ToString());
 
-				blobInfo.Imports = await _blobCollection.Find(filter).Project(x => x.Id).ToListAsync(cancellationToken);
-				blobInfo.Imports.Sort();
-			}
+			UpdateDefinition<BlobInfo> update =
+				Builders<BlobInfo>.Update.SetOnInsert(x => x.Shadow, true);
+
+			BlobInfo result = await _blobCollection.FindOneAndUpdateAsync(filter, update, new FindOneAndUpdateOptions<BlobInfo, BlobInfo> { IsUpsert = true, ReturnDocument = ReturnDocument.After }, cancellationToken);
+			return result.Id;
+		}
 
-			await _blobCollection.InsertOneAsync(blobInfo, new InsertOneOptions { }, cancellationToken);
-			_logger.LogDebug("Created blob {BlobId} at {Path} ({NumImports} imports)", blobInfo.Id, blobInfo.Path, blobInfo.Imports.Count);
+		internal async Task<BlobInfo?> GetBlobAsync(ObjectId id, CancellationToken cancellationToken = default)
+		{
+			return await _blobCollection.Find(x => x.Id == id).FirstOrDefaultAsync(cancellationToken);
+		}
+
+		internal async Task<BlobInfo?> FindBlobAsync(NamespaceId namespaceId, BlobLocator locator, CancellationToken cancellationToken = default)
+		{
+			return await _blobCollection.Find(x => x.NamespaceId == namespaceId && x.Path == locator.ToString()).FirstOrDefaultAsync(cancellationToken);
 		}
 
 		async ValueTask CheckBlobExistsAsync(NamespaceId namespaceId, BlobLocator locator, CancellationToken cancellationToken)
