Horde: Remove EC2 IMDS server availability check during startup

- Let failed (intermittent) HTTP requests return an unknown instance state and retry in the instance monitor loop.
- EC2 instance lifecycle monitoring should only be enabled in AWS environments in the first place.
- This may fix a bug where spot interruptions aren't properly detected. One theory is that IMDS server availability failures during initialization disable the monitoring completely.

diff --git a/Engine/Source/Programs/Horde/HordeAgent/Services/AwsInstanceLifecycleService.cs b/Engine/Source/Programs/Horde/HordeAgent/Services/AwsInstanceLifecycleService.cs
--- a/Engine/Source/Programs/Horde/HordeAgent/Services/AwsInstanceLifecycleService.cs
+++ b/Engine/Source/Programs/Horde/HordeAgent/Services/AwsInstanceLifecycleService.cs
@@ -15,6 +15,12 @@ namespace HordeAgent.Services;
 /// </summary>
 public enum Ec2InstanceState
 {
+	/// <summary>
+	/// Unknown state
+	/// Likely because of a failed HTTP request to local IMDS server
+	/// </summary>
+	Unknown,
+	
 	/// <summary>
 	/// Normal state
 	/// </summary>
@@ -106,23 +112,35 @@ class AwsInstanceLifecycleService : BackgroundService
 
 	private async Task<Ec2InstanceState> GetStateAsync(CancellationToken cancellationToken)
 	{
-		HttpResponseMessage spotRes = await _httpClient.GetAsync(new Uri(BaseUri + "/spot/instance-action"), cancellationToken);
-		if (spotRes.StatusCode == HttpStatusCode.OK)
-		{
-			return Ec2InstanceState.TerminatingSpot;
-		}
-
-		HttpResponseMessage asgRes = await _httpClient.GetAsync(new Uri(BaseUri + "/autoscaling/target-lifecycle-state"), cancellationToken);
-		if (asgRes.StatusCode == HttpStatusCode.OK)
+		try
 		{
-			string state = await asgRes.Content.ReadAsStringAsync(cancellationToken);
-			if (state == "Terminated")
+			HttpResponseMessage spotRes = await _httpClient.GetAsync(new Uri(BaseUri + "/spot/instance-action"), cancellationToken);
+			if (spotRes.StatusCode == HttpStatusCode.OK)
 			{
-				return Ec2InstanceState.TerminatingAsg;
+				return Ec2InstanceState.TerminatingSpot;
 			}
+			
+			HttpResponseMessage asgRes = await _httpClient.GetAsync(new Uri(BaseUri + "/autoscaling/target-lifecycle-state"), cancellationToken);
+			if (asgRes.StatusCode == HttpStatusCode.OK)
+			{
+				string state = await asgRes.Content.ReadAsStringAsync(cancellationToken);
+				if (state == "Terminated")
+				{
+					return Ec2InstanceState.TerminatingAsg;
+				}
+			}
+			
+			return Ec2InstanceState.InService;
+		}
+		catch (Exception e)
+		{
+			// HTTP timeouts from the IMDS server can occur, in addition to unknown exceptions.
+			// Unclear why timeouts are seen in the first place from the metadata server.
+			// It should not be under load at all but it's possible an agent performing heavy work can affect this.
+			
+			_logger.LogWarning("Request to IMDS server failed. Reason={Reason}", e.Message);
+			return Ec2InstanceState.Unknown;
 		}
-
-		return Ec2InstanceState.InService;
 	}
 
 	private async Task<bool> IsSpotInstanceAsync(CancellationToken cancellationToken)
@@ -137,35 +155,15 @@ class AwsInstanceLifecycleService : BackgroundService
 		return false;
 	}
 
-	private async Task<bool> IsImdsAvailableAsync(CancellationToken cancellationToken)
-	{
-		try
-		{
-			HttpResponseMessage res = await _httpClient.GetAsync(new Uri(BaseUri + "/"), cancellationToken);
-			return res.StatusCode == HttpStatusCode.OK;
-		}
-		catch (Exception)
-		{
-			// Timed out or other error. Can safely assume the metadata server is not available.
-			return false;
-		}
-	}
-
 	internal async Task MonitorInstanceLifecycleAsync(CancellationToken cancellationToken)
 	{
-		if (!await IsImdsAvailableAsync(cancellationToken))
-		{
-			_logger.LogInformation("EC2 metadata server (IMDS) not available. Will not monitor EC2 lifecycle state");
-			return;
-		}
-
 		_logger.LogInformation("Monitoring EC2 instance lifecycle state...");
 		while (!cancellationToken.IsCancellationRequested)
 		{
 			try
 			{
 				Ec2InstanceState state = await GetStateAsync(cancellationToken);
-				if (state != Ec2InstanceState.InService)
+				if (state is Ec2InstanceState.TerminatingAsg or Ec2InstanceState.TerminatingSpot)
 				{
 					bool isSpot = await IsSpotInstanceAsync(cancellationToken);
 					TimeSpan ttl = GetTimeToLive(state);
@@ -182,12 +180,6 @@ class AwsInstanceLifecycleService : BackgroundService
 					return;
 				}
 			}
-			catch (TaskCanceledException) when (!cancellationToken.IsCancellationRequested)
-			{
-				// Ignore HTTP timeouts from the IMDS server and let loop try again
-				// Unclear why timeouts are seen in the first place from the metadata server.
-				// It should not be under load at all but it's possible an agent performing heavy work can affect this.
-			}
 			catch (Exception e)
 			{
 				_logger.LogError(e, "Unhandled exception during EC2 instance monitoring");
