Horde: Update release notes.

#rnx

diff --git a/Engine/Source/Programs/Horde/Docs/ReleaseNotes.md b/Engine/Source/Programs/Horde/Docs/ReleaseNotes.md
--- a/Engine/Source/Programs/Horde/Docs/ReleaseNotes.md
+++ b/Engine/Source/Programs/Horde/Docs/ReleaseNotes.md
@@ -1,5 +1,14 @@
 # Release Notes
 
+## 2024-11-15
+
+* Prevent exceeding open file handle limit on MacOS when uploading data. (38125783)
+* Remove EC2 IMDS server availability check during startup (38123434)
+* Fix an issue where blobs uploaded out of order can cause references to be dropped. Now creates a "shadow" blob info record for any referenced blobs that don't exist, and converts that to a regular record when data is sent. (38121174)
+* Preserve type URLs in Protobuf messages for upgrades and conforms during serialization. (38119696)
+* Fix test flakiness in ComputeSocketTests. Increase test timeouts and disable parallelization to handle CPU contention better during parallel test execution. (38088741)
+* Ensure agent stopping state is reported to server (38085676)
+
 ## 2024-11-14
 
 * Handle EndOfStreamException in RecvOptionalAsync. Treat EndOfStreamException same as zero bytes read, returning false to indicate no data available rather than propagating the exception. (38077833)
