Horde: Add a warning message to detect blobs which are referenced but never uploaded.

#rnx

diff --git a/Engine/Source/Programs/Horde/Plugins/Storage/HordeServer.Storage/Storage/StorageService.cs b/Engine/Source/Programs/Horde/Plugins/Storage/HordeServer.Storage/Storage/StorageService.cs
--- a/Engine/Source/Programs/Horde/Plugins/Storage/HordeServer.Storage/Storage/StorageService.cs
+++ b/Engine/Source/Programs/Horde/Plugins/Storage/HordeServer.Storage/Storage/StorageService.cs
@@ -930,6 +930,10 @@ namespace HordeServer.Storage
 				_logger.LogDebug("Adding {NumBlobs} blobs for GC consideration ({FirstId} to {LastId})", current.Count, current[0].Id, current[^1].Id);
 				foreach (BlobInfo blobInfo in current)
 				{
+					if (blobInfo.Shadow)
+					{
+						_logger.LogWarning("Referenced blob {NamespaceId} {Path} was never uploaded.", blobInfo.NamespaceId, blobInfo.Path); 
+					}
 					AddGcCheckRecord(blobInfo.NamespaceId, blobInfo.Id, latestInfoId);
 				}
 
