Horde: Annotate list of changes printed by release notes to indicate whether changes affect the server, agent, or dashboard.

#rnx

diff --git a/Engine/Source/Programs/Horde/Tools/ReleaseNotes/Program.cs b/Engine/Source/Programs/Horde/Tools/ReleaseNotes/Program.cs
--- a/Engine/Source/Programs/Horde/Tools/ReleaseNotes/Program.cs
+++ b/Engine/Source/Programs/Horde/Tools/ReleaseNotes/Program.cs
@@ -18,6 +18,26 @@ namespace CreateReleaseNotes
 			$"Engine/Source/Programs/Shared/...",
 		};
 
+		static readonly string[] s_serverFilter =
+		{
+			"...",
+			"-/Engine/Source/Programs/Horde/HordeDashboard/...",
+			"-/Engine/Source/Programs/Horde/HordeAgent/...",
+		};
+
+		static readonly string[] s_agentFilter =
+		{
+			"...",
+			"-/Engine/Source/Programs/Horde/HordeDashboard/...",
+			"-/Engine/Source/Programs/Horde/HordeServer/...",
+		};
+
+		static readonly string[] s_dashboardFilter =
+		{
+			"-/...",
+			"/Engine/Source/Programs/Horde/HordeDashboard/...",
+		};
+
 		record class ChangeInfo(int Number, string Description, string? JiraTicket);
 
 		class Options
@@ -25,8 +45,17 @@ namespace CreateReleaseNotes
 			[CommandLine("-Server=", Required = true)]
 			public Uri Server { get; set; } = null!;
 
-			[CommandLine("-MinChange=")]
-			public int? MinChange { get; set; }
+			[CommandLine("-Change=")]
+			public int? Change { get; set; }
+
+			[CommandLine("-ServerChange=")]
+			public int? ServerChange { get; set; }
+
+			[CommandLine("-ClientChange=")]
+			public int? AgentChange { get; set; }
+
+			[CommandLine("-DashboardChange=")]
+			public int? DashboardChange { get; set; }
 		}
 
 		static async Task<int> Main(string[] args)
@@ -51,24 +80,37 @@ namespace CreateReleaseNotes
 			IHordeClient horde = serviceProvider.GetRequiredService<IHordeClient>();
 			HordeHttpClient hordeHttpClient = horde.CreateHttpClient();
 
-			int minChange = options.MinChange ?? 0;
-			if (minChange == 0)
+			int serverChange = options.ServerChange ?? options.Change ?? 0;
+			int agentChange = options.AgentChange ?? options.Change ?? 0;
+
+			if (serverChange == 0 || agentChange == 0)
 			{
 				logger.LogInformation("Querying Horde for currently deployed version...");
-				GetServerInfoResponse info = await hordeHttpClient.GetServerInfoAsync();
 
-				Match match = Regex.Match(info.ServerVersion, @"-(\d+)$");
-				if (!match.Success)
+				GetServerInfoResponse info = await hordeHttpClient.GetServerInfoAsync();
+				if (serverChange == 0 && !TryParseMinChange(info.ServerVersion, out serverChange))
 				{
 					logger.LogError("Unexpected server version format: {Version}", info.ServerVersion);
 					return 1;
 				}
-
-				logger.LogInformation("Current version: {Version}", info.ServerVersion);
-				minChange = Int32.Parse(match.Groups[1].Value) + 1;
+				if (agentChange == 0 && !TryParseMinChange(info.AgentVersion, out agentChange))
+				{
+					logger.LogError("Unexpected agent version format: {Version}", info.AgentVersion);
+					return 1;
+				}
 			}
 
-			logger.LogInformation("Finding changes after CL {MinChange}...", minChange);
+			int dashboardChange = options.DashboardChange ?? serverChange;
+
+			logger.LogInformation("");
+			logger.LogInformation("Current versions:");
+			logger.LogInformation("  Server:    {ServerChange}", serverChange);
+			logger.LogInformation("  Agent:     {AgentChange}", agentChange);
+			logger.LogInformation("  Dashboard: {DashboardChange}", dashboardChange);
+			logger.LogInformation("");
+
+			int change = Math.Min(serverChange, agentChange);
+			logger.LogInformation("Finding changes after CL {MinChange}...", change);
 			using IPerforceConnection perforce = await PerforceConnection.CreateAsync(serviceProvider.GetRequiredService<ILogger<PerforceConnection>>());
 
 			List<ChangeInfo> parsedChanges = new List<ChangeInfo>();
@@ -77,38 +119,71 @@ namespace CreateReleaseNotes
 
 			List<string> paths = s_paths.Select(x => $"//{perforceInfo.ClientName}/{x}").ToList();
 
-			List<ChangesRecord> changes = await perforce.GetChangesAsync(ChangesOptions.LongOutput, null, minChange, -1, ChangeStatus.Submitted, null, paths);
+			List<ChangesRecord> changes = await perforce.GetChangesAsync(ChangesOptions.None, null, change + 1, -1, ChangeStatus.Submitted, null, paths);
+			changes.RemoveAll(x => x.User.Equals("buildmachine", StringComparison.OrdinalIgnoreCase));
+			changes = changes.DistinctBy(x => x.Number).ToList();
+
+			List<DescribeRecord> describeRecords = await perforce.DescribeAsync(changes.Select(x => x.Number).ToArray());
+
 			HashSet<int> changeNumbers = new HashSet<int>();
 
+			FileFilter serverFilter = new FileFilter(s_serverFilter);
+			FileFilter agentFilter = new FileFilter(s_agentFilter);
+			FileFilter dashboardFilter = new FileFilter(s_dashboardFilter);
+
 			logger.LogInformation("");
-			foreach (ChangesRecord change in changes)
+			foreach (DescribeRecord describeRecord in describeRecords.OrderBy(x => x.Number))
 			{
-				if (changeNumbers.Add(change.Number))
+				string description = describeRecord.Description;
+
+				Match jiraMatch = Regex.Match(description, @"^\s*#jira ([a-zA-Z]+-[0-9]+.*)$");
+				string? jiraTicket = jiraMatch.Success ? jiraMatch.Groups[1].Value : null;
+
+				description = Regex.Replace(description, @"^[a-zA-Z]+:\s*", "");
+				description = Regex.Replace(description, @"^\s*#.*$", "", RegexOptions.Multiline);
+				description = Regex.Replace(description, @"\n\s*", "\n");
+				description = description.Trim();
+
+				if (!Regex.IsMatch(describeRecord.Description, @"^\s*#rnx\s*$", RegexOptions.Multiline))
 				{
-					string description = change.Description;
+					parsedChanges.Add(new ChangeInfo(describeRecord.Number, description, jiraTicket));
+				}
 
-					Match jiraMatch = Regex.Match(description, @"^\s*#jira ([a-zA-Z]+-[0-9]+.*)$");
-					string? jiraTicket = jiraMatch.Success ? jiraMatch.Groups[1].Value : null;
+				string logDescription = Regex.Replace(description, @"\s+", " ");
 
-					description = Regex.Replace(description, @"^[a-zA-Z]+:\s*", "");
-					description = Regex.Replace(description, @"^\s*#.*$", "", RegexOptions.Multiline);
-					description = Regex.Replace(description, @"\n\s*", "\n");
-					description = description.Trim();
+				const int MaxDescriptionLength = 80;
+				if (logDescription.Length > MaxDescriptionLength)
+				{
+					logDescription = logDescription.Substring(0, MaxDescriptionLength);
+				}
 
-					if (!Regex.IsMatch(change.Description, @"^\s*#rnx\s*$", RegexOptions.Multiline))
-					{
-						parsedChanges.Add(new ChangeInfo(change.Number, description, jiraTicket));
-					}
+				bool affectsServer = false;
+				bool affectsAgent = false;
+				bool affectsDashboard = false;
 
-					string logDescription = Regex.Replace(description, @"\s+", " ");
+				foreach (DescribeFileRecord fileRecord in describeRecord.Files)
+				{
+					string path = fileRecord.DepotFile;
+					path = Regex.Replace(path, @"^//[^/]+/[^/]+/", "");
 
-					const int MaxDescriptionLength = 80;
-					if (logDescription.Length > MaxDescriptionLength)
+					if (describeRecord.Number > serverChange && serverFilter.Matches(path))
+					{
+						affectsServer = true;
+					}
+					if (describeRecord.Number > agentChange && agentFilter.Matches(path))
+					{
+						affectsAgent = true;
+					}
+					if (describeRecord.Number > dashboardChange && dashboardFilter.Matches(path))
 					{
-						logDescription = logDescription.Substring(0, MaxDescriptionLength);
+						affectsDashboard = true;
 					}
+				}
 
-					logger.LogInformation("{Change} {Author,-20} {Description}", change.Number, change.User.ToLower(), logDescription);
+				if (affectsServer || affectsAgent || affectsDashboard)
+				{
+					string types = (affectsServer? "S" : " ") + (affectsAgent ? "A" : " ") + (affectsDashboard? "D" : " ");
+					logger.LogInformation("[{Types}] {Change} {Author,-20} {Description}", types, describeRecord.Number, describeRecord.User.ToLower(), logDescription);
 				}
 			}
 
@@ -157,5 +232,26 @@ namespace CreateReleaseNotes
 			await FileReference.WriteAllLinesAsync(file, fileLines);
 			return 0;
 		}
+
+		static bool TryParseMinChange(string? name, out int version)
+		{
+			if (name == null)
+			{
+				version = 0;
+				return false;
+			}
+
+			Match match = Regex.Match(name, @"-(\d+)$");
+			if (!match.Success)
+			{
+				version = 0;
+				return false;
+			}
+			else
+			{
+				version = Int32.Parse(match.Groups[1].Value) + 1;
+				return true;
+			}
+		}
 	}
 }
