Removed extra Project Name field in Horde Analytics dashboard

diff --git a/Engine/Source/Programs/Horde/HordeServer/Defaults/default-views.telemetry.json b/Engine/Source/Programs/Horde/HordeServer/Defaults/default-views.telemetry.json
--- a/Engine/Source/Programs/Horde/HordeServer/Defaults/default-views.telemetry.json
+++ b/Engine/Source/Programs/Horde/HordeServer/Defaults/default-views.telemetry.json
@@ -79,10 +79,6 @@
 						"Development"
 					]
 				},
-				{
-					"name": "Project",
-					"group": "ProjectName"
-				},
 				{
 					"name": "Hardware_Platform",
 					"group": "Hardware_Platform"
