Horde: Add CLI argument for loading additional config files

Allows specifying custom config file paths, enabling more flexible configuration in buildfarm environments instead of relying solely on built-in default locations.

diff --git a/Engine/Source/Programs/Horde/HordeAgent/AgentApp.cs b/Engine/Source/Programs/Horde/HordeAgent/AgentApp.cs
--- a/Engine/Source/Programs/Horde/HordeAgent/AgentApp.cs
+++ b/Engine/Source/Programs/Horde/HordeAgent/AgentApp.cs
@@ -111,11 +111,12 @@ namespace HordeAgent
 				configOverrides.Add($"{AgentSettings.SectionName}:{nameof(AgentSettings.WorkingDir)}", workingDirOverride);
 			}
 			
+			List<FileReference> configFilesFromArgs = SplitPath(arguments.GetStringOrDefault("-ConfigFiles=", ""));
 			List<string> configSourcesInstalled = [];
 			
-			// Create the base configuration data by just reading from the application directory. We need to check some settings before
-			// being able to read user configuration files.
-			IConfiguration configuration = CreateConfig(false, null, configOverrides, out _);
+			// Create the base configuration data by just reading from the application directory.
+			// We need to check some settings before being able to read user configuration files.
+			IConfiguration configuration = CreateConfig(false, [], configOverrides, out _);
 			AgentSettings settings = BindSettings(configuration);
 
 			if (settings.Installed)
@@ -129,9 +130,10 @@ namespace HordeAgent
 						await CopyDefaultConfigFilesAsync(DataDir);
 					}
 				}
-
-				FileReference agentConfigFile = FileReference.Combine(DataDir, "agent.json");
-				configuration = CreateConfig(true, agentConfigFile, configOverrides, out configSourcesInstalled);
+				
+				List<FileReference> agentConfigFiles = [FileReference.Combine(DataDir, "agent.json")];
+				agentConfigFiles.AddRange(configFilesFromArgs);
+				configuration = CreateConfig(true, agentConfigFiles, configOverrides, out configSourcesInstalled);
 				settings = BindSettings(configuration);
 			}
 
@@ -248,7 +250,7 @@ namespace HordeAgent
 			return new DefaultSystemMetrics();
 		}
 
-		static IConfiguration CreateConfig(bool readInstalledConfig, FileReference? agentConfigFile, Dictionary<string, string?> configOverrides, out List<string> configSources)
+		static IConfiguration CreateConfig(bool readInstalledConfig, List<FileReference> agentConfigFiles, Dictionary<string, string?> configOverrides, out List<string> configSources)
 		{
 			configSources = new List<string>();
 			string? environment = Environment.GetEnvironmentVariable("DOTNET_ENVIRONMENT");
@@ -273,15 +275,15 @@ namespace HordeAgent
 				.AddJsonFile($"appsettings.{environment}.json", optional: true) // environment variable overrides, also used in k8s setups with Helm
 				.AddJsonFile("appsettings.User.json", optional: true);
 			
-			if (agentConfigFile != null && Path.Exists(agentConfigFile.FullName))
+			foreach (FileReference configFile in agentConfigFiles.Where(path => Path.Exists(path.FullName)))
 			{
 				// Adding a JSON file outside the common base path requires adding a completely separate JsonConfigurationSource entry
 				builder.Add(new JsonConfigurationSource()
 				{
-					Path = agentConfigFile.GetFileName(),
+					Path = configFile.GetFileName(),
 					Optional = true,
 					ReloadOnChange = true,
-					FileProvider = new PhysicalFileProvider(agentConfigFile.Directory.FullName),
+					FileProvider = new PhysicalFileProvider(configFile.Directory.FullName),
 				});
 			}
 			
@@ -355,5 +357,27 @@ namespace HordeAgent
 			// When C# project is packaged as a single file, GetExecutingAssembly above does not work
 			return DirectoryReference.FromFile(new FileReference(Environment.ProcessPath!));
 		}
+		
+		/// <summary>
+		/// Splits a string containing multiple file paths
+		/// Uses platform-specific path separators: semicolon (;) on Windows, colon (:) on other operating systems.
+		/// </summary>
+		/// <param name="paths">String containing one or more paths separated by the platform-specific separator</param>
+		/// <returns>List of FileReference objects, empty list if input is null or empty</returns>
+		internal static List<FileReference> SplitPath(string paths)
+		{
+			if (String.IsNullOrEmpty(paths))
+			{
+				return new List<FileReference>();
+			}
+			
+			char separator = OperatingSystem.IsWindows() ? ';' : ':';
+			return paths
+				.Split(new[] { separator }, StringSplitOptions.RemoveEmptyEntries)
+				.Select(path => path.Trim())
+				.Where(path => !String.IsNullOrEmpty(path))
+				.Select(path => new FileReference(path))
+				.ToList();
+		}
 	}
 }
diff --git a/Engine/Source/Programs/Horde/HordeAgent/Commands/Service/ServiceInstall.cs b/Engine/Source/Programs/Horde/HordeAgent/Commands/Service/ServiceInstall.cs
--- a/Engine/Source/Programs/Horde/HordeAgent/Commands/Service/ServiceInstall.cs
+++ b/Engine/Source/Programs/Horde/HordeAgent/Commands/Service/ServiceInstall.cs
@@ -13,6 +13,7 @@ namespace HordeAgent.Commands.Service
 		string? UserName = null,
 		string? Password = null,
 		string? ServerProfile = null,
+		string? ConfigFiles = null,
 		string DotNetExecutable = "dotnet",
 		bool StartService = true
 	);
@@ -47,6 +48,10 @@ namespace HordeAgent.Commands.Service
 		[CommandLine("-Start=")]
 		[Description("Whether to start the service after installation (true/false)")]
 		public string? Start { get; set; } = "true";
+		
+		[CommandLine("-ConfigFiles=")]
+		[Description("Paths to additional config JSON files to load, separated by colon (Linux/macOS) or semi-colon (Windows)")]
+		public string? ConfigFiles { get; set; } = null;
 
 		/// <summary>
 		/// Installs the service
@@ -126,6 +131,11 @@ namespace HordeAgent.Commands.Service
 			{
 				commandLine.Append($" -Server={options.ServerProfile}");
 			}
+			
+			if (options.ConfigFiles != null)
+			{
+				commandLine.Append($" -ConfigFiles=\"{options.ConfigFiles}\"");
+			}
 
 			using (WindowsService service = serviceManager.Create(ServiceName, "Horde Agent", commandLine.ToString(), options.UserName, options.Password))
 			{
diff --git a/Engine/Source/Programs/Horde/HordeAgent/Commands/Service/ServiceRun.cs b/Engine/Source/Programs/Horde/HordeAgent/Commands/Service/ServiceRun.cs
--- a/Engine/Source/Programs/Horde/HordeAgent/Commands/Service/ServiceRun.cs
+++ b/Engine/Source/Programs/Horde/HordeAgent/Commands/Service/ServiceRun.cs
@@ -19,6 +19,11 @@ namespace HordeAgent.Commands.Service
 		[CommandLine("-LogLevel")]
 		[Description("Log verbosity level (use normal MS levels such as debug, warning or information)")]
 		public string LogLevelStr { get; set; } = "information";
+		
+		[CommandLine("-ConfigFiles=")]
+		[Description("Paths to additional config JSON files to load, separated by colon (Linux/macOS) or semi-colon (Windows)")]
+		// This property only provides help text for the CLI interface. The actual config files are processed earlier in AgentApp.cs.
+		public string? ConfigFiles { get; set; } = null;
 
 		readonly DefaultServices _defaultServices;
 
