Horde: Add download with Unreal Toolbox preference and generalized artifact button
#skipci

diff --git a/Engine/Source/Programs/Horde/HordeDashboard/src/backend/Api.ts b/Engine/Source/Programs/Horde/HordeDashboard/src/backend/Api.ts
--- a/Engine/Source/Programs/Horde/HordeDashboard/src/backend/Api.ts
+++ b/Engine/Source/Programs/Horde/HordeDashboard/src/backend/Api.ts
@@ -1322,7 +1322,7 @@ export type GetJobResponse = {
 	/** Custom list of targets for the job, will take precedence over all argument and parameter targets */
 	targets?: string[];
 
-	/**  Whether the current user is allowed to update this job */		
+	/**  Whether the current user is allowed to update this job */
 	canUpdate?: boolean;
 }
 
@@ -1655,6 +1655,29 @@ export type GetArtifactZipRequest = {
 
 // Artifacts V2
 
+export enum ArtifactBrowserType {
+	Directory,
+	File,
+	NavigateUp
+}
+
+export type ArtifactBrowserItem = {
+	key: string;
+	text: string;
+	icon?: string;
+	size?: number;
+	type: ArtifactBrowserType;
+	dirResponse?: GetArtifactDirectoryEntryResponse;
+	fileResponse?: GetArtifactFileEntryResponse;
+}
+
+export type ArtifactSelection = {
+	filesSelected: number;
+	directoriesSelected: number;
+	size: number;
+	items: ArtifactBrowserItem[];
+}
+
 export type ArtifactContextType = "step-trace" | "step-output" | "step-saved" | string;
 
 /// Request to create a zip file with artifact data
@@ -3282,6 +3305,12 @@ export type UserClaim = {
 	value: string;
 }
 
+export enum DownloadFormat {
+	Zip = "Zip",
+	UGS = "UGS",
+	UTOOLBOX = "UTOOLBOX"
+}
+
 export enum DashboardPreference {
 	Darktheme = "Darktheme",
 	DisplayUTC = "DisplayUTC",
@@ -3293,7 +3322,9 @@ export enum DashboardPreference {
 	LocalCache = "LocalCache",
 	LeftAlignLog = "LeftAlignLog",
 	CompactViews = "CompactViews",
-	ShowPreflights = "ShowPreflights"
+	ShowPreflights = "ShowPreflights",
+	DownloadFormat = "DownloadFormat"
+
 }
 
 export type DashboardSettings = {
@@ -4985,7 +5016,7 @@ export type GetToolResponse = {
 	deployments: GetToolDeploymentResponse[];
 	public: boolean;
 	bundled?: boolean;
-	showInUgs: boolean;	
+	showInUgs: boolean;
 	showInDashboard: boolean;
 };
 
diff --git a/Engine/Source/Programs/Horde/HordeDashboard/src/backend/Dashboard.ts b/Engine/Source/Programs/Horde/HordeDashboard/src/backend/Dashboard.ts
--- a/Engine/Source/Programs/Horde/HordeDashboard/src/backend/Dashboard.ts
+++ b/Engine/Source/Programs/Horde/HordeDashboard/src/backend/Dashboard.ts
@@ -2,7 +2,7 @@
 
 import { action, makeObservable, observable } from 'mobx';
 import backend from '.';
-import { AuthMethod, DashboardPreference, GetDashboardAgentCategoryResponse, GetDashboardConfigResponse, GetDashboardPoolCategoryResponse, GetJobTemplateSettingsResponse, GetUserResponse, UserClaim } from './Api';
+import { AuthMethod, DashboardPreference, DownloadFormat, GetDashboardAgentCategoryResponse, GetDashboardConfigResponse, GetDashboardPoolCategoryResponse, GetJobTemplateSettingsResponse, GetUserResponse, UserClaim } from './Api';
 import { getSiteConfig } from './Config';
 
 export enum StatusColor {
@@ -386,6 +386,14 @@ export class Dashboard {
         return this.preferences.get(DashboardPreference.CompactViews) === 'true';
     }
 
+    setDownloadFormat(value: DownloadFormat | undefined) {
+        this.setPreference(DashboardPreference.DownloadFormat, value);
+    }
+
+    get downloadFormat(): DownloadFormat {
+        return this.preferences.get(DashboardPreference.DownloadFormat) as DownloadFormat ?? DownloadFormat.Zip;
+    }
+
     private hasLoggedLocalCache = false;
 
     get localCache(): boolean {
diff --git a/Engine/Source/Programs/Horde/HordeDashboard/src/base/utilities/stringUtills.ts b/Engine/Source/Programs/Horde/HordeDashboard/src/base/utilities/stringUtills.ts
new file mode 100644
--- /dev/null
+++ b/Engine/Source/Programs/Horde/HordeDashboard/src/base/utilities/stringUtills.ts
@@ -0,0 +1,13 @@
+
+export function formatBytes(bytes: number, decimals = 2) {
+    if (!+bytes) return '0 Bytes'
+ 
+    const k = 1024
+    const dm = decimals < 0 ? 0 : decimals
+    const sizes = ['Bytes', 'KB', 'MB', 'GB', 'TB', 'PB', 'EB', 'ZB', 'YB']
+ 
+    const i = Math.floor(Math.log(bytes) / Math.log(k))
+ 
+    return `${parseFloat((bytes / Math.pow(k, i)).toFixed(dm))} ${sizes[i]}`
+ }
+ 
\ No newline at end of file
diff --git a/Engine/Source/Programs/Horde/HordeDashboard/src/components/DashboardView.tsx b/Engine/Source/Programs/Horde/HordeDashboard/src/components/DashboardView.tsx
--- a/Engine/Source/Programs/Horde/HordeDashboard/src/components/DashboardView.tsx
+++ b/Engine/Source/Programs/Horde/HordeDashboard/src/components/DashboardView.tsx
@@ -1,10 +1,10 @@
 // Copyright Epic Games, Inc. All Rights Reserved.
 
-import { ColorPicker, DefaultButton, DialogFooter, IColorPickerStyles, Label, MessageBar, MessageBarType, Modal, Pivot, PivotItem, PrimaryButton, Spinner, SpinnerSize, Stack, Text, TextField, Toggle } from '@fluentui/react';
+import { ColorPicker, ComboBox, DefaultButton, DialogFooter, IColorPickerStyles, IComboBoxOption, Label, MessageBar, MessageBarType, Modal, Pivot, PivotItem, PrimaryButton, Spinner, SpinnerSize, Stack, Text, TextField, Toggle } from '@fluentui/react';
 import { observer } from 'mobx-react-lite';
 import React, { useState } from 'react';
 import backend from '../backend';
-import { AuthMethod, DashboardPreference, UpdateCurrentAccountRequest } from '../backend/Api';
+import { AuthMethod, DashboardPreference, DownloadFormat, UpdateCurrentAccountRequest } from '../backend/Api';
 import dashboard, { StatusColor, WebBrowser } from '../backend/Dashboard';
 import { useWindowSize } from '../base/utilities/hooks';
 import { getHordeStyling } from '../styles/Styles';
@@ -307,17 +307,17 @@ const ChangePasswordModal: React.FC<{ onClose: () => void }> = ({ onClose }) =>
          </Stack>}
 
          <Stack style={{ padding: 8 }}>
-            <TextField label={"Current Password"} autoComplete="off" defaultValue={ secrets.currentPassword} spellCheck={false} type="password" canRevealPassword onChange={(ev, value) => { setSecrets({ ...secrets, currentPassword: value ?? "" }) }} />
+            <TextField label={"Current Password"} autoComplete="off" defaultValue={secrets.currentPassword} spellCheck={false} type="password" canRevealPassword onChange={(ev, value) => { setSecrets({ ...secrets, currentPassword: value ?? "" }) }} />
          </Stack>
 
          <Stack style={{ padding: 8 }}>
-            <TextField label={"New Password"} autoComplete="new-password" spellCheck={false} type="password" defaultValue={ secrets.newPassword} canRevealPassword onChange={(ev, value) => { setSecrets({ ...secrets, newPassword: value ?? "" }) }} />
+            <TextField label={"New Password"} autoComplete="new-password" spellCheck={false} type="password" defaultValue={secrets.newPassword} canRevealPassword onChange={(ev, value) => { setSecrets({ ...secrets, newPassword: value ?? "" }) }} />
          </Stack>
 
          <Stack style={{ padding: 8 }}>
-            <TextField label={"Confirm Password"} autoComplete="new-password" spellCheck={false} type="password" defaultValue={ secrets.confirmPassword} canRevealPassword onChange={(ev, value) => { setSecrets({ ...secrets, confirmPassword: value ?? "" }) }} />
+            <TextField label={"Confirm Password"} autoComplete="new-password" spellCheck={false} type="password" defaultValue={secrets.confirmPassword} canRevealPassword onChange={(ev, value) => { setSecrets({ ...secrets, confirmPassword: value ?? "" }) }} />
          </Stack>
-         
+
          <Stack horizontal style={{ paddingTop: 64 }}>
             <Stack grow />
             <Stack horizontal tokens={{ childrenGap: 28 }}>
@@ -360,6 +360,19 @@ const GeneralPanel: React.FC = observer(() => {
 
    let name = dashboard.username;
 
+   const downloadOptions: IComboBoxOption[] = [
+      { key: `key_download_option_zip`, text: "Zip", data: DownloadFormat.Zip },
+      { key: `key_download_option_ugs`, text: "Unreal Game Sync", data: DownloadFormat.UGS },
+      { key: `key_download_option_toolbox`, text: "Unreal Toolbox", data: DownloadFormat.UTOOLBOX }
+   ];
+
+   let defaultDownloadKey = "key_download_option_zip";
+   if (dashboard.downloadFormat === DownloadFormat.UGS) {
+      defaultDownloadKey = "key_download_option_ugs";
+   } else if (dashboard.downloadFormat === DownloadFormat.UTOOLBOX) {
+      defaultDownloadKey = "key_download_option_toolbox";
+   }
+
    const pivotItems = ["Settings", "Claims", "Entitlements"].map(tab => {
       return <PivotItem headerText={tab} itemKey={tab} key={tab} style={{ color: modeColors.text }} />;
    });
@@ -383,7 +396,7 @@ const GeneralPanel: React.FC = observer(() => {
                </Pivot>
             </Stack>
             {pivotKey === "Claims" && <UserAccountPanel mode={pivotKey} />}
-            {pivotKey === "Entitlements" && <UserAccountPanel mode={pivotKey}/>}
+            {pivotKey === "Entitlements" && <UserAccountPanel mode={pivotKey} />}
             {pivotKey === "Settings" && <Stack style={{ paddingLeft: 12 }} tokens={{ childrenGap: 12 }}>
                <Stack style={{ paddingLeft: 8 }}>
                   <Toggle label="Dark Mode" inlineLabel={true} defaultChecked={dashboard.darktheme} onChange={(ev, checked) => {
@@ -409,12 +422,17 @@ const GeneralPanel: React.FC = observer(() => {
                      dashboard.setLeftAlignLog(checked ? true : false);
                   }} />
                </Stack>
-
-               <Stack className="horde-no-darktheme" style={{paddingLeft: 8}} tokens={{ childrenGap: 12 }}>
+               <Stack style={{ paddingLeft: 8 }} tokens={{ childrenGap: 12 }}>
+                  <ComboBox style={{ width: 218 }} label="Download Format" defaultSelectedKey={defaultDownloadKey} options={downloadOptions} autoComplete="off" onChange={(ev, option, index, value) => {
+                     if (option) {
+                        dashboard.setDownloadFormat(option?.data);
+                     }                     
+                   }} />
+               </Stack>
+               <Stack style={{ paddingLeft: 8 }} tokens={{ childrenGap: 12 }}>
                   <Label>Status Colors</Label>
                   <BadgeColors />
                </Stack>
-
                {dashboard.authMethod === AuthMethod.Horde && <Stack style={{ paddingTop: 24, paddingLeft: 8 }}>
                   <PrimaryButton style={{ width: 180 }} text="Change Password" onClick={() => setChangePassword(true)} />
                </Stack>}
diff --git a/Engine/Source/Programs/Horde/HordeDashboard/src/components/artifacts/ArtifactButton.tsx b/Engine/Source/Programs/Horde/HordeDashboard/src/components/artifacts/ArtifactButton.tsx
new file mode 100644
--- /dev/null
+++ b/Engine/Source/Programs/Horde/HordeDashboard/src/components/artifacts/ArtifactButton.tsx
@@ -0,0 +1,188 @@
+import { IContextualMenuProps, DirectionalHint, PrimaryButton } from "@fluentui/react";
+import backend from "horde/backend";
+import { ArtifactBrowserItem, ArtifactBrowserType, ArtifactSelection, CreateZipRequest, DownloadFormat, GetArtifactResponse } from "horde/backend/Api";
+import dashboard from "horde/backend/Dashboard";
+import { formatBytes } from "horde/base/utilities/stringUtills";
+import { useNavigate } from "react-router-dom";
+
+
+export const ArtifactButton: React.FC<{ jobId?: string, stepId?: string, artifact?: GetArtifactResponse, pathIn?: string, selection?: ArtifactSelection, openArtifactInfo?: () => void, disabled?: boolean, minWidth?: number }> = ({ jobId, stepId, artifact, pathIn, selection, openArtifactInfo, disabled, minWidth }) => {
+
+   const navigate = useNavigate();
+
+   const format = dashboard.downloadFormat;
+   const path = pathIn ? pathIn + "/" : "";
+   const items = selection?.items;
+
+   let sizeText = "0KB";
+   if (selection?.size) {
+      sizeText = formatBytes(selection.size, (selection.size < (1024 * 1024)) ? 0 : 1)
+   }
+
+   let buttonText = "Download";
+
+   if (selection?.directoriesSelected || (selection?.filesSelected ?? 0) > 1) {
+      buttonText = `Download (${sizeText})`;
+   } else if (selection?.filesSelected === 1) {
+      buttonText = `Download (${sizeText})`;
+   }
+
+   let jobUrl = "";
+   if (jobId) {
+      jobUrl = `/job/${jobId}`;
+      if (stepId) {
+         jobUrl += `?step=${stepId}`;
+      }
+   }
+
+   const downloadProps: IContextualMenuProps = {
+      items: [],
+      directionalHint: DirectionalHint.bottomRightEdge
+   };
+
+   if (openArtifactInfo) {
+      downloadProps.items.push({
+         key: 'view_artifact_info',
+         text: 'View Artifact Info',
+         disabled: !artifact,
+         onClick: () => {
+            openArtifactInfo();
+         }
+      })
+   }
+
+   if (jobUrl) {
+      {
+         downloadProps.items.push({
+            key: 'navigate_to_job',
+            text: 'Navigate to Job',
+            onClick: () => {
+               navigate(jobUrl);
+            }
+         })
+      }
+   }
+
+   function downloadToolbox() {
+      if (artifact?.id) {
+         window.location.assign(`horde-artifact://${window.location.hostname}:${window.location.port}/api/v2/artifacts/${artifact.id}`);
+      }
+   }
+
+   function downloadUGS() {
+      if (artifact?.id) {
+         window.location.assign(`/api/v2/artifacts/${artifact.id}/download?format=ugs`);
+      }
+   }
+
+   function downloadZip() {
+
+      if (!artifact) {
+         return;
+      }
+
+      if (!items?.length) {
+         window.location.assign(`/api/v2/artifacts/${artifact.id}/download?format=zip`);
+         return;
+      }
+
+      // download a single file
+      if (items.length === 1) {
+         const item = items[0];
+         if (item.type === ArtifactBrowserType.File) {
+
+            try {
+               backend.downloadArtifactV2(artifact.id, path + item.text);
+            } catch (err) {
+               console.error(err);
+            } finally {
+
+            }
+
+            return;
+         }
+      }
+
+      let zipRequest: CreateZipRequest | undefined;
+
+      if ((items?.length ?? 0) > 0) {
+
+         const filter = items.map(s => {
+            const item = s as ArtifactBrowserItem;
+
+            if (item.type === ArtifactBrowserType.NavigateUp) {
+               return "";
+            }
+
+            if (item.type === ArtifactBrowserType.Directory) {
+               return `${path}${item.text}/...`;
+            }
+            return `${path}${item.text}`;
+         }).filter(f => !!f);
+
+         zipRequest = {
+            filter: filter
+         }
+      } else {
+         if (path) {
+            zipRequest = {
+               filter: [path + "..."]
+            }
+         }
+      }
+
+      try {
+         backend.downloadArtifactZipV2(artifact.id, zipRequest);
+      } catch (err) {
+         console.error(err);
+      } finally {
+
+      }
+   }
+
+   if (artifact?.id) {
+
+      downloadProps.items.push(
+         {
+            key: 'download_zip',
+            text: 'Download Zip',
+            onClick: () => {
+               downloadZip();
+            }
+         }
+      )
+
+      downloadProps.items.push(
+         {
+            key: 'download_ugs',
+            text: 'Download with UGS',
+            onClick: () => {
+               downloadUGS();
+            }
+         }
+      )
+      downloadProps.items.push(
+         {
+            key: 'download_toolbox',
+            text: 'Download with Toolbox',
+            onClick: () => {
+               downloadToolbox();
+            }
+         }
+      )
+   }
+
+   return <PrimaryButton split menuProps={downloadProps} text={buttonText} styles={{ root: { minWidth: minWidth ?? 104, fontFamily: 'Horde Open Sans SemiBold !important' } }} disabled={disabled} onClick={() => {
+      if (format === DownloadFormat.Zip) {
+         downloadZip();
+      }
+      else if (format === DownloadFormat.UGS) {
+         downloadUGS();
+      }
+      else if (format === DownloadFormat.UTOOLBOX) {
+         downloadToolbox();
+      }
+
+   }} />
+
+}
\ No newline at end of file
diff --git a/Engine/Source/Programs/Horde/HordeDashboard/src/components/artifacts/ArtifactsModal.tsx b/Engine/Source/Programs/Horde/HordeDashboard/src/components/artifacts/ArtifactsModal.tsx
--- a/Engine/Source/Programs/Horde/HordeDashboard/src/components/artifacts/ArtifactsModal.tsx
+++ b/Engine/Source/Programs/Horde/HordeDashboard/src/components/artifacts/ArtifactsModal.tsx
@@ -6,25 +6,11 @@ import { observer } from "mobx-react-lite";
 import { useEffect, useState } from "react";
 import { NavigateFunction, useNavigate } from "react-router-dom";
 import backend from "../../backend";
-import { ArtifactContextType, CreateZipRequest, GetArtifactDirectoryEntryResponse, GetArtifactDirectoryResponse, GetArtifactFileEntryResponse, GetArtifactResponse } from "../../backend/Api";
+import { ArtifactBrowserItem, ArtifactBrowserType, ArtifactContextType, CreateZipRequest, GetArtifactDirectoryEntryResponse, GetArtifactDirectoryResponse, GetArtifactFileEntryResponse, GetArtifactResponse } from "../../backend/Api";
 import dashboard, { StatusColor } from "../../backend/Dashboard";
 import { getHordeStyling } from "../../styles/Styles";
-
-enum BrowserType {
-   Directory,
-   File,
-   NavigateUp
-}
-
-type BrowserItem = {
-   key: string;
-   text: string;
-   icon?: string;
-   size?: number;
-   type: BrowserType;
-   dirResponse?: GetArtifactDirectoryEntryResponse;
-   fileResponse?: GetArtifactFileEntryResponse;
-}
+import { formatBytes } from "horde/base/utilities/stringUtills";
+import { ArtifactButton } from "./ArtifactButton";
 
 
 class ArtifactsHandler {
@@ -88,7 +74,7 @@ class ArtifactsHandler {
             artifacts.push(artifact);
          } catch (err) {
             console.log(err);
-         }         
+         }
       }
 
       this.artifacts = artifacts;
@@ -185,7 +171,7 @@ class ArtifactsHandler {
          if (path?.length) {
             url += `&artifactPath=${encodeURIComponent(path.replaceAll("+", "%2B"))}`;
          }
-         
+
          navigate(url, { replace: true });
          console.log(this.artifact.id, path);
       }
@@ -195,13 +181,13 @@ class ArtifactsHandler {
 
    }
 
-   get currentSelection(): { filesSelected: number, directoriesSelected: number, size: number, items: BrowserItem[] } {
+   get currentSelection(): { filesSelected: number, directoriesSelected: number, size: number, items: ArtifactBrowserItem[] } {
 
       let result = {
          filesSelected: 0,
          directoriesSelected: 0,
          size: 0,
-         items: [] as BrowserItem[]
+         items: [] as ArtifactBrowserItem[]
       }
 
       const browse = this.browse;
@@ -209,27 +195,27 @@ class ArtifactsHandler {
          return result;
       }
 
-      let selection = this.selection?.getSelection() as (BrowserItem[] | undefined);
+      let selection = this.selection?.getSelection() as (ArtifactBrowserItem[] | undefined);
 
       if (!selection?.length) {
          selection = [];
 
          browse.directories?.forEach(d => {
-            selection!.push({ key: d.hash, text: d.name, icon: "Folder", type: BrowserType.Directory, size: d.length });
+            selection!.push({ key: d.hash, text: d.name, icon: "Folder", type: ArtifactBrowserType.Directory, size: d.length });
          });
 
          browse.files?.forEach(d => {
-            selection!.push({ key: d.hash, text: d.name, icon: "Document", type: BrowserType.File, size: d.length });
+            selection!.push({ key: d.hash, text: d.name, icon: "Document", type: ArtifactBrowserType.File, size: d.length });
          });
       }
 
       selection.forEach(b => {
          result.items.push(b);
-         if (b.type === BrowserType.Directory) {
+         if (b.type === ArtifactBrowserType.Directory) {
             result.directoriesSelected++;
             result.size += b.size ?? 0;
          }
-         if (b.type === BrowserType.File) {
+         if (b.type === ArtifactBrowserType.File) {
             result.filesSelected++;
             result.size += b.size ?? 0;
          }
@@ -267,7 +253,7 @@ class ArtifactsHandler {
 
    path?: string;
 
-   selection?: Selection = new Selection({ canSelectItem: (item: any) => { return item.type !== BrowserType.NavigateUp; }, onSelectionChanged: () => { if (this.selectionCallback) this.selectionCallback() } });
+   selection?: Selection = new Selection({ canSelectItem: (item: any) => { return item.type !== ArtifactBrowserType.NavigateUp; }, onSelectionChanged: () => { if (this.selectionCallback) this.selectionCallback() } });
 
    browse?: GetArtifactDirectoryResponse;
 
@@ -426,24 +412,10 @@ const BrowseBreadCrumbs: React.FC<{ handler: ArtifactsHandler }> = observer(({ h
 
 })
 
-function formatBytes(bytes: number, decimals = 2) {
-   if (!+bytes) return '0 Bytes'
-
-   const k = 1024
-   const dm = decimals < 0 ? 0 : decimals
-   const sizes = ['Bytes', 'KB', 'MB', 'GB', 'TB', 'PB', 'EB', 'ZB', 'YB']
-
-   const i = Math.floor(Math.log(bytes) / Math.log(k))
-
-   return `${parseFloat((bytes / Math.pow(k, i)).toFixed(dm))} ${sizes[i]}`
-}
-
 
 const DownloadButton: React.FC<{ handler: ArtifactsHandler, openArtifactInfo: () => void }> = observer(({ handler, openArtifactInfo }) => {
 
-   const [selectKey, setSelectionKey] = useState(0);
-   const navigate = useNavigate();
-
+   const [selectKey, setSelectionKey] = useState(0);   
 
    // subscribe
    if (handler.updated) { };
@@ -458,131 +430,10 @@ const DownloadButton: React.FC<{ handler: ArtifactsHandler, openArtifactInfo: ()
       setSelectionKey(selectKey + 1);
    }
 
-   let buttonText = "Download";
-
-
    const selection = handler.currentSelection;
 
-
-   let sizeText = "0KB";
-   if (selection.size) {
-      sizeText = formatBytes(selection.size, (selection.size < (1024 * 1024)) ? 0 : 1)
-   }
-
-   if (selection.directoriesSelected || selection.filesSelected > 1) {
-      buttonText = `Download (${sizeText})`;
-   } else if (selection.filesSelected === 1) {
-      buttonText = `Download (${sizeText})`;
-   }
-
-   let jobUrl = "";
-   if (handler.jobId) {
-      jobUrl = `/job/${handler.jobId}`;
-      if (handler.stepId) {
-         jobUrl += `?step=${handler.stepId}`;
-      }
-   }
-
-   const downloadProps: IContextualMenuProps = {
-      items: [
-         {
-            key: 'view_artifact_info',
-            text: 'View Artifact Info',
-            disabled: !handler.artifact,
-            onClick: () => {
-               openArtifactInfo();
-            }
-         },
-         {
-            key: 'navigate_to_job',
-            text: 'Navigate to Job',
-            onClick: () => {
-               navigate(jobUrl);
-            }
-         }
-      ],
-      directionalHint: DirectionalHint.bottomRightEdge
-   };
-
-   if (handler.artifact?.id) {
-      const id = handler.artifact.id;
-      downloadProps.items.unshift(
-         {
-            key: 'download_ugs',
-            text: 'Download with UGS',
-            onClick: () => {
-               window.location.assign(`/api/v2/artifacts/${id}/download?format=ugs`);
-            }
-         }
-      )
-   }
-
-
    return <Stack id="callout_target_artifactinfo" horizontal verticalAlign="center" tokens={{ childrenGap: 8 }}>
-      <PrimaryButton split menuProps={downloadProps} styles={{ root: { fontFamily: 'Horde Open Sans SemiBold !important' } }} disabled={!selection.filesSelected && !selection.directoriesSelected} onClick={async () => {
-
-         const selection = handler.currentSelection.items;
-
-         if (!selection?.length || !handler.artifact) {
-            return;
-         }
-
-         // download a single file
-         if (selection.length === 1) {
-            const item = selection[0] as BrowserItem;
-            if (item.type === BrowserType.File) {
-
-               try {
-                  backend.downloadArtifactV2(handler.artifact.id, (handler.path ? handler.path + "/" : "") + item.text);
-               } catch (err) {
-                  console.error(err);
-               } finally {
-
-               }
-
-               return;
-            }
-         }
-
-         const path = handler.path ? handler.path + "/" : "";
-
-         let zipRequest: CreateZipRequest | undefined;
-
-         if ((handler.selection?.count ?? 0) > 0) {
-
-            const filter = selection.map(s => {
-               const item = s as BrowserItem;
-   
-               if (item.type === BrowserType.NavigateUp) {
-                  return "";
-               }
-   
-               if (item.type === BrowserType.Directory) {
-                  return `${path}${item.text}/...`;
-               }
-               return `${path}${item.text}`;
-            }).filter(f => !!f);
-
-            zipRequest = {
-               filter: filter
-            }   
-         } else {
-            if (path) {
-               zipRequest = {
-                  filter: [path + "..."]
-               }   
-            }
-         }
-
-         try {
-            backend.downloadArtifactZipV2(handler.artifact.id, zipRequest);
-         } catch (err) {
-            console.error(err);
-         } finally {
-
-         }
-
-      }}>{buttonText}</PrimaryButton>
+      <ArtifactButton minWidth={150} jobId={handler?.jobId} stepId={handler?.stepId} artifact={handler.artifact} selection={selection} pathIn={handler.path} openArtifactInfo={openArtifactInfo} />
    </Stack>
 
 });
@@ -631,11 +482,11 @@ const JobDetailArtifactsInner: React.FC<{ jobId: string; stepId: string, artifac
       return <Stack><Spinner size={SpinnerSize.large} /></Stack>;
    }
 
-   const items: BrowserItem[] = [];
+   const items: ArtifactBrowserItem[] = [];
 
-   const getFileHRef = (item: BrowserItem) => {
+   const getFileHRef = (item: ArtifactBrowserItem) => {
 
-      if (item.type !== BrowserType.File) {
+      if (item.type !== ArtifactBrowserType.File) {
          return undefined;
       }
 
@@ -648,7 +499,7 @@ const JobDetailArtifactsInner: React.FC<{ jobId: string; stepId: string, artifac
 
    // use the up arrow instead
    if (handler.path?.length) {
-      items.push({ key: "navigate up", text: "..", type: BrowserType.NavigateUp });
+      items.push({ key: "navigate up", text: "..", type: ArtifactBrowserType.NavigateUp });
    }
 
    browse.directories?.forEach(d => {
@@ -656,7 +507,7 @@ const JobDetailArtifactsInner: React.FC<{ jobId: string; stepId: string, artifac
       function recurseDirectories(dir: GetArtifactDirectoryEntryResponse, flattened: GetArtifactDirectoryEntryResponse[]) {
          if (!dir.directories || dir.directories.length > 1 || dir.files?.length) {
             const name = flattened.length ? flattened.map(d => d.name).join("/") + "/" + dir.name : dir.name;
-            items.push({ key: dir.hash, text: name, icon: "Folder", type: BrowserType.Directory, size: dir.length, dirResponse: dir });
+            items.push({ key: dir.hash, text: name, icon: "Folder", type: ArtifactBrowserType.Directory, size: dir.length, dirResponse: dir });
          } else {
             flattened.push(dir);
             dir.directories.forEach(d => recurseDirectories(d, [...flattened]));
@@ -669,7 +520,7 @@ const JobDetailArtifactsInner: React.FC<{ jobId: string; stepId: string, artifac
    });
 
    browse.files?.forEach(d => {
-      items.push({ key: d.hash, text: d.name, icon: "Document", type: BrowserType.File, size: d.length, fileResponse: d });
+      items.push({ key: d.hash, text: d.name, icon: "Document", type: ArtifactBrowserType.File, size: d.length, fileResponse: d });
    });
 
    const columns: IColumn[] = [
@@ -700,7 +551,7 @@ const JobDetailArtifactsInner: React.FC<{ jobId: string; stepId: string, artifac
             return null;
          }
 
-         const isFile = item.type === BrowserType.File;
+         const isFile = item.type === ArtifactBrowserType.File;
 
          const href = getFileHRef(item);
 
@@ -726,18 +577,18 @@ const JobDetailArtifactsInner: React.FC<{ jobId: string; stepId: string, artifac
 
          return <Stack verticalFill verticalAlign="center" style={{ cursor: "pointer" }} onClick={(ev) => {
 
-            if (item.type === BrowserType.Directory) {
+            if (item.type === ArtifactBrowserType.Directory) {
                const nbrowse = handler.path ? `${handler.path}/${item.text}` : item.text;
                handler.browseTo(nbrowse, navigate);
             }
-            if (item.type === BrowserType.NavigateUp && handler.path) {
+            if (item.type === ArtifactBrowserType.NavigateUp && handler.path) {
                const nbrowse = handler.path.split("/")
                nbrowse.pop();
                handler.browseTo(nbrowse.join("/"), navigate);
             }
 
          }}>
-            {item.type !== BrowserType.NavigateUp && <Stack data-selection-disabled={item.type !== BrowserType.File} horizontal tokens={{ childrenGap: 8 }}>
+            {item.type !== ArtifactBrowserType.NavigateUp && <Stack data-selection-disabled={item.type !== ArtifactBrowserType.File} horizontal tokens={{ childrenGap: 8 }}>
                <Stack>
                   <FontIcon style={{ paddingTop: 1, fontSize: 16 }} iconName={item.icon} />
                </Stack>
@@ -745,7 +596,7 @@ const JobDetailArtifactsInner: React.FC<{ jobId: string; stepId: string, artifac
                   {pathElements}
                </Stack>
             </Stack>}
-            {item.type === BrowserType.NavigateUp && <Stack data-selection-disabled verticalFill horizontal verticalAlign="center" tokens={{ childrenGap: 9 }}>
+            {item.type === ArtifactBrowserType.NavigateUp && <Stack data-selection-disabled verticalFill horizontal verticalAlign="center" tokens={{ childrenGap: 9 }}>
                <Stack>
                   <FontIcon style={{ paddingTop: 1, fontSize: 15 }} iconName="ArrowUp" />
                </Stack>
@@ -795,7 +646,7 @@ const JobDetailArtifactsInner: React.FC<{ jobId: string; stepId: string, artifac
          </Stack>
       </Callout>}
 
-      <Stack tokens={{childrenGap: 12}}>
+      <Stack tokens={{ childrenGap: 12 }}>
          <Stack horizontal verticalAlign="center" tokens={{ childrenGap: 18 }} style={{ paddingBottom: 12 }}>
             <Stack>
                <BrowseBreadCrumbs handler={handler} />
@@ -822,8 +673,8 @@ const JobDetailArtifactsInner: React.FC<{ jobId: string; stepId: string, artifac
                      selection={handler.selection}
                      selectionPreservedOnEmptyClick={true}
                      onShouldVirtualize={() => false}
-                     onItemInvoked={(item: BrowserItem) => {
-                        if (item?.type !== BrowserType.File) {
+                     onItemInvoked={(item: ArtifactBrowserItem) => {
+                        if (item?.type !== ArtifactBrowserType.File) {
                            return;
                         }
                         const href = getFileHRef(item);
diff --git a/Engine/Source/Programs/Horde/HordeDashboard/src/components/artifacts/ArtifactsSearch.tsx b/Engine/Source/Programs/Horde/HordeDashboard/src/components/artifacts/ArtifactsSearch.tsx
--- a/Engine/Source/Programs/Horde/HordeDashboard/src/components/artifacts/ArtifactsSearch.tsx
+++ b/Engine/Source/Programs/Horde/HordeDashboard/src/components/artifacts/ArtifactsSearch.tsx
@@ -11,6 +11,7 @@ import { getHordeStyling } from "../../styles/Styles";
 import { StreamChooser, streamIdAll } from "../projects/StreamChooser";
 import { JobArtifactsModal } from "./ArtifactsModal";
 import { useQuery } from "../JobDetailCommon";
+import { ArtifactButton } from "./ArtifactButton";
 
 type ArtifactSearchState = {
    streamId?: string;
@@ -305,20 +306,6 @@ const ArtifactsList: React.FC<{ state: ArtifactQueryState, artifacts?: GetArtifa
 
          const stream = projectStore.streamById(item.streamId);
 
-         const downloadProps: IContextualMenuProps = {
-            items: [
-               {
-                  key: 'download_ugs',
-                  text: 'Download with UGS',
-                  onClick: () => {
-                     window.location.assign(`/api/v2/artifacts/${item.id}/download?format=ugs`);
-                  }
-               }
-            ],
-            directionalHint: DirectionalHint.bottomLeftEdge
-         };
-   
-
          return <Stack horizontal verticalAlign="center" verticalFill tokens={{ childrenGap: 24 }} styles={{ root: { backgroundColor: background, paddingLeft: 12, paddingRight: 12, paddingTop: 8, paddingBottom: 8 } }}>
             <Stack style={{ width: 360, paddingLeft: 8 }}>
                <Text style={{ fontWeight: 600 }}>{item.name}</Text>
@@ -332,7 +319,7 @@ const ArtifactsList: React.FC<{ state: ArtifactQueryState, artifacts?: GetArtifa
             <Stack style={{ width: 144 }}>
                <Text>{item.type}</Text>
             </Stack>
-            <Stack horizontal tokens={{childrenGap: 18}}>
+            <Stack horizontal tokens={{ childrenGap: 18 }}>
                <DefaultButton style={{ width: 90 }} text="Browse" onClick={() => {
 
                   const key = item.keys.find(k => k.startsWith("job:") && k.indexOf("/step:") !== -1);
@@ -347,14 +334,10 @@ const ArtifactsList: React.FC<{ state: ArtifactQueryState, artifacts?: GetArtifa
                   state.updateSearch(navigate, false);
 
                }} />
-               <DefaultButton split text="Download Zip" menuProps={downloadProps}
-                  disabled={!item.id}
-                  style={{ fontFamily: "Horde Open Sans SemiBold", width: 100, padding: 0 }}
-                  onClick={() => {
-                     window.location.assign(`/api/v2/artifacts/${item.id}/download?format=zip`);
-                  }}
+               <ArtifactButton 
+                  disabled={!item.id}                  
+                  artifact={item}
                />
-
             </Stack>
          </Stack>
       }
diff --git a/Engine/Source/Programs/Horde/HordeDashboard/src/components/jobDetailsV2/JobDetailArtifacts.tsx b/Engine/Source/Programs/Horde/HordeDashboard/src/components/jobDetailsV2/JobDetailArtifacts.tsx
--- a/Engine/Source/Programs/Horde/HordeDashboard/src/components/jobDetailsV2/JobDetailArtifacts.tsx
+++ b/Engine/Source/Programs/Horde/HordeDashboard/src/components/jobDetailsV2/JobDetailArtifacts.tsx
@@ -13,6 +13,7 @@ import { JobDataView, JobDetailsV2 } from "./JobDetailsViewCommon";
 import { JobArtifactsModal } from '../artifacts/ArtifactsModal';
 
 import { useConst } from '@fluentui/react-hooks';
+import { ArtifactButton } from '../artifacts/ArtifactButton';
 const sideRail: ISideRailLink = { text: "Artifacts", url: "rail_artifacts" };
 
 class JobArtifactsDataView extends JobDataView {
@@ -203,19 +204,6 @@ export const JobArtifactsPanel: React.FC<{ jobDetails: JobDetailsV2 }> = observe
 
       let stepFinished = !!step?.finishTime;
 
-      const downloadProps: IContextualMenuProps = {
-         items: [
-            {
-               key: 'download_ugs',
-               text: 'Download with UGS',
-               onClick: () => {
-                  window.location.assign(`/api/v2/artifacts/${item.id}/download?format=ugs`);
-               }
-            }
-         ],
-         directionalHint: DirectionalHint.bottomLeftEdge
-      };
-
       if (column.key === 'column_download') {
          return <Stack horizontalAlign="end" verticalAlign="center" verticalFill={true} style={{ paddingRight: 8 }}>
             <Stack horizontal tokens={{ childrenGap: 18 }}>
@@ -226,13 +214,9 @@ export const JobArtifactsPanel: React.FC<{ jobDetails: JobDetailsV2 }> = observe
                      setSelected(item)
                   }}
                />
-
-               <PrimaryButton split text="Download as Zip" menuProps={downloadProps}
+               <ArtifactButton   
                   disabled={!stepFinished || !item.id}
-                  style={{ fontFamily: "Horde Open Sans SemiBold" }}
-                  onClick={() => {
-                     window.location.assign(`/api/v2/artifacts/${item.id}/download?format=zip`);
-                  }}
+                  artifact={item}
                />
             </Stack>
          </Stack>
