Horde: Add TopN/BottomN sampling constraints for Horde Analytics

diff --git a/Engine/Source/Programs/Horde/Plugins/Analytics/HordeServer.Analytics.Tests/Telemetry/TelemetryTests.cs b/Engine/Source/Programs/Horde/Plugins/Analytics/HordeServer.Analytics.Tests/Telemetry/TelemetryTests.cs
--- a/Engine/Source/Programs/Horde/Plugins/Analytics/HordeServer.Analytics.Tests/Telemetry/TelemetryTests.cs
+++ b/Engine/Source/Programs/Horde/Plugins/Analytics/HordeServer.Analytics.Tests/Telemetry/TelemetryTests.cs
@@ -442,5 +442,132 @@ namespace HordeServer.Analytics.Tests.Telemetry
 				}
 			}
 		}
+
+		[TestMethod]
+		public async Task TopAndBottomAsync()
+		{
+			// TopN
+			{
+				Clock.UtcNow = new DateTime(2023, 6, 8, 4, 30, 0, DateTimeKind.Utc);
+
+				MetricConfig metricConfig = new MetricConfig();
+				metricConfig.Id = new MetricId("test-metric-1");
+				metricConfig.Function = AggregationFunction.Sum;
+				metricConfig.Interval = TimeSpan.FromHours(1.0);
+				metricConfig.Property = JsonPath.Parse("$.Payload.foo");
+				metricConfig.GroupBy = "$.Payload.group";
+				metricConfig.TopN = 3;
+
+				TelemetryStoreConfig telemetryStoreConfig = new TelemetryStoreConfig();
+				telemetryStoreConfig.Id = TelemetryStoreId.Default;
+				telemetryStoreConfig.Metrics.Add(metricConfig);
+
+				AnalyticsConfig analyticsConfig = new AnalyticsConfig();
+				analyticsConfig.Stores.Add(telemetryStoreConfig);
+
+				GlobalConfig globalConfig = new GlobalConfig();
+				globalConfig.Plugins.AddAnalyticsTestConfig(analyticsConfig);
+				SetConfig(globalConfig);
+
+				MetricTelemetrySink sink = ServiceProvider.GetRequiredService<MetricTelemetrySink>();
+				IMetricCollection collection = ServiceProvider.GetRequiredService<IMetricCollection>();
+
+				sink.SendEvent(TelemetryStoreId.Default, new TelemetryEvent(s_metadata, new { foo = 1, group = "first" }));
+				sink.SendEvent(TelemetryStoreId.Default, new TelemetryEvent(s_metadata, new { foo = 2, group = "first" }));
+				sink.SendEvent(TelemetryStoreId.Default, new TelemetryEvent(s_metadata, new { foo = 10, group = "first" }));
+				sink.SendEvent(TelemetryStoreId.Default, new TelemetryEvent(s_metadata, new { foo = 4, group = "second" }));
+				await sink.FlushAsync(CancellationToken.None);
+				await collection.FlushAsync(CancellationToken.None);
+				
+				sink.SendEvent(TelemetryStoreId.Default, new TelemetryEvent(s_metadata, new { foo = 5, group = "third" }));
+				sink.SendEvent(TelemetryStoreId.Default, new TelemetryEvent(s_metadata, new { foo = 6, group = "fourth" }));
+				sink.SendEvent(TelemetryStoreId.Default, new TelemetryEvent(s_metadata, new { foo = 7, group = "fifth" }));
+				sink.SendEvent(TelemetryStoreId.Default, new TelemetryEvent(s_metadata, new { foo = 8, group = "sixth" }));
+				sink.SendEvent(TelemetryStoreId.Default, new TelemetryEvent(s_metadata, new { foo = 9 }));
+				
+				await sink.FlushAsync(CancellationToken.None);
+				await collection.FlushAsync(CancellationToken.None);
+
+				List<IMetric> metrics = await collection.FindAsync(telemetryStoreConfig.Id, [metricConfig.Id]);
+				metrics = metrics.OrderBy(x => x.Group).ToList();
+				Assert.AreEqual(3, metrics.Count);
+
+				Assert.AreEqual("", metrics[0].Group);
+				Assert.AreEqual(1, metrics[0].Count);
+				Assert.AreEqual(9, metrics[0].Value);
+
+				Assert.AreEqual("first", metrics[1].Group);
+				Assert.AreEqual(3, metrics[1].Count);
+				Assert.AreEqual(13, metrics[1].Value);
+
+				Assert.AreEqual("sixth", metrics[2].Group);
+				Assert.AreEqual(1, metrics[2].Count);
+				Assert.AreEqual(8, metrics[2].Value);
+			}
+
+			// BottomN
+			{
+				Clock.UtcNow = new DateTime(2023, 6, 8, 4, 30, 0, DateTimeKind.Utc);
+
+				MetricConfig metricConfig = new MetricConfig();
+				metricConfig.Id = new MetricId("test-metric-1");
+				metricConfig.Function = AggregationFunction.Sum;
+				metricConfig.Interval = TimeSpan.FromHours(1.0);
+				metricConfig.Property = JsonPath.Parse("$.Payload.foo");
+				metricConfig.GroupBy = "$.Payload.group";
+				metricConfig.BottomN = 3;
+
+				TelemetryStoreConfig telemetryStoreConfig = new TelemetryStoreConfig();
+				telemetryStoreConfig.Id = TelemetryStoreId.Default;
+				telemetryStoreConfig.Metrics.Add(metricConfig);
+
+				AnalyticsConfig analyticsConfig = new AnalyticsConfig();
+				analyticsConfig.Stores.Add(telemetryStoreConfig);
+
+				GlobalConfig globalConfig = new GlobalConfig();
+				globalConfig.Plugins.AddAnalyticsTestConfig(analyticsConfig);
+				SetConfig(globalConfig);
+
+				MetricTelemetrySink sink = ServiceProvider.GetRequiredService<MetricTelemetrySink>();
+				IMetricCollection collection = ServiceProvider.GetRequiredService<IMetricCollection>();
+
+				sink.SendEvent(TelemetryStoreId.Default, new TelemetryEvent(s_metadata, new { foo = 1, group = "first" }));
+				sink.SendEvent(TelemetryStoreId.Default, new TelemetryEvent(s_metadata, new { foo = 2, group = "first" }));
+				sink.SendEvent(TelemetryStoreId.Default, new TelemetryEvent(s_metadata, new { foo = 10, group = "first" }));
+
+				sink.SendEvent(TelemetryStoreId.Default, new TelemetryEvent(s_metadata, new { foo = 4, group = "second" }));
+				sink.SendEvent(TelemetryStoreId.Default, new TelemetryEvent(s_metadata, new { foo = 5, group = "third" }));
+				
+				await sink.FlushAsync(CancellationToken.None);
+				await collection.FlushAsync(CancellationToken.None);
+
+				sink.SendEvent(TelemetryStoreId.Default, new TelemetryEvent(s_metadata, new { foo = 5, group = "fourth" }));
+				sink.SendEvent(TelemetryStoreId.Default, new TelemetryEvent(s_metadata, new { foo = 1, group = "fourth" }));
+				sink.SendEvent(TelemetryStoreId.Default, new TelemetryEvent(s_metadata, new { foo = 7, group = "fifth" }));
+				sink.SendEvent(TelemetryStoreId.Default, new TelemetryEvent(s_metadata, new { foo = 8, group = "sixth" }));
+				sink.SendEvent(TelemetryStoreId.Default, new TelemetryEvent(s_metadata, new { foo = 9 }));
+				
+				await sink.FlushAsync(CancellationToken.None);
+				await collection.FlushAsync(CancellationToken.None);
+
+				List<IMetric> metrics = await collection.FindAsync(telemetryStoreConfig.Id, [metricConfig.Id]);
+				metrics = metrics.OrderBy(x => x.Group).ToList();
+				Assert.AreEqual(3, metrics.Count);
+
+				Assert.AreEqual("fourth", metrics[0].Group);
+				Assert.AreEqual(2, metrics[0].Count);
+				Assert.AreEqual(6, metrics[0].Value);
+
+				Assert.AreEqual("second", metrics[1].Group);
+				Assert.AreEqual(1, metrics[1].Count);
+				Assert.AreEqual(4, metrics[1].Value);
+
+				Assert.AreEqual("third", metrics[2].Group);
+				Assert.AreEqual(1, metrics[2].Count);
+				Assert.AreEqual(5, metrics[2].Value);
+			}
+
+		}
+
 	}
 }
diff --git a/Engine/Source/Programs/Horde/Plugins/Analytics/HordeServer.Analytics/Telemetry/Metrics/MetricCollection.cs b/Engine/Source/Programs/Horde/Plugins/Analytics/HordeServer.Analytics/Telemetry/Metrics/MetricCollection.cs
--- a/Engine/Source/Programs/Horde/Plugins/Analytics/HordeServer.Analytics/Telemetry/Metrics/MetricCollection.cs
+++ b/Engine/Source/Programs/Horde/Plugins/Analytics/HordeServer.Analytics/Telemetry/Metrics/MetricCollection.cs
@@ -247,6 +247,8 @@ namespace HordeServer.Telemetry.Metrics
 				_queuedSamples = new Dictionary<SampleKey, List<double>>();
 			}
 
+			Dictionary<SampleKey, MetricConfig> tbConstraint = new Dictionary<SampleKey, MetricConfig>();
+
 			// Add the samples to the database
 			try
 			{
@@ -256,6 +258,11 @@ namespace HordeServer.Telemetry.Metrics
 					if (_analyticsConfig.CurrentValue.TryGetTelemetryStore(sampleKey.Store, out TelemetryStoreConfig? telemetryStoreConfig) && telemetryStoreConfig.TryGetMetric(sampleKey.Metric, out metricConfig))
 					{
 						await CombineValuesAsync(sampleKey.Store, metricConfig, sampleKey.Group, sampleKey.Time, sampleValues, cancellationToken);
+
+						if (metricConfig.TopN > 0 || metricConfig.BottomN > 0)
+						{
+							tbConstraint.Add(sampleKey, metricConfig);
+						}
 					}
 				}
 			}
@@ -270,6 +277,48 @@ namespace HordeServer.Telemetry.Metrics
 				}
 				throw;
 			}
+
+			// Constrain to TopN or BottomN
+			HashSet<ObjectId> removeIds = new HashSet<ObjectId>();			
+			if (tbConstraint.Count > 0)
+			{
+				foreach ((SampleKey sampleKey, MetricConfig metricConfig) in tbConstraint)
+				{
+					FilterDefinition < MetricDocument > filter = FilterDefinition<MetricDocument>.Empty;
+					filter &= Builders<MetricDocument>.Filter.Eq(x => x.TelemetryStoreId, sampleKey.Store);
+					filter &= Builders<MetricDocument>.Filter.Eq(x => x.Time, sampleKey.Time);
+					filter &= Builders<MetricDocument>.Filter.Eq(x => x.MetricId, sampleKey.Metric);
+
+					List<MetricDocument> results = await _metrics.Find(filter).SortByDescending(x => x.Value).ToListAsync(cancellationToken);
+
+					if (results.Count <= metricConfig.TopN && results.Count <= metricConfig.BottomN)
+					{
+						continue;
+					}
+
+					if (metricConfig.TopN > 0)
+					{
+						for (int i = metricConfig.TopN; i < results.Count; i++)
+						{
+							removeIds.Add(results[i].Id);
+						}
+					}
+					else
+					{
+						for (int i = 0; i < results.Count - metricConfig.BottomN; i++)
+						{
+							removeIds.Add(results[i].Id);
+						}
+					}
+				}			
+			}
+
+			if (removeIds.Count > 0)
+			{
+				FilterDefinition<MetricDocument> filter = FilterDefinition<MetricDocument>.Empty;
+				filter &= Builders<MetricDocument>.Filter.In(x => x.Id, removeIds);
+				await _metrics.DeleteManyAsync(filter, cancellationToken);
+			}
 		}
 
 		async Task CombineValuesAsync(TelemetryStoreId telemetryStoreId, MetricConfig metricConfig, string group, DateTime time, List<double> values, CancellationToken cancellationToken)
diff --git a/Engine/Source/Programs/Horde/Plugins/Analytics/HordeServer.Analytics/Telemetry/Metrics/MetricConfig.cs b/Engine/Source/Programs/Horde/Plugins/Analytics/HordeServer.Analytics/Telemetry/Metrics/MetricConfig.cs
--- a/Engine/Source/Programs/Horde/Plugins/Analytics/HordeServer.Analytics/Telemetry/Metrics/MetricConfig.cs
+++ b/Engine/Source/Programs/Horde/Plugins/Analytics/HordeServer.Analytics/Telemetry/Metrics/MetricConfig.cs
@@ -114,6 +114,16 @@ namespace HordeServer.Telemetry.Metrics
 		/// </summary>
 		public int Percentile { get; set; } = 95;
 
+		/// <summary>
+		/// 
+		/// </summary>
+		public int TopN { get; set; } = 0;
+
+		/// <summary>
+		/// 
+		/// </summary>
+		public int BottomN { get; set; } = 0;
+
 		/// <summary>
 		/// Interval for each metric. Supports times such as "2d", "1h", "1h30m", "20s".
 		/// </summary>
