Changed the Context and LoadingNames in EditorTelemetry to be consistent with the Span Names

diff --git a/Engine/Plugins/Experimental/EditorTelemetry/Source/Private/EditorTelemetry.cpp b/Engine/Plugins/Experimental/EditorTelemetry/Source/Private/EditorTelemetry.cpp
--- a/Engine/Plugins/Experimental/EditorTelemetry/Source/Private/EditorTelemetry.cpp
+++ b/Engine/Plugins/Experimental/EditorTelemetry/Source/Private/EditorTelemetry.cpp
@@ -148,12 +148,13 @@ void FEditorTelemetry::RecordEvent_Cooking(TArray<FAnalyticsEventAttribute> Attr
 #endif
 }
 
-void FEditorTelemetry::RecordEvent_Loading(const FString& LoadingName, double LoadingSeconds, TArray<FAnalyticsEventAttribute> Attributes )
+void FEditorTelemetry::RecordEvent_Loading(const FString& Context, double LoadingSeconds, TArray<FAnalyticsEventAttribute> Attributes )
 {
-	const int SchemaVersion = 3;
+	const int SchemaVersion = 4;
 
 	Attributes.Emplace(TEXT("SchemaVersion"), SchemaVersion);
-	Attributes.Emplace(TEXT("LoadingName"), LoadingName);
+	Attributes.Emplace(TEXT("Context"), Context);
+	Attributes.Emplace(TEXT("LoadingName"), Context);
 	Attributes.Emplace(TEXT("LoadingSeconds"), LoadingSeconds);
 	
 #if ENABLE_COOK_STATS
@@ -204,10 +205,10 @@ void FEditorTelemetry::RecordEvent_DDCResource(const FString& Context, TArray<FA
 
 	GatherDerivedDataCacheResourceStats(ResourceStats);
 
-	const int SchemaVersion = 3;
+	const int SchemaVersion = 4;
 
 	Attributes.Emplace(TEXT("SchemaVersion"), SchemaVersion);
-	Attributes.Emplace(TEXT("LoadingName"), Context);
+	Attributes.Emplace(TEXT("Context"), Context);
 	
 	// Send a resource event per asset type
 	for (const FDerivedDataCacheResourceStat& Stat : ResourceStats)
@@ -245,10 +246,10 @@ void FEditorTelemetry::RecordEvent_DDCResource(const FString& Context, TArray<FA
 void FEditorTelemetry::RecordEvent_DDCSummary(const FString& Context, TArray<FAnalyticsEventAttribute> Attributes)
 {
 #if ENABLE_COOK_STATS
-	const int SchemaVersion = 3;
+	const int SchemaVersion = 4;
 
 	Attributes.Emplace(TEXT("SchemaVersion"), SchemaVersion);
-	Attributes.Emplace(TEXT("LoadingName"), Context);
+	Attributes.Emplace(TEXT("Context"), Context);
 
 	// Gather the summary stats
 	FDerivedDataCacheSummaryStats SummaryStats;
@@ -280,10 +281,10 @@ void FEditorTelemetry::RecordEvent_Zen(const FString& Context, TArray<FAnalytics
 	// Gather Zen analytics
 	if (UE::Zen::IsDefaultServicePresent())
 	{
-		const int SchemaVersion = 1;
+		const int SchemaVersion = 2;
 
 		Attributes.Emplace(TEXT("SchemaVersion"), SchemaVersion);
-		Attributes.Emplace(TEXT("LoadingName"), Context);
+		Attributes.Emplace(TEXT("Context"), Context);
 
 		UE::Zen::GetDefaultServiceInstance().GatherAnalytics(Attributes);
 		FStudioTelemetry::Get().RecordEvent(TEXT("Core.Zen"), Attributes);
@@ -295,10 +296,10 @@ void FEditorTelemetry::RecordEvent_VirtualAssets(const FString& Context, TArray<
 {
 	if (UE::Virtualization::IVirtualizationSystem::Get().IsEnabled())
 	{
-		const int SchemaVersion = 1;
+		const int SchemaVersion = 2;
 
 		Attributes.Emplace(TEXT("SchemaVersion"), SchemaVersion);
-		Attributes.Emplace(TEXT("LoadingName"), Context);
+		Attributes.Emplace(TEXT("Context"), Context);
 
 		// Gather Virtualization analytics
 		UE::Virtualization::IVirtualizationSystem::Get().GatherAnalytics(Attributes);
@@ -500,9 +501,9 @@ void FEditorTelemetry::StartSession()
 				FStudioTelemetry::Get().EndSpan(EditorBootSpan);
 
 				// Callback is received when the editor has booted but has not been initialized
-				FEditorTelemetry::RecordEvent_Loading(TEXT("BootEditor"), EditorBootSpan->GetDuration(), EditorBootSpan->GetAttributes());
-				FEditorTelemetry::RecordEvent_CoreSystems(TEXT("BootEditor"), EditorBootSpan->GetAttributes());
-
+				FEditorTelemetry::RecordEvent_Loading(TEXT("Editor.Boot"), EditorBootSpan->GetDuration(), EditorBootSpan->GetAttributes());
+				FEditorTelemetry::RecordEvent_CoreSystems(TEXT("Editor.Boot"), EditorBootSpan->GetAttributes());
+				
 				EditorInitilizeSpan = FStudioTelemetry::Get().StartSpan(EditorInitilizeSpanName, EditorSpan);
 			});
 
@@ -519,9 +520,9 @@ void FEditorTelemetry::StartSession()
 				FStudioTelemetry::Get().EndSpan(EditorInitilizeSpan);
 				EditorInteractSpan = FStudioTelemetry::Get().StartSpan(EditorInteractSpanName, EditorSpan);
 
-				FEditorTelemetry::RecordEvent_Loading(TEXT("TotalEditorStartup"), TimeToInitializeEditor, Attributes);
-				FEditorTelemetry::RecordEvent_CoreSystems(TEXT("TotalEditorStartup"), Attributes);
-
+				FEditorTelemetry::RecordEvent_Loading(TEXT("Editor.Initialize"), TimeToInitializeEditor, Attributes);
+				FEditorTelemetry::RecordEvent_CoreSystems(TEXT("Editor.Initialize"), Attributes);
+				
 				// Install callbacks for Open Asset Dialogue
 				GEditor->GetEditorSubsystem<UAssetEditorSubsystem>()->OnAssetEditorRequestedOpen().AddLambda([this](UObject* Asset)
 					{
@@ -538,8 +539,6 @@ void FEditorTelemetry::StartSession()
 						{
 							Attributes.Emplace(TEXT("AssetPath"), Asset->GetFullName());
 							Attributes.Emplace(TEXT("AssetClass"), Asset->GetClass()->GetName());
-
-							FEditorTelemetry::RecordEvent_Loading(TEXT("OpenAssetEditor"), FPlatformTime::Seconds() - AssetOpenStartTime, Attributes);
 						}
 
 						FStudioTelemetry::Get().EndSpan(OpenAssetEditorSpan, Attributes);
@@ -607,8 +606,8 @@ void FEditorTelemetry::StartSession()
 
 				FStudioTelemetry::Get().EndSpan(PIELoadMapSpan);
 
-				FEditorTelemetry::RecordEvent_Loading(TEXT("PIE.LoadMapTime"), PIELoadMapSpan->GetDuration(), PIELoadMapSpan->GetAttributes());
-				FEditorTelemetry::RecordEvent_CoreSystems(TEXT("PIE.LoadMapTime"), PIELoadMapSpan->GetAttributes());
+				FEditorTelemetry::RecordEvent_Loading(TEXT("PIE.LoadMap"), PIELoadMapSpan->GetDuration(), PIELoadMapSpan->GetAttributes());
+				FEditorTelemetry::RecordEvent_CoreSystems(TEXT("PIE.LoadMap"), PIELoadMapSpan->GetAttributes());
 			});
 
 		FWorldDelegates::OnPIEReady.AddLambda([this](UGameInstance* GameInstance)
@@ -629,9 +628,6 @@ void FEditorTelemetry::StartSession()
 						FStudioTelemetry::Get().EndSpan(PIEStartupSpan);
 
 						// Record the time from start PIE to PIE
-						FEditorTelemetry::RecordEvent_Loading(TEXT("PIE.TotalStartupTime"), PIEStartupSpan->GetDuration(), PIEStartupSpan->GetAttributes());
-						FEditorTelemetry::RecordEvent_CoreSystems(TEXT("PIE.TotalStartupTime"), PIEStartupSpan->GetAttributes());
-
 						if (PIESessionCount == 0)
 						{
 							const double TimeInEditor = EditorLoadMapSpan.IsValid() ? EditorLoadMapSpan->GetDuration() : 0.0;
@@ -639,8 +635,11 @@ void FEditorTelemetry::StartSession()
 							const double TimeToBootToPIE = TimeToBootEditor + TimeInEditor + TimeToStartPIE;
 
 							// Record the absolute time from editor boot to PIE
-							FEditorTelemetry::RecordEvent_Loading(TEXT("TimeToPIE"), TimeToBootToPIE, PIEStartupSpan->GetAttributes());
-							FEditorTelemetry::RecordEvent_CoreSystems(TEXT("TimeToPIE"), PIEStartupSpan->GetAttributes());
+							FEditorTelemetry::RecordEvent_Loading(TEXT("TimeToPIE"), TimeToBootToPIE, PIEStartupSpan->GetAttributes());	
+
+							// Record the first PIE transition
+							FEditorTelemetry::RecordEvent_Loading(TEXT("PIE.Startup"), PIEStartupSpan->GetDuration(), PIEStartupSpan->GetAttributes());
+							FEditorTelemetry::RecordEvent_CoreSystems(TEXT("PIE.Startup"), PIEStartupSpan->GetAttributes());
 						}
 					}
 
@@ -667,8 +666,12 @@ void FEditorTelemetry::StartSession()
 					// PIE has shutdown, ie. the user has pressed the Stop PIE button, and we are going back to interactive Editor mode	
 					FStudioTelemetry::Get().EndSpan(PIESpan);
 
-					FEditorTelemetry::RecordEvent_Loading(TEXT("PIE.EndTime"), PIESpan->GetDuration(), PIESpan->GetAttributes());
-					FEditorTelemetry::RecordEvent_CoreSystems(TEXT("PIE.EndTime"), PIESpan->GetAttributes());
+					if (PIESessionCount == 0)
+					{
+						// Record the first PIE shutdown
+						FEditorTelemetry::RecordEvent_Loading(TEXT("PIE.Shutdown"), PIESpan->GetDuration(), PIESpan->GetAttributes());
+						FEditorTelemetry::RecordEvent_CoreSystems(TEXT("PIE.Shutdown"), PIESpan->GetAttributes());
+					}
 				}
 
 				PIESessionCount++;
@@ -705,8 +708,8 @@ void FEditorTelemetry::StartSession()
 
 					FStudioTelemetry::Get().EndSpan(EditorLoadMapSpan);
 
-					FEditorTelemetry::RecordEvent_Loading(TEXT("LoadMap"), EditorLoadMapSpan->GetDuration(), EditorLoadMapSpan->GetAttributes());
-					FEditorTelemetry::RecordEvent_CoreSystems(TEXT("LoadMap"), EditorLoadMapSpan->GetAttributes());
+					FEditorTelemetry::RecordEvent_Loading(TEXT("Editor.LoadMap"), EditorLoadMapSpan->GetDuration(), EditorLoadMapSpan->GetAttributes());
+					FEditorTelemetry::RecordEvent_CoreSystems(TEXT("Editor.LoadMap"), EditorLoadMapSpan->GetAttributes());
 
 					EditorLoadMapSpan.Reset();
 				}
diff --git a/Engine/Source/Programs/Horde/HordeServer/Defaults/default-metrics.telemetry.json b/Engine/Source/Programs/Horde/HordeServer/Defaults/default-metrics.telemetry.json
--- a/Engine/Source/Programs/Horde/HordeServer/Defaults/default-metrics.telemetry.json
+++ b/Engine/Source/Programs/Horde/HordeServer/Defaults/default-metrics.telemetry.json
@@ -17,7 +17,7 @@
 		},
 		{
 			"id": "core-loading-user-editorboottime",
-			"filter": "$[?(@.Payload.EventName == 'Core.Loading' && @.Payload.LoadingName=='BootEditor' && @.Payload.Config_IsBuildMachine==false && @.Payload.Config_IsRunningCommandlet==false && @.Payload.LoadingSeconds>5 && @.Payload.LoadingSeconds<3600)]",
+			"filter": "$[?(@.Payload.EventName == 'Core.Loading' && @.Payload.LoadingName=='Editor.Boot' && @.Payload.Config_IsBuildMachine==false && @.Payload.Config_IsRunningCommandlet==false && @.Payload.LoadingSeconds>5 && @.Payload.LoadingSeconds<3600)]",
 			"property": "$.Payload.LoadingSeconds",
 			"groupBy": "$.Payload.Build_BranchName,$.Payload.ProjectName,$.Payload.ProjectName",
 			"function": "Average",
@@ -25,7 +25,7 @@
 		},
 		{
 			"id": "core-loading-user-editorstartuptime",
-			"filter": "$[?(@.Payload.EventName == 'Core.Loading' && @.Payload.LoadingName=='TotalEditorStartup' && @.Payload.Config_IsBuildMachine==false && @.Payload.Config_IsRunningCommandlet==false && @.Payload.LoadingSeconds>5 && @.Payload.LoadingSeconds<3600)]",
+			"filter": "$[?(@.Payload.EventName == 'Core.Loading' && @.Payload.LoadingName=='Editor.Initialize' && @.Payload.Config_IsBuildMachine==false && @.Payload.Config_IsRunningCommandlet==false && @.Payload.LoadingSeconds>5 && @.Payload.LoadingSeconds<3600)]",
 			"property": "$.Payload.LoadingSeconds",
 			"groupBy": "$.Payload.Build_BranchName,$.Payload.ProjectName,$.Payload.ProjectName",
 			"function": "Average",
@@ -33,7 +33,7 @@
 		},
 		{
 			"id": "core-loading-user-loadmaptime",
-			"filter": "$[?(@.Payload.EventName == 'Core.Loading' && @.Payload.LoadingName=='LoadMap' && @.Payload.Config_IsBuildMachine==false && @.Payload.Config_IsRunningCommandlet==false && @.Payload.LoadingSeconds>5 && @.Payload.LoadingSeconds<3600)]",
+			"filter": "$[?(@.Payload.EventName == 'Core.Loading' && @.Payload.LoadingName=='Editor.LoadMap' && @.Payload.Config_IsBuildMachine==false && @.Payload.Config_IsRunningCommandlet==false && @.Payload.LoadingSeconds>5 && @.Payload.LoadingSeconds<3600)]",
 			"property": "$.Payload.LoadingSeconds",
 			"groupBy": "$.Payload.Build_BranchName,$.Payload.ProjectName,$.Payload.MapName",
 			"function": "Average",
@@ -41,14 +41,14 @@
 		},
 		{
 			"id": "core-loading-user-loadmap-count",
-			"filter": "$[?(@.Payload.EventName == 'Core.Loading' && @.Payload.LoadingName=='LoadMap' && @.Payload.Config_IsBuildMachine==false && @.Payload.Config_IsRunningCommandlet==false && @.Payload.LoadingSeconds>5 && @.Payload.LoadingSeconds<3600)]",
+			"filter": "$[?(@.Payload.EventName == 'Core.Loading' && @.Payload.LoadingName=='Editor.LoadMap' && @.Payload.Config_IsBuildMachine==false && @.Payload.Config_IsRunningCommandlet==false && @.Payload.LoadingSeconds>5 && @.Payload.LoadingSeconds<3600)]",
 			"groupBy": "$.Payload.Build_BranchName,$.Payload.ProjectName,$.Payload.MapName",
 			"function": "Count",
 			"interval": "03:00:00"
 		},
 		{
 			"id": "core-loading-user-pieloadmaptime",
-			"filter": "$[?(@.Payload.EventName == 'Core.Loading' && @.Payload.LoadingName=='PIE.LoadMapTime' && @.Payload.Config_IsBuildMachine==false && @.Payload.Config_IsRunningCommandlet==false && @.Payload.LoadingSeconds>5 && @.Payload.LoadingSeconds<3600)]",
+			"filter": "$[?(@.Payload.EventName == 'Core.Loading' && @.Payload.LoadingName=='PIE.LoadMap' && @.Payload.Config_IsBuildMachine==false && @.Payload.Config_IsRunningCommandlet==false && @.Payload.LoadingSeconds>5 && @.Payload.LoadingSeconds<3600)]",
 			"property": "$.Payload.LoadingSeconds",
 			"groupBy": "$.Payload.Build_BranchName,$.Payload.ProjectName,$.Payload.MapName",
 			"function": "Average",
@@ -64,7 +64,7 @@
 		},
 		{
 			"id": "core-loading-user-piestartuptime",
-			"filter": "$[?(@.Payload.EventName == 'Core.Loading' && @.Payload.LoadingName=='PIE.TotalStartupTime' && @.Payload.Config_IsBuildMachine==false && @.Payload.Config_IsRunningCommandlet==false && @.Payload.LoadingSeconds>5 && @.Payload.LoadingSeconds<3600)]",
+			"filter": "$[?(@.Payload.EventName == 'Core.Loading' && @.Payload.LoadingName=='PIE.Startup' && @.Payload.Config_IsBuildMachine==false && @.Payload.Config_IsRunningCommandlet==false && @.Payload.LoadingSeconds>5 && @.Payload.LoadingSeconds<3600)]",
 			"property": "$.Payload.LoadingSeconds",
 			"groupBy": "$.Payload.Build_BranchName,$.Payload.ProjectName,$.Payload.MapName",
 			"function": "Average",
@@ -72,7 +72,7 @@
 		},
 		{
 			"id": "core-loading-user-editorboottime-kpi",
-			"filter": "$[?(@.Payload.EventName == 'Core.Loading' && @.Payload.LoadingName=='BootEditor' && @.Payload.Config_IsBuildMachine==false && @.Payload.Config_IsRunningCommandlet==false && @.Payload.LoadingSeconds>5 && @.Payload.LoadingSeconds<3600)]",
+			"filter": "$[?(@.Payload.EventName == 'Core.Loading' && @.Payload.LoadingName=='Editor.Boot' && @.Payload.Config_IsBuildMachine==false && @.Payload.Config_IsRunningCommandlet==false && @.Payload.LoadingSeconds>5 && @.Payload.LoadingSeconds<3600)]",
 			"property": "$.Payload.LoadingSeconds",
 			"groupBy": "$.Payload.Build_BranchName,$.Payload.ProjectName",
 			"function": "Average",
@@ -80,7 +80,7 @@
 		},
 		{
 			"id": "core-loading-user-editorstartuptime-kpi",
-			"filter": "$[?(@.Payload.EventName == 'Core.Loading' && @.Payload.LoadingName=='TotalEditorStartup' && @.Payload.Config_IsBuildMachine==false && @.Payload.Config_IsRunningCommandlet==false && @.Payload.LoadingSeconds>5 && @.Payload.LoadingSeconds<3600)]",
+			"filter": "$[?(@.Payload.EventName == 'Core.Loading' && @.Payload.LoadingName=='Editor.Initialize' && @.Payload.Config_IsBuildMachine==false && @.Payload.Config_IsRunningCommandlet==false && @.Payload.LoadingSeconds>5 && @.Payload.LoadingSeconds<3600)]",
 			"property": "$.Payload.LoadingSeconds",
 			"groupBy": "$.Payload.Build_BranchName,$.Payload.ProjectName",
 			"function": "Average",
@@ -88,7 +88,7 @@
 		},
 		{
 			"id": "core-loading-user-loadmaptime-kpi",
-			"filter": "$[?(@.Payload.EventName == 'Core.Loading' && @.Payload.LoadingName=='LoadMap' && @.Payload.Config_IsBuildMachine==false && @.Payload.Config_IsRunningCommandlet==false && @.Payload.LoadingSeconds>5 && @.Payload.LoadingSeconds<3600)]",
+			"filter": "$[?(@.Payload.EventName == 'Core.Loading' && @.Payload.LoadingName=='Editor.LoadMap' && @.Payload.Config_IsBuildMachine==false && @.Payload.Config_IsRunningCommandlet==false && @.Payload.LoadingSeconds>5 && @.Payload.LoadingSeconds<3600)]",
 			"property": "$.Payload.LoadingSeconds",
 			"groupBy": "$.Payload.Build_BranchName,$.Payload.ProjectName",
 			"function": "Average",
@@ -96,7 +96,7 @@
 		},
 		{
 			"id": "core-loading-user-piestartuptime-kpi",
-			"filter": "$[?(@.Payload.EventName == 'Core.Loading' && @.Payload.LoadingName=='PIE.TotalStartupTime' && @.Payload.Config_IsBuildMachine==false && @.Payload.Config_IsRunningCommandlet==false && @.Payload.LoadingSeconds>5 && @.Payload.LoadingSeconds<3600)]",
+			"filter": "$[?(@.Payload.EventName == 'Core.Loading' && @.Payload.LoadingName=='PIE.Startup' && @.Payload.Config_IsBuildMachine==false && @.Payload.Config_IsRunningCommandlet==false && @.Payload.LoadingSeconds>5 && @.Payload.LoadingSeconds<3600)]",
 			"property": "$.Payload.LoadingSeconds",
 			"groupBy": "$.Payload.Build_BranchName,$.Payload.ProjectName",
 			"function": "Average",
@@ -161,7 +161,7 @@
 		},
 		{
 			"id": "core-loading-user-editorboottime-percentile",
-			"filter": "$[?(@.Payload.EventName == 'Core.Loading' && @.Payload.LoadingName=='BootEditor' && @.Payload.Config_IsBuildMachine==false && @.Payload.Config_IsRunningCommandlet==false && @.Payload.LoadingSeconds>5 && @.Payload.LoadingSeconds<3600 )]",
+			"filter": "$[?(@.Payload.EventName == 'Core.Loading' && @.Payload.LoadingName=='Editor.Boot' && @.Payload.Config_IsBuildMachine==false && @.Payload.Config_IsRunningCommandlet==false && @.Payload.LoadingSeconds>5 && @.Payload.LoadingSeconds<3600 )]",
 			"groupBy": "$.Payload.Build_BranchName,$.Payload.ProjectName",
 			"property": "$.Payload.LoadingSeconds",
 			"function": "Percentile",
@@ -170,7 +170,7 @@
 		},
 		{
 			"id": "core-loading-user-editorstartuptime-percentile",
-			"filter": "$[?(@.Payload.EventName == 'Core.Loading' && @.Payload.LoadingName=='TotalEditorStartup' && @.Payload.Config_IsBuildMachine==false && @.Payload.Config_IsRunningCommandlet==false && @.Payload.LoadingSeconds>5 && @.Payload.LoadingSeconds<3600 )]",
+			"filter": "$[?(@.Payload.EventName == 'Core.Loading' && @.Payload.LoadingName=='Editor.Initialize' && @.Payload.Config_IsBuildMachine==false && @.Payload.Config_IsRunningCommandlet==false && @.Payload.LoadingSeconds>5 && @.Payload.LoadingSeconds<3600 )]",
 			"groupBy": "$.Payload.Build_BranchName,$.Payload.ProjectName",
 			"property": "$.Payload.LoadingSeconds",
 			"function": "Percentile",
@@ -179,7 +179,7 @@
 		},
 		{
 			"id": "core-loading-user-loadmaptime-percentile",
-			"filter": "$[?(@.Payload.EventName == 'Core.Loading' && @.Payload.LoadingName=='LoadMap' && @.Payload.Config_IsBuildMachine==false && @.Payload.LoadingSeconds>5 && @.Payload.LoadingSeconds<3600)]",
+			"filter": "$[?(@.Payload.EventName == 'Core.Loading' && @.Payload.LoadingName=='Editor.LoadMap' && @.Payload.Config_IsBuildMachine==false && @.Payload.LoadingSeconds>5 && @.Payload.LoadingSeconds<3600)]",
 			"groupBy": "$.Payload.Build_BranchName,$.Payload.ProjectName,$.Payload.MapName",
 			"property": "$.Payload.LoadingSeconds",
 			"function": "Percentile",
@@ -188,7 +188,7 @@
 		},
 		{
 			"id": "core-loading-user-piestartuptime-percentile",
-			"filter": "$[?(@.Payload.EventName == 'Core.Loading' && @.Payload.LoadingName=='PIE.TotalStartupTime' && @.Payload.Config_IsBuildMachine==false && @.Payload.LoadingSeconds>5 && @.Payload.LoadingSeconds<3600)]",
+			"filter": "$[?(@.Payload.EventName == 'Core.Loading' && @.Payload.LoadingName=='PIE.Startup' && @.Payload.Config_IsBuildMachine==false && @.Payload.LoadingSeconds>5 && @.Payload.LoadingSeconds<3600)]",
 			"groupBy": "$.Payload.Build_BranchName,$.Payload.ProjectName,$.Payload.MapName",
 			"property": "$.Payload.LoadingSeconds",
 			"function": "Percentile",
