Horde: Fix agent rates being left empty. Was broken by config refactor in CL 23685875; possibly not noticed due to cached data in Redis.

#rnx

diff --git a/Engine/Source/Programs/Horde/Plugins/Build/HordeServer.Build.Tests/Agents/AgentServiceTest.cs b/Engine/Source/Programs/Horde/Plugins/Build/HordeServer.Build.Tests/Agents/AgentServiceTest.cs
--- a/Engine/Source/Programs/Horde/Plugins/Build/HordeServer.Build.Tests/Agents/AgentServiceTest.cs
+++ b/Engine/Source/Programs/Horde/Plugins/Build/HordeServer.Build.Tests/Agents/AgentServiceTest.cs
@@ -8,6 +8,7 @@ using EpicGames.Horde.Jobs;
 using HordeCommon.Rpc.Messages;
 using HordeServer.Agents;
 using HordeServer.Auditing;
+using HordeServer.Plugins;
 using HordeServer.Utilities;
 using Microsoft.AspNetCore.Mvc;
 using Microsoft.Extensions.DependencyInjection;
@@ -144,7 +145,8 @@ public class AgentServiceTest : BuildTestSetup
 			new AgentRateConfig() { Condition = "aws-instance-type == 'c5.24xlarge' && osfamily == 'windows'", Rate = 200, },
 			new AgentRateConfig() { Condition = "aws-instance-type == 'c4.4xlarge' && osfamily == 'windows'", Rate = 300 }
 		};
-		await AgentService.UpdateRateTableAsync(agentRateConfigs);
+
+		UpdateConfig(x => ((ComputeConfig)x.Plugins[new PluginName("compute")]).Rates = agentRateConfigs);
 
 		double? rate1 = await AgentService.GetRateAsync(agent1.Id);
 		Assert.AreEqual(200, rate1!.Value, 0.1);
diff --git a/Engine/Source/Programs/Horde/Plugins/Compute/HordeServer.Compute/Agents/AgentService.cs b/Engine/Source/Programs/Horde/Plugins/Compute/HordeServer.Compute/Agents/AgentService.cs
--- a/Engine/Source/Programs/Horde/Plugins/Compute/HordeServer.Compute/Agents/AgentService.cs
+++ b/Engine/Source/Programs/Horde/Plugins/Compute/HordeServer.Compute/Agents/AgentService.cs
@@ -20,6 +20,7 @@ using HordeServer.Tasks;
 using HordeServer.Utilities;
 using Microsoft.Extensions.Hosting;
 using Microsoft.Extensions.Logging;
+using Microsoft.Extensions.Options;
 using OpenTelemetry.Trace;
 using StackExchange.Redis;
 
@@ -71,16 +72,12 @@ namespace HordeServer.Agents
 		readonly ITaskSource[] _taskSources;
 		readonly IRedisService _redisService;
 		readonly IHostApplicationLifetime _applicationLifetime;
+		readonly IOptionsMonitor<ComputeConfig> _computeConfig;
 		readonly IClock _clock;
 		readonly Tracer _tracer;
 		readonly ILogger _logger;
 		readonly ITicker _ticker;
 
-		readonly RedisStringKey<AgentRateTable> _agentRateTableData = new("compute:agent-rates");
-
-		/// <summary>Lazily updated costs for different agent types</summary>
-		readonly AsyncCachedValue<AgentRateTable?> _cachedRates;
-
 		/// <summary>Lazily updated list of current pools</summary>
 		readonly AsyncCachedValue<IReadOnlyList<IPoolConfig>> _cachedPools;
 
@@ -102,6 +99,7 @@ namespace HordeServer.Agents
 			IEnumerable<ITaskSource> taskSources,
 			IRedisService redisService,
 			IHostApplicationLifetime applicationLifetime,
+			IOptionsMonitor<ComputeConfig> computeConfig,
 			IClock clock,
 			Tracer tracer,
 			Meter meter,
@@ -111,10 +109,10 @@ namespace HordeServer.Agents
 			_leases = leases;
 			_aclService = aclService;
 			_downtimeService = downtimeService;
-			_cachedRates = new AsyncCachedValue<AgentRateTable?>(_ => redisService.GetDatabase().StringGetAsync(_agentRateTableData), TimeSpan.FromMinutes(2));
 			_cachedPools = new AsyncCachedValue<IReadOnlyList<IPoolConfig>>(poolCollection.GetConfigsAsync, TimeSpan.FromSeconds(30.0));
 			_taskSources = taskSources.ToArray();
 			_applicationLifetime = applicationLifetime;
+			_computeConfig = computeConfig;
 			_redisService = redisService;
 			_clock = clock;
 			_ticker = clock.AddTicker($"{nameof(AgentService)}.{nameof(TickAsync)}", TimeSpan.FromSeconds(30.0), TickAsync, logger);
@@ -139,7 +137,6 @@ namespace HordeServer.Agents
 		/// <inheritdoc/>
 		public async ValueTask DisposeAsync()
 		{
-			await _cachedRates.DisposeAsync();
 			await _cachedPools.DisposeAsync();
 			await _ticker.DisposeAsync();
 		}
@@ -753,16 +750,6 @@ namespace HordeServer.Agents
 			return _leases.GetAsync(leaseId, cancellationToken);
 		}
 
-		/// <summary>
-		/// Update the rate table for different agent types
-		/// </summary>
-		/// <param name="entries">New entries for the rate table</param>
-		/// <returns></returns>
-		public async Task UpdateRateTableAsync(List<AgentRateConfig> entries)
-		{
-			await _redisService.GetDatabase().StringSetAsync(_agentRateTableData, new AgentRateTable { Entries = entries });
-		}
-
 		/// <summary>
 		/// Gets the rate for the given agent
 		/// </summary>
@@ -792,19 +779,15 @@ namespace HordeServer.Agents
 				double rate = 0.0;
 
 				// Get the rate table
-				AgentRateTable? rateTable = await _cachedRates.GetAsync(cancellationToken);
-				if (rateTable != null && rateTable.Entries.Count > 0)
+				IAgent? agent = await GetAgentAsync(agentId, cancellationToken);
+				if (agent != null)
 				{
-					IAgent? agent = await GetAgentAsync(agentId, cancellationToken);
-					if (agent != null)
+					foreach (AgentRateConfig config in _computeConfig.CurrentValue.Rates)
 					{
-						foreach (AgentRateConfig config in rateTable.Entries)
+						if (config.Condition != null && config.Condition.Evaluate(x => agent.GetPropertyValues(x)))
 						{
-							if (config.Condition != null && config.Condition.Evaluate(x => agent.GetPropertyValues(x)))
-							{
-								rate = config.Rate;
-								break;
-							}
+							rate = config.Rate;
+							break;
 						}
 					}
 				}
