Horde: Add device audit logging, server side

diff --git a/Engine/Source/Programs/Horde/Plugins/Build/HordeServer.Build/BuildPlugin.cs b/Engine/Source/Programs/Horde/Plugins/Build/HordeServer.Build/BuildPlugin.cs
--- a/Engine/Source/Programs/Horde/Plugins/Build/HordeServer.Build/BuildPlugin.cs
+++ b/Engine/Source/Programs/Horde/Plugins/Build/HordeServer.Build/BuildPlugin.cs
@@ -3,6 +3,7 @@
 using System.Diagnostics.CodeAnalysis;
 using EpicGames.Core;
 using EpicGames.Horde.Artifacts;
+using EpicGames.Horde.Devices;
 using EpicGames.Horde.Jobs.Templates;
 using EpicGames.Horde.Streams;
 using HordeServer.Acls;
@@ -107,6 +108,7 @@ namespace HordeServer
 			services.AddSingleton<ICommitService, CommitService>();
 
 			services.AddSingleton<DeviceService>();
+			services.AddSingleton<IAuditLog<DeviceId>>(sp => sp.GetRequiredService<IAuditLogFactory<DeviceId>>().Create("Devices.Log", "DeviceId"));
 			services.AddSingleton<IBlockCache>(sp => CreateBlockCache());
 			services.AddSingleton<TestDataService>();
 
diff --git a/Engine/Source/Programs/Horde/Plugins/Build/HordeServer.Build/Devices/DeviceCollection.cs b/Engine/Source/Programs/Horde/Plugins/Build/HordeServer.Build/Devices/DeviceCollection.cs
--- a/Engine/Source/Programs/Horde/Plugins/Build/HordeServer.Build/Devices/DeviceCollection.cs
+++ b/Engine/Source/Programs/Horde/Plugins/Build/HordeServer.Build/Devices/DeviceCollection.cs
@@ -5,7 +5,9 @@ using EpicGames.Horde.Devices;
 using EpicGames.Horde.Jobs;
 using EpicGames.Horde.Streams;
 using EpicGames.Horde.Users;
+using HordeServer.Auditing;
 using HordeServer.Server;
+using HordeServer.Users;
 using HordeServer.Utilities;
 using Microsoft.Extensions.Logging;
 using MongoDB.Bson;
@@ -435,14 +437,19 @@ namespace HordeServer.Devices
 
 		readonly IMongoCollection<DevicePoolTelemetryDocument> _poolTelemetry;
 
+		readonly IUserCollection _userCollection;
+
+		readonly IAuditLog<DeviceId> _auditLog;
 		readonly ILogger<DeviceCollection> _logger;
 
 		/// <summary>
 		/// Constructor
 		/// </summary>
-		public DeviceCollection(IMongoService mongoService, ILogger<DeviceCollection> logger)
+		public DeviceCollection(IMongoService mongoService, IUserCollection userCollection, ILogger<DeviceCollection> logger, IAuditLog<DeviceId> auditLog)
 		{
+			_auditLog = auditLog;
 			_logger = logger;
+			_userCollection = userCollection;
 
 			_devices = mongoService.GetCollection<DeviceDocument>("Devices", keys => keys.Ascending(x => x.Name), unique: true);
 			_reservations = mongoService.GetCollection<DeviceReservationDocument>("Devices.Reservations");
@@ -529,29 +536,67 @@ namespace HordeServer.Devices
 
 			DateTime utcNow = DateTime.UtcNow;
 
+			IAuditLogChannel<DeviceId> logger = GetDeviceLogger(deviceId);
+			IDevice? current = await _devices.Find<DeviceDocument>(x => x.Id == deviceId).FirstOrDefaultAsync();
+
+			bool log = false;
+			if (current != null)
+			{
+				log = true;
+			}
+
 			if (modifiedByUserId != null)
 			{
 				updates.Add(updateBuilder.Set(x => x.ModifiedByUser, modifiedByUserId.ToString()));
+				if (log && current!.ModifiedByUser != modifiedByUserId.ToString())
+				{
+					IUser? user = await _userCollection.GetCachedUserAsync(modifiedByUserId);
+					logger.LogInformation("Modified by user {UserMame}", user?.Name ?? modifiedByUserId.ToString());
+				}				
 			}
 
 			if (newPoolId != null)
 			{
 				updates.Add(updateBuilder.Set(x => x.PoolId, newPoolId.Value));
+				if (log && current!.PoolId != newPoolId.Value)
+				{
+					logger.LogInformation("Assigned to pool {PoolId}", newPoolId.Value.ToString());
+				}				
 			}
 
 			if (newName != null)
 			{
 				updates.Add(updateBuilder.Set(x => x.Name, newName));
+				if (log && current!.Name != newName)
+				{
+					logger.LogInformation("Renamed to {DeviceName}", newName);
+				}				
 			}
 
 			if (newEnabled != null)
 			{
 				updates.Add(updateBuilder.Set(x => x.Enabled, newEnabled));
+
+				if (log && current!.Enabled != newEnabled)
+				{
+					if (newEnabled.Value)
+					{
+						logger.LogInformation("Device enabled");
+					}
+					else
+					{
+						logger.LogInformation("Device disabled");
+					}
+				}
 			}
 
 			if (newAddress != null)
 			{
 				updates.Add(updateBuilder.Set(x => x.Address, newAddress));
+				if (log && current!.Address != newAddress)
+				{
+					logger.LogInformation("Device address changed to {DeviceAddress}", newAddress);
+				}
 			}
 
 			if (!String.IsNullOrEmpty(newModelId))
@@ -559,10 +604,19 @@ namespace HordeServer.Devices
 				if (newModelId == "Base")
 				{
 					updates.Add(updateBuilder.Set(x => x.ModelId, null));
+
+					if (log && current!.ModelId != null)
+					{
+						logger.LogInformation("Device model changed to Base");
+					}
 				}
 				else
 				{
 					updates.Add(updateBuilder.Set(x => x.ModelId, newModelId));
+					if (log && current!.ModelId != newModelId)
+					{
+						logger.LogInformation("Device model changed to {DeviceModel}", newModelId);
+					}
 				}
 			}
 
@@ -579,6 +633,11 @@ namespace HordeServer.Devices
 					problemTime = utcNow;
 				}
 
+				if (log && problemTime != null && current!.ProblemTimeUtc != problemTime)
+				{
+					logger.LogInformation("Device reported a problem");
+				}
+
 				updates.Add(updateBuilder.Set(x => x.ProblemTimeUtc, problemTime));
 			}
 
@@ -591,6 +650,11 @@ namespace HordeServer.Devices
 				}
 
 				updates.Add(updateBuilder.Set(x => x.MaintenanceTimeUtc, maintenanceTime));
+
+				if (log && maintenanceTime != null && current!.MaintenanceTimeUtc != maintenanceTime)
+				{
+					logger.LogInformation("Device put into maintenance");
+				}
 			}
 
 			if (updates.Count > 0)
@@ -615,7 +679,6 @@ namespace HordeServer.Devices
 			FilterDefinition<DeviceDocument> filter = Builders<DeviceDocument>.Filter.Eq(x => x.Id, deviceId);
 			DeleteResult result = await _devices.DeleteOneAsync(filter);
 			return result.DeletedCount > 0;
-
 		}
 
 		/// <inheritdoc/>
@@ -1166,5 +1229,10 @@ namespace HordeServer.Devices
 				}
 			}
 		}
+
+		IAuditLogChannel<DeviceId> GetDeviceLogger(DeviceId deviceId)
+		{
+			return _auditLog[deviceId];
+		}
 	}
 }
diff --git a/Engine/Source/Programs/Horde/Plugins/Build/HordeServer.Build/Devices/DeviceService.cs b/Engine/Source/Programs/Horde/Plugins/Build/HordeServer.Build/Devices/DeviceService.cs
--- a/Engine/Source/Programs/Horde/Plugins/Build/HordeServer.Build/Devices/DeviceService.cs
+++ b/Engine/Source/Programs/Horde/Plugins/Build/HordeServer.Build/Devices/DeviceService.cs
@@ -9,6 +9,7 @@ using EpicGames.Horde.Jobs.Graphs;
 using EpicGames.Horde.Projects;
 using EpicGames.Horde.Streams;
 using EpicGames.Horde.Users;
+using HordeServer.Auditing;
 using HordeServer.Jobs;
 using HordeServer.Notifications;
 using HordeServer.Projects;
@@ -72,6 +73,7 @@ namespace HordeServer.Devices
 		readonly ITicker _ticker;
 		readonly ITicker _telemetryTicker;
 		readonly IServerInfo _serverInfo;
+		readonly IAuditLog<DeviceId> _auditLog;
 		readonly IOptionsMonitor<BuildServerConfig> _settings;
 		readonly IOptionsMonitor<BuildConfig> _buildConfig;
 
@@ -85,7 +87,7 @@ namespace HordeServer.Devices
 		/// <summary>
 		/// Device service constructor
 		/// </summary>
-		public DeviceService(IDeviceCollection devices, IUserCollection userCollection, JobService jobService, IStreamCollection streamCollection, IServerInfo serverInfo, IOptionsMonitor<BuildServerConfig> settings, IOptionsMonitor<BuildConfig> buildConfig, INotificationService notificationService, IClock clock, Tracer tracer, ILogger<DeviceService> logger)
+		public DeviceService(IDeviceCollection devices, IUserCollection userCollection, JobService jobService, IStreamCollection streamCollection, IServerInfo serverInfo, IOptionsMonitor<BuildServerConfig> settings, IOptionsMonitor<BuildConfig> buildConfig, INotificationService notificationService, IClock clock, Tracer tracer, ILogger<DeviceService> logger, IAuditLog<DeviceId> auditLog)
 		{
 			_userCollection = userCollection;
 			_devices = devices;
@@ -99,6 +101,7 @@ namespace HordeServer.Devices
 			_serverInfo = serverInfo;
 			_settings = settings;
 			_buildConfig = buildConfig;
+			_auditLog = auditLog;
 		}
 
 		/// <inheritdoc/>
@@ -122,6 +125,16 @@ namespace HordeServer.Devices
 			await _telemetryTicker.DisposeAsync();
 		}
 
+		/// <summary>
+		/// Get the logger for the specified device id
+		/// </summary>
+		/// <param name="deviceId"></param>
+		/// <returns></returns>
+		public IAuditLogChannel<DeviceId> GetDeviceLogger(DeviceId deviceId)
+		{
+			return _auditLog[deviceId];
+		}
+
 		/// <summary>
 		/// Ticks service
 		/// </summary>
diff --git a/Engine/Source/Programs/Horde/Plugins/Build/HordeServer.Build/Devices/DevicesController.cs b/Engine/Source/Programs/Horde/Plugins/Build/HordeServer.Build/Devices/DevicesController.cs
--- a/Engine/Source/Programs/Horde/Plugins/Build/HordeServer.Build/Devices/DevicesController.cs
+++ b/Engine/Source/Programs/Horde/Plugins/Build/HordeServer.Build/Devices/DevicesController.cs
@@ -10,6 +10,8 @@ using Microsoft.Extensions.Logging;
 using Microsoft.Extensions.Options;
 using MongoDB.Bson;
 using MongoDB.Driver;
+using EpicGames.Core;
+using HordeServer.Auditing;
 
 namespace HordeServer.Devices
 {
@@ -725,6 +727,33 @@ namespace HordeServer.Devices
 			return response;
 		}
 
+		/// <summary>
+		/// Retrieve historical information about a specific device
+		/// </summary>
+		/// <param name="deviceId">Id of the device to get information about</param>
+		/// <param name="minTime">Minimum time for records to return</param>
+		/// <param name="maxTime">Maximum time for records to return</param>
+		/// <param name="index">Offset of the first result</param>
+		/// <param name="count">Number of records to return</param>
+		/// <returns>Information about the requested agent</returns>
+		[HttpGet]
+		[Route("/api/v1/devices/{deviceId}/history")]
+		public async Task<ActionResult> GetDeviceHistoryAsync(DeviceId deviceId, [FromQuery] DateTime? minTime = null, [FromQuery] DateTime? maxTime = null, [FromQuery] int index = 0, [FromQuery] int count = 50)
+		{
+			if (!DeviceService.Authorize(DeviceAclAction.DeviceRead, User, _buildConfig.CurrentValue))
+			{
+				return Forbid();
+			}
+
+			Response.ContentType = "application/json";
+			Response.StatusCode = 200;
+			await Response.StartAsync();
+			await _deviceService.GetDeviceLogger(deviceId).FindAsync(Response.BodyWriter, minTime, maxTime, index, count);
+			await Response.CompleteAsync();
+
+			return Empty;
+		}
+
 		// LEGACY V1 API
 
 		enum LegacyPerfSpec
