Horde: Improve device audit messages to include the user changing properties

diff --git a/Engine/Source/Programs/Horde/Plugins/Build/HordeServer.Build/Devices/DeviceCollection.cs b/Engine/Source/Programs/Horde/Plugins/Build/HordeServer.Build/Devices/DeviceCollection.cs
--- a/Engine/Source/Programs/Horde/Plugins/Build/HordeServer.Build/Devices/DeviceCollection.cs
+++ b/Engine/Source/Programs/Horde/Plugins/Build/HordeServer.Build/Devices/DeviceCollection.cs
@@ -538,6 +538,11 @@ namespace HordeServer.Devices
 
 			IAuditLogChannel<DeviceId> logger = GetDeviceLogger(deviceId);
 			IDevice? current = await _devices.Find<DeviceDocument>(x => x.Id == deviceId).FirstOrDefaultAsync();
+			IUser? user = null;
+			if (modifiedByUserId != null)
+			{
+				user = await _userCollection.GetCachedUserAsync(modifiedByUserId);
+			}
 
 			bool log = false;
 			if (current != null)
@@ -549,9 +554,8 @@ namespace HordeServer.Devices
 			{
 				updates.Add(updateBuilder.Set(x => x.ModifiedByUser, modifiedByUserId.ToString()));
 				if (log && current!.ModifiedByUser != modifiedByUserId.ToString())
-				{
-					IUser? user = await _userCollection.GetCachedUserAsync(modifiedByUserId);
-					logger.LogInformation("Modified by user {UserMame}", user?.Name ?? modifiedByUserId.ToString());
+				{					
+					logger.LogInformation("Modified by user {UserName}", user?.Name ?? modifiedByUserId.ToString());
 				}				
 			}
 
@@ -569,7 +573,14 @@ namespace HordeServer.Devices
 				updates.Add(updateBuilder.Set(x => x.Name, newName));
 				if (log && current!.Name != newName)
 				{
-					logger.LogInformation("Renamed to {DeviceName}", newName);
+					if (user?.Name == null)
+					{
+						logger.LogInformation("Renamed to {DeviceName}", newName);
+					}
+					else
+					{
+						logger.LogInformation("Renamed to {DeviceName} by {UserName}", newName, user.Name);
+					}					
 				}				
 			}
 
@@ -581,11 +592,25 @@ namespace HordeServer.Devices
 				{
 					if (newEnabled.Value)
 					{
-						logger.LogInformation("Device enabled");
+						if (user?.Name == null)
+						{
+							logger.LogInformation("Device enabled");
+						}
+						else
+						{
+							logger.LogInformation("Device enabled by {UserName}", user.Name);
+						}							
 					}
 					else
 					{
-						logger.LogInformation("Device disabled");
+						if (user?.Name == null)
+						{
+							logger.LogInformation("Device disabled");
+						}
+						else
+						{
+							logger.LogInformation("Device disabled by {UserName}", user.Name);
+						}
 					}
 				}
 			}
@@ -595,7 +620,14 @@ namespace HordeServer.Devices
 				updates.Add(updateBuilder.Set(x => x.Address, newAddress));
 				if (log && current!.Address != newAddress)
 				{
-					logger.LogInformation("Device address changed to {DeviceAddress}", newAddress);
+					if (user?.Name == null)
+					{
+						logger.LogInformation("Device address changed to {DeviceAddress}", newAddress);
+					}
+					else
+					{
+						logger.LogInformation("Device address changed to {DeviceAddress} by {UserName}", newAddress, user.Name);
+					}
 				}
 			}
 
@@ -607,15 +639,29 @@ namespace HordeServer.Devices
 
 					if (log && current!.ModelId != null)
 					{
-						logger.LogInformation("Device model changed to Base");
+						if (user?.Name == null)
+						{
+							logger.LogInformation("Device model changed to Base");
+						}
+						else
+						{
+							logger.LogInformation("Device model changed to Base by {UserName}", user.Name);
+						}
 					}
 				}
 				else
 				{
 					updates.Add(updateBuilder.Set(x => x.ModelId, newModelId));
 					if (log && current!.ModelId != newModelId)
-					{
-						logger.LogInformation("Device model changed to {DeviceModel}", newModelId);
+					{						
+						if (user?.Name == null)
+						{
+							logger.LogInformation("Device model changed to {DeviceModel}", newModelId);
+						}
+						else
+						{
+							logger.LogInformation("Device model changed to {DeviceModel} by {UserName}", newModelId, user.Name);
+						}
 					}
 				}
 			}
@@ -651,9 +697,30 @@ namespace HordeServer.Devices
 
 				updates.Add(updateBuilder.Set(x => x.MaintenanceTimeUtc, maintenanceTime));
 
-				if (log && maintenanceTime != null && current!.MaintenanceTimeUtc != maintenanceTime)
+				if (log && current!.MaintenanceTimeUtc != maintenanceTime)
 				{
-					logger.LogInformation("Device put into maintenance");
+					if (maintenanceTime != null)
+					{
+						if (user?.Name == null)
+						{
+							logger.LogInformation("Device put into maintenance");
+						}
+						else
+						{
+							logger.LogInformation("Device put into maintenance by {UserName}", user.Name);
+						}
+					}
+					else
+					{
+						if (user?.Name == null)
+						{
+							logger.LogInformation("Device taken out of maintenance");
+						}
+						else
+						{
+							logger.LogInformation("Device taken out of maintenance by {UserName}", user.Name);
+						}
+					}
 				}
 			}
 
