Horde: Add support for device audit logs, frontend
#skipci

diff --git a/Engine/Source/Programs/Horde/HordeDashboard/src/App.tsx b/Engine/Source/Programs/Horde/HordeDashboard/src/App.tsx
--- a/Engine/Source/Programs/Horde/HordeDashboard/src/App.tsx
+++ b/Engine/Source/Programs/Horde/HordeDashboard/src/App.tsx
@@ -159,6 +159,7 @@ const Main: React.FC = () => {
                { path: "devices", element: <DeviceView /> },
                { path: "audit/agent/:agentId", element: <AuditLogView /> },
                { path: "audit/issue/:issueId", element: <AuditLogView /> },
+               { path: "audit/device/:deviceId", element: <AuditLogView /> },
                { path: "audit/template/:streamId/:templateId", element: <AuditLogView /> },
                { path: "automation", element: <AutomationView /> },
                { path: "tools", element: <ToolView /> },
diff --git a/Engine/Source/Programs/Horde/HordeDashboard/src/backend/Backend.ts b/Engine/Source/Programs/Horde/HordeDashboard/src/backend/Backend.ts
--- a/Engine/Source/Programs/Horde/HordeDashboard/src/backend/Backend.ts
+++ b/Engine/Source/Programs/Horde/HordeDashboard/src/backend/Backend.ts
@@ -1420,6 +1420,16 @@ export class Backend {
         });
     }
 
+    getDeviceHistory(id: string, query: AuditLogQuery): Promise<AuditLogEntry[]> {
+        return new Promise<AuditLogEntry[]>((resolve, reject) => {
+            this._fetch.get(`/api/v1/devices/${encodeURIComponent(id)}/history`, { params: query }).then((response) => {
+                const history = (response.data?.entries ?? []) as AuditLogEntry[];
+                resolve(history);
+            }).catch(reason => { reject(reason); });
+        });
+    }
+
+
     updateLease(leaseId: string, update: UpdateLeaseRequest): Promise<void> {
         return new Promise<void>((resolve, reject) => {
             this._fetch.put(`/api/v1/leases/${leaseId}`, update).then(() => {
diff --git a/Engine/Source/Programs/Horde/HordeDashboard/src/components/AuditLog.tsx b/Engine/Source/Programs/Horde/HordeDashboard/src/components/AuditLog.tsx
--- a/Engine/Source/Programs/Horde/HordeDashboard/src/components/AuditLog.tsx
+++ b/Engine/Source/Programs/Horde/HordeDashboard/src/components/AuditLog.tsx
@@ -65,6 +65,7 @@ class AuditLogHandler {
 
       this.agentId = undefined;
       this.issueId = undefined;
+      this.deviceId = undefined;
       this.templateAudit = undefined;
 
       this.timeSelectKey = "time_live_tail";
@@ -116,6 +117,21 @@ class AuditLogHandler {
       this.update();
    }
 
+   setDevice(deviceId?: string) {
+
+      if (!deviceId) {
+         this.clear();
+         return;
+      }
+
+      if (this.deviceId === deviceId) {
+         return;
+      }
+
+      this.deviceId = deviceId;
+      this.update();
+   }
+
    setTemplateAudit(templateAudit?: TemplateAudit) {
 
       if (!templateAudit) {
@@ -131,7 +147,6 @@ class AuditLogHandler {
       this.update();
    }
 
-
    get isAgentLog(): boolean {
       return !!this.agentId;
    }
@@ -140,6 +155,10 @@ class AuditLogHandler {
       return !!this.issueId;
    }
 
+   get isDeviceLog(): boolean {
+      return !!this.deviceId;
+   }
+
    get isTemplateAuditLog(): boolean {
       return !!this.templateAudit;
    }
@@ -163,7 +182,7 @@ class AuditLogHandler {
          this.canceled.add(i);
       }
 
-      if (!this.agentId && !this.issueId && !this.templateAudit) {
+      if (!this.agentId && !this.issueId && !this.templateAudit && !this.deviceId) {
          return;
       }
 
@@ -185,6 +204,8 @@ class AuditLogHandler {
                // live tail of the last 4 days
                // We need to optimize the agent history endpoint, this was an attempt to limit results though doesn't work 
                minTime = undefined; //new Date(new Date().valueOf() - (60 * 24 * 4 * 60000));
+            } else if (this.deviceId) {
+               minTime = undefined;
             } else if (this.templateAudit) {
                minTime = undefined;
             }
@@ -198,6 +219,10 @@ class AuditLogHandler {
             entries = await backend.getAgentHistory(this.agentId, { minTime: minTime?.toISOString(), maxTime: maxTime?.toISOString(), count: 8192 * 2 });
          }
 
+         if (this.deviceId) {
+            entries = await backend.getDeviceHistory(this.deviceId, { minTime: minTime?.toISOString(), maxTime: maxTime?.toISOString(), count: 8192 * 2 });
+         }
+
          if (this.issueId) {
             entries = await backend.getIssueHistory(this.issueId, { minTime: minTime?.toISOString(), maxTime: maxTime?.toISOString(), count: 8192 * 2 });
          }
@@ -272,6 +297,7 @@ class AuditLogHandler {
 
    agentId?: string;
    issueId?: string;
+   deviceId?: string;
    templateAudit?: TemplateAudit;
 
    timeSelectKey?: string;
@@ -343,7 +369,7 @@ const timeSelections: TimeSelection[] = [
 
 const selection = new Selection({ selectionMode: SelectionMode.multiple });
 
-export const AuditLogPanel: React.FC<{ agentId?: string, issueId?: string, templateAudit?: TemplateAudit }> = observer(({ agentId, issueId, templateAudit }) => {
+export const AuditLogPanel: React.FC<{ agentId?: string, deviceId?: string, issueId?: string, templateAudit?: TemplateAudit }> = observer(({ agentId, deviceId, issueId, templateAudit }) => {
 
    const windowSize = useWindowSize();
    const [showDatePicker, setShowDatePicker] = useState(false);
@@ -371,6 +397,10 @@ export const AuditLogPanel: React.FC<{ agentId?: string, issueId?: string, templ
       handler.setAgent(agentId);
    }
 
+   if (deviceId && handler.deviceId !== deviceId) {
+      handler.setDevice(deviceId);
+   }
+
    if (issueId && handler.issueId !== issueId) {
       handler.setIssue(issueId);
    }
@@ -379,7 +409,7 @@ export const AuditLogPanel: React.FC<{ agentId?: string, issueId?: string, templ
       handler.setTemplateAudit(templateAudit);
    }
 
-   if (!agentId && !issueId && !templateAudit) {
+   if (!agentId && !deviceId && !issueId && !templateAudit) {
       return null;
    }
 
@@ -461,9 +491,22 @@ export const AuditLogPanel: React.FC<{ agentId?: string, issueId?: string, templ
             link: `/agents?agentId=${encodeURIComponent(agentId)}`
          }
       ];
+   }
 
+   if (deviceId) {
+      crumbItems = [
+         {
+            text: `Devices`,
+            link: '/devices'
+         },
+         {
+            text: `Device Audit - ${deviceId.toUpperCase()}`,
+            link: `/devices?deviceId=${encodeURIComponent(deviceId)}`
+         }
+      ];
    }
 
+
    if (issueId) {
       crumbItems = [
          {
@@ -530,7 +573,7 @@ export const AuditLogPanel: React.FC<{ agentId?: string, issueId?: string, templ
 
    return <Stack>
       <Breadcrumbs items={crumbItems} />
-      <Stack tokens={{ childrenGap: 12 }} style={{backgroundColor: modeColors.background}}>
+      <Stack tokens={{ childrenGap: 12 }} style={{ backgroundColor: modeColors.background }}>
          <Stack horizontal>
             <div key={`windowsize_logview1_${windowSize.width}_${windowSize.height}`} style={{ width: vw / 2 - (1440 / 2) - 990, flexShrink: 0, backgroundColor: modeColors.background }} />
             <Stack tokens={{ childrenGap: 0 }} styles={{ root: { backgroundColor: modeColors.background, margin: "auto", paddingTop: 12, paddingRight: 10 } }}>
@@ -622,7 +665,6 @@ export const AuditLogPanel: React.FC<{ agentId?: string, issueId?: string, templ
                         </Stack>
                      </Stack>
                   </Stack>
-
                </Stack>
             </Stack>
          </Stack>
@@ -657,7 +699,6 @@ export const AuditLogPanel: React.FC<{ agentId?: string, issueId?: string, templ
                </FocusZone>
             </Stack>
          </Stack>
-
       </Stack>
    </Stack>
 
@@ -667,7 +708,7 @@ export const AuditLogView: React.FC = () => {
 
    const { hordeClasses } = getHordeStyling();
 
-   const { agentId, issueId, streamId, templateId } = useParams<{ agentId: string, issueId: string, streamId: string, templateId: string }>();
+   const { agentId, deviceId, issueId, streamId, templateId } = useParams<{ agentId: string, deviceId: string, issueId: string, streamId: string, templateId: string }>();
 
    let templateAudit: TemplateAudit | undefined;
    if (!!streamId && !!templateId) {
@@ -681,6 +722,7 @@ export const AuditLogView: React.FC = () => {
       <TopNav />
       <Stack>
          {!!agentId && <AuditLogPanel agentId={agentId} />}
+         {!!deviceId && <AuditLogPanel deviceId={deviceId} />}
          {!!issueId && <AuditLogPanel issueId={issueId} />}
          {!!templateAudit && <AuditLogPanel templateAudit={templateAudit} />}
       </Stack>
@@ -696,6 +738,7 @@ export const AuditLogView: React.FC = () => {
 enum AuditLogType {
    Unknown,
    Agent,
+   Device,
    Issue,
    Template
 }
@@ -731,6 +774,10 @@ class AuditLine {
          instance = new TemplateAuditLine(logType, entry);
       }
 
+      if (logType === AuditLogType.Device) {
+         instance = new DeviceAuditLine(logType, entry);
+      }
+
       if (instance) {
          instance.search = search;
          instance.process();
@@ -969,6 +1016,10 @@ class TemplateAuditLine extends AuditLine {
 
 }
 
+class DeviceAuditLine extends AuditLine {
+
+}
+
 // Line Rendering -------------------------------------------------------------------------------------------------------------------
 
 
@@ -988,6 +1039,10 @@ const renderAuditEntry = (entry: AuditLogEntry, search?: string) => {
       type = AuditLogType.Template;
    }
 
+   if (handler.isDeviceLog) {
+      type = AuditLogType.Device;
+   }
+
    if (type === AuditLogType.Unknown) {
       return <div>Unknown Log Type</div>
    }
diff --git a/Engine/Source/Programs/Horde/HordeDashboard/src/components/DeviceInfoView.tsx b/Engine/Source/Programs/Horde/HordeDashboard/src/components/DeviceInfoView.tsx
--- a/Engine/Source/Programs/Horde/HordeDashboard/src/components/DeviceInfoView.tsx
+++ b/Engine/Source/Programs/Horde/HordeDashboard/src/components/DeviceInfoView.tsx
@@ -2,7 +2,7 @@ import { DefaultButton, DetailsList, DetailsRow, IColumn, IconButton, IDetailsLi
 import { observer } from "mobx-react-lite";
 import moment from "moment";
 import { useState } from "react";
-import { Link } from "react-router-dom";
+import { Link, useNavigate } from "react-router-dom";
 import backend from "../backend";
 import { DeviceStatus, DeviceTelemetryQuery, GetDeviceResponse, GetDeviceTelemetryResponse, GetDeviceUtilizationResponse, GetTelemetryInfoResponse, JobData } from "../backend/Api";
 import dashboard, { StatusColor } from "../backend/Dashboard";
@@ -43,6 +43,7 @@ export const DeviceInfoModal: React.FC<{ handler: DeviceHandler, deviceIn?: GetD
    }
 
    const [jobState, setJobState] = useState<{ items: JobItem[], telemetryItems: JobTelemetryItem[], querying: boolean, queried: boolean, modifiedByUser?: string }>({ items: [], telemetryItems: [], querying: true, queried: false });
+   const navigate = useNavigate();
 
    if (!deviceIn) {
       return null;
@@ -249,7 +250,9 @@ export const DeviceInfoModal: React.FC<{ handler: DeviceHandler, deviceIn?: GetD
       return <Stack />;
    }
 
-   let failureRate = jobState.telemetryItems.length == 0? 0 : Math.floor(jobState.telemetryItems.filter(item => !!item.job.problemTimeUtc).length / jobState.telemetryItems.length * 100);
+   let failureRate = jobState.telemetryItems.length == 0 ? 0 : Math.floor(jobState.telemetryItems.filter(item => !!item.job.problemTimeUtc).length / jobState.telemetryItems.length * 100);
+   
+
 
    return (<Modal isOpen={true} styles={{ main: { padding: 8, width: 1140, height: '800px', backgroundColor: '#FFFFFF' } }} className={hordeClasses.modal} onDismiss={() => { onClose() }}>
       <Stack styles={{ root: { paddingTop: 8, paddingLeft: 24, paddingRight: 12, paddingBottom: 8 } }}>
@@ -293,7 +296,7 @@ export const DeviceInfoModal: React.FC<{ handler: DeviceHandler, deviceIn?: GetD
                            <PrimaryButton text="Edit" onClick={() => { onEdit(deviceIn) }} />
                         </Stack>
                         <Stack style={{ paddingTop: 8, paddingRight: 24 }}>
-                           <DefaultButton text="Close" onClick={() => { onClose() }} />
+                           <DefaultButton text="Audit" onClick={() => { navigate(`/audit/device/${deviceIn.id}`) }} />
                         </Stack>
                      </Stack>
                   </Stack>
