Horde: Fix not being able to set a device to base model if another model is set

diff --git a/Engine/Source/Programs/Horde/HordeDashboard/src/components/DeviceEditor.tsx b/Engine/Source/Programs/Horde/HordeDashboard/src/components/DeviceEditor.tsx
--- a/Engine/Source/Programs/Horde/HordeDashboard/src/components/DeviceEditor.tsx
+++ b/Engine/Source/Programs/Horde/HordeDashboard/src/components/DeviceEditor.tsx
@@ -387,7 +387,7 @@ export const DeviceEditor: React.FC<{ handler: DeviceHandler, deviceIn?: GetDevi
                name: device.name.trim(),
                address: device.address,
                poolId: device.poolId,
-               modelId: device.modelId,
+               modelId: device.modelId ?? "Base",
                notes: device.notes,
                problem: false, // always clear automatically generated problem state when saving
                maintenance: device.maintenance,
diff --git a/Engine/Source/Programs/Horde/Plugins/Build/HordeServer.Build/Devices/DevicesController.cs b/Engine/Source/Programs/Horde/Plugins/Build/HordeServer.Build/Devices/DevicesController.cs
--- a/Engine/Source/Programs/Horde/Plugins/Build/HordeServer.Build/Devices/DevicesController.cs
+++ b/Engine/Source/Programs/Horde/Plugins/Build/HordeServer.Build/Devices/DevicesController.cs
@@ -240,7 +240,7 @@ namespace HordeServer.Devices
 			{
 				modelIdValue = new string(update.ModelId);
 
-				if (platform.Models?.FirstOrDefault(x => x == modelIdValue) == null)
+				if (modelIdValue != "Base" && platform.Models?.FirstOrDefault(x => x == modelIdValue) == null)
 				{
 					return BadRequest($"Bad model id {update.ModelId} for platform {platform.Id} on request");
 				}
