Horde: Updating schema documentation

#rnx

diff --git a/Engine/Source/Programs/Horde/Docs/Config/Schema/Telemetry.md b/Engine/Source/Programs/Horde/Docs/Config/Schema/Telemetry.md
--- a/Engine/Source/Programs/Horde/Docs/Config/Schema/Telemetry.md
+++ b/Engine/Source/Programs/Horde/Docs/Config/Schema/Telemetry.md
@@ -66,6 +66,8 @@ Name | Description
 `groupBy` | `string`<br>Property to group by. Specified as a comma-separated list of JSON path expressions.
 `function` | [AggregationFunction](#aggregationfunction-enum)<br>How to aggregate samples for this metric
 `percentile` | `integer`<br>For the percentile function, specifies the percentile to measure
+`topN` | `integer`<br>
+`bottomN` | `integer`<br>
 `interval` | `string`<br>Interval for each metric. Supports times such as "2d", "1h", "1h30m", "20s".
 
 ## AggregationFunction (Enum)
