Horde; Disable topn/bottomn until can get it working in dev environment
#rnx

diff --git a/Engine/Source/Programs/Horde/Plugins/Analytics/HordeServer.Analytics/Telemetry/Metrics/MetricCollection.cs b/Engine/Source/Programs/Horde/Plugins/Analytics/HordeServer.Analytics/Telemetry/Metrics/MetricCollection.cs
--- a/Engine/Source/Programs/Horde/Plugins/Analytics/HordeServer.Analytics/Telemetry/Metrics/MetricCollection.cs
+++ b/Engine/Source/Programs/Horde/Plugins/Analytics/HordeServer.Analytics/Telemetry/Metrics/MetricCollection.cs
@@ -279,6 +279,7 @@ namespace HordeServer.Telemetry.Metrics
 			}
 
 			// Constrain to TopN or BottomN
+			/*
 			HashSet<ObjectId> removeIds = new HashSet<ObjectId>();			
 			if (tbConstraint.Count > 0)
 			{
@@ -319,6 +320,7 @@ namespace HordeServer.Telemetry.Metrics
 				filter &= Builders<MetricDocument>.Filter.In(x => x.Id, removeIds);
 				await _metrics.DeleteManyAsync(filter, cancellationToken);
 			}
+			*/
 		}
 
 		async Task CombineValuesAsync(TelemetryStoreId telemetryStoreId, MetricConfig metricConfig, string group, DateTime time, List<double> values, CancellationToken cancellationToken)
