Horde Analytics Dashboard Tweaks

diff --git a/Engine/Source/Programs/Horde/HordeServer/Defaults/default-metrics.telemetry.json b/Engine/Source/Programs/Horde/HordeServer/Defaults/default-metrics.telemetry.json
--- a/Engine/Source/Programs/Horde/HordeServer/Defaults/default-metrics.telemetry.json
+++ b/Engine/Source/Programs/Horde/HordeServer/Defaults/default-metrics.telemetry.json
@@ -54,6 +54,14 @@
 			"function": "Average",
 			"interval": "03:00:00"
 		},
+		{
+			"id": "core-loading-user-timetoedit",
+			"filter": "$[?(@.Payload.EventName == 'Core.Loading' && @.Payload.LoadingName=='TimeToEdit' && @.Payload.Config_IsBuildMachine==false && @.Payload.Config_IsRunningCommandlet==false && @.Payload.LoadingSeconds>5 && @.Payload.LoadingSeconds<3600)]",
+			"property": "$.Payload.LoadingSeconds",
+			"groupBy": "$.Payload.Build_BranchName,$.Payload.ProjectName,$.Payload.MapName",
+			"function": "Average",
+			"interval": "03:00:00"
+		},
 		{
 			"id": "core-loading-user-timetopie",
 			"filter": "$[?(@.Payload.EventName == 'Core.Loading' && @.Payload.LoadingName=='TimeToPIE' && @.Payload.Config_IsBuildMachine==false && @.Payload.Config_IsRunningCommandlet==false && @.Payload.LoadingSeconds>5 && @.Payload.LoadingSeconds<3600)]",
@@ -102,6 +110,14 @@
 			"function": "Average",
 			"interval": "03:00:00"
 		},
+		{
+			"id": "core-loading-user-timetoedit-kpi",
+			"filter": "$[?(@.Payload.EventName == 'Core.Loading' && @.Payload.LoadingName=='TimeToEdit' && @.Payload.Config_IsBuildMachine==false && @.Payload.Config_IsRunningCommandlet==false && @.Payload.LoadingSeconds>5 && @.Payload.LoadingSeconds<3600)]",
+			"property": "$.Payload.LoadingSeconds",
+			"groupBy": "$.Payload.Build_BranchName,$.Payload.ProjectName",
+			"function": "Average",
+			"interval": "03:00:00"
+		},
 		{
 			"id": "core-loading-user-timetopie-kpi",
 			"filter": "$[?(@.Payload.EventName == 'Core.Loading' && @.Payload.LoadingName=='TimeToPIE' && @.Payload.Config_IsBuildMachine==false && @.Payload.Config_IsRunningCommandlet==false && @.Payload.LoadingSeconds>5 && @.Payload.LoadingSeconds<3600)]",
@@ -150,6 +166,15 @@
 			"function": "Average",
 			"interval": "03:00:00"
 		},
+		{
+			"id": "core-loading-user-timetoedit-percentile",
+			"filter": "$[?(@.Payload.EventName == 'Core.Loading' && @.Payload.LoadingName=='TimeToEdit' && @.Payload.Config_IsBuildMachine==false && @.Payload.Config_IsRunningCommandlet==false && @.Payload.LoadingSeconds>5 && @.Payload.LoadingSeconds<3600)]",
+			"groupBy": "$.Payload.Build_BranchName,$.Payload.ProjectName,$.Payload.MapName",
+			"property": "$.Payload.LoadingSeconds",
+			"function": "Percentile",
+			"percentile": 75,
+			"interval": "03:00:00"
+		},
 		{
 			"id": "core-loading-user-timetopie-percentile",
 			"filter": "$[?(@.Payload.EventName == 'Core.Loading' && @.Payload.LoadingName=='TimeToPIE' && @.Payload.Config_IsBuildMachine==false && @.Payload.Config_IsRunningCommandlet==false && @.Payload.LoadingSeconds>5 && @.Payload.LoadingSeconds<3600)]",
diff --git a/Engine/Source/Programs/Horde/HordeServer/Defaults/default-views.telemetry.json b/Engine/Source/Programs/Horde/HordeServer/Defaults/default-views.telemetry.json
--- a/Engine/Source/Programs/Horde/HordeServer/Defaults/default-views.telemetry.json
+++ b/Engine/Source/Programs/Horde/HordeServer/Defaults/default-views.telemetry.json
@@ -91,12 +91,17 @@
 						{
 							"name": "User Load Time (75th Percentile)",
 							"metrics": [
+								{
+									"id": "core-loading-user-timetoedit-percentile",
+									"alias": "TimeToEdit",
+									"threshold": 120
+								},
 								{
 									"id": "core-loading-user-timetopie-percentile",
 									"alias": "TimeToPIE",
 									"threshold": 360
-								},
-								{
+								}
+								/*,{
 									"id": "core-loading-user-editorboottime-percentile",
 									"alias": "BootEditor",
 									"threshold": 120
@@ -110,7 +115,7 @@
 									"id": "core-loading-user-piestartuptime-percentile",
 									"alias": "PIEStartup",
 									"threshold": 360
-								}
+								}*/
 							],
 							"graph": "indicator",
 							"max": 720
@@ -119,6 +124,16 @@
 							"name": "User Load Time (Average)",
 							"metrics": [
 								{
+									"id": "core-loading-user-timetoedit",
+									"alias": "TimeToPIE",
+									"threshold": 120
+								},
+								{
+									"id": "core-loading-user-timetopie",
+									"alias": "TimeToPIE",
+									"threshold": 360
+								}
+								/*,{
 									"id": "core-loading-user-editorboottime",
 									"alias": "BootEditor",
 									"threshold": 120
@@ -132,12 +147,7 @@
 									"id": "core-loading-user-piestartuptime",
 									"alias": "PIEStartup",
 									"threshold": 360
-								},
-								{
-									"id": "core-loading-user-timetopie",
-									"alias": "TimeToPIE",
-									"threshold": 360
-								}
+								}*/
 							],
 							"graph": "indicator",
 							"max": 900
@@ -145,6 +155,15 @@
 						{
 							"name": "User Loading Time (KPIs)",
 							"metrics": [
+								{
+									"id": "core-loading-user-timetoedit-kpi",
+									"alias": "TimeToEdit"
+								},
+								{
+									"id": "core-loading-user-timetopie-kpi",
+									"alias": "TimeToPIE"
+								}
+								/*,
 								{
 									"id": "core-loading-user-editorboottime-kpi",
 									"alias": "BootEditor"
@@ -156,11 +175,7 @@
 								{
 									"id": "core-loading-user-piestartuptime-kpi",
 									"alias": "PIEStartup"
-								},
-								{
-									"id": "core-loading-user-timetopie-kpi",
-									"alias": "TimeToPIE"
-								}
+								}*/
 							]
 						},
 						{
