Horde: Update release notes
#rnx
#skipci

diff --git a/Engine/Source/Programs/Horde/Docs/ReleaseNotes.md b/Engine/Source/Programs/Horde/Docs/ReleaseNotes.md
--- a/Engine/Source/Programs/Horde/Docs/ReleaseNotes.md
+++ b/Engine/Source/Programs/Horde/Docs/ReleaseNotes.md
@@ -1,5 +1,25 @@
 # Release Notes
 
+## 2024-11-19
+
+* Rename default GetSparseClassData accessor to GetMutableSparseClassData to make mutability more obvious and to encourage use of more explicit accessor with EGetSparseClassDataMethod parameter - deprecating old version (38208446)
+* [UBA]
+  * Changed caching code to use new uba path which shares logic with vfs
+  * Removed last parts of OutputStatsThresholdMs
+  * Removed code wrapping root path instance now when we use handles instead (38199398)
+* Fix not being able to set a device to base model if another model is set (38193884)
+* Add support for device audit logs, frontend (38192783)
+* Improve device audit messages to include the user changing properties (38192679)
+* [UBT]
+  * Removed system roots and autosdk root now when all modules should provide the roots they need (38192101)
+* Add device audit logging, server side (38189896)
+* Decrease test time for AES transport test.  Removes the largest buffer size to reduce test time. (38186094)
+* Changed the Context and LoadingNames in EditorTelemetry to be consistent with the Span Names (38183893)
+* Add TopN/BottomN sampling constraints for Horde Analytics (38181517)
+* Add download with Unreal Toolbox preference and generalized artifact button (38181432)
+* Add CLI argument for loading additional config files.  Allows specifying custom config file paths, enabling more flexible configuration in buildfarm environments instead of relying solely on built-in default locations. (38180737)
+
+
 ## 2024-11-15
 
 * Prevent exceeding open file handle limit on MacOS when uploading data. (38125783)
