Horde: Fix only types declared in the main application assembly being considered for JSON schema generation.

#rnx

diff --git a/Engine/Source/Programs/Horde/HordeServer/Server/SchemaController.cs b/Engine/Source/Programs/Horde/HordeServer/Server/SchemaController.cs
--- a/Engine/Source/Programs/Horde/HordeServer/Server/SchemaController.cs
+++ b/Engine/Source/Programs/Horde/HordeServer/Server/SchemaController.cs
@@ -52,11 +52,14 @@ namespace HordeServer.Server
 		static Type[] FindSchemaTypes()
 		{
 			List<Type> schemaTypes = new List<Type>();
-			foreach (Type type in Assembly.GetExecutingAssembly().GetTypes())
+			foreach (Assembly assembly in AppDomain.CurrentDomain.GetAssemblies())
 			{
-				if (type.GetCustomAttribute<JsonSchemaAttribute>() != null)
+				foreach (Type type in assembly.GetTypes())
 				{
-					schemaTypes.Add(type);
+					if (type.GetCustomAttribute<JsonSchemaAttribute>() != null)
+					{
+						schemaTypes.Add(type);
+					}
 				}
 			}
 			return schemaTypes.ToArray();
