PR #12533: Fix Horde configPath for globals.json with perforce:// uri not working

diff --git a/Engine/Source/Programs/Horde/HordeServer/Configuration/ConfigService.cs b/Engine/Source/Programs/Horde/HordeServer/Configuration/ConfigService.cs
--- a/Engine/Source/Programs/Horde/HordeServer/Configuration/ConfigService.cs
+++ b/Engine/Source/Programs/Horde/HordeServer/Configuration/ConfigService.cs
@@ -517,16 +517,15 @@ namespace HordeServer.Configuration
 		internal static Uri GetGlobalConfigUri(string configPath, string defaultConfigDir)
 		{
 			bool isPerforcePath = configPath.StartsWith("//", StringComparison.Ordinal);
-			bool isAbsPath = Path.IsPathRooted(configPath);
-
 			if (isPerforcePath)
 			{
 				return ConfigNode.CombinePaths(new Uri(FileReference.Combine(new DirectoryReference(defaultConfigDir), "_").FullName), configPath);
 			}
 
-			if (isAbsPath)
+			Uri? uri;
+			if (Uri.TryCreate(configPath, UriKind.Absolute, out uri))
 			{
-				return new UriBuilder("file", String.Empty) { Path = configPath }.Uri;
+				return uri;
 			}
 
 			// Path is relative
