Horde: Add support for IMDSv2 for AWS EC2 in agent

Used for detecting spot interruptions and auto-scaling changes on a running agent.

diff --git a/Engine/Source/Programs/Horde/HordeAgent.Tests/Services/AwsInstanceLifecycleServiceTests.cs b/Engine/Source/Programs/Horde/HordeAgent.Tests/Services/AwsInstanceLifecycleServiceTests.cs
--- a/Engine/Source/Programs/Horde/HordeAgent.Tests/Services/AwsInstanceLifecycleServiceTests.cs
+++ b/Engine/Source/Programs/Horde/HordeAgent.Tests/Services/AwsInstanceLifecycleServiceTests.cs
@@ -3,6 +3,7 @@
 using System;
 using System.Collections.Generic;
 using System.IO;
+using System.Linq;
 using System.Net;
 using System.Net.Http;
 using System.Threading;
@@ -13,8 +14,6 @@ using HordeCommon.Rpc.Messages;
 using Microsoft.Extensions.Logging;
 using Microsoft.Extensions.Options;
 using Microsoft.VisualStudio.TestTools.UnitTesting;
-using Moq;
-using Moq.Protected;
 
 namespace HordeAgent.Tests.Services;
 
@@ -37,58 +36,82 @@ public class WorkerServiceStub : IWorkerService
 }
 
 /// <summary>
-/// Fake implementation of the AWS EC2 metadata server (IMDS)
+/// Fake implementation of the AWS EC2 metadata server (IMDSv2)
 /// </summary>
-public class FakeAwsImds
+public class FakeAwsImds : HttpMessageHandler
 {
-	public const string BaseUrl = "http://169.254.169.254/latest/meta-data";
-	public const string SpotInstanceData = "{\"action\": \"terminate\", \"time\": \"2017-09-18T08:22:00Z\"}";
-	public const string OnDemand = "on-demand";
-	public const string Spot = "spot";
-
-	private readonly Mock<HttpMessageHandler> _mock;
-	public string InstanceLifeCycle { get; set; } = OnDemand;
-	public string? TargetLifecycleState { get; set; }
-	public string? SpotInstanceAction { get; set; }
-
-	public FakeAwsImds()
+	public bool IsSpotInstance { get; set; }
+	public bool IsSpotTerminating { get; set; }
+	public bool IsAsgTerminating { get; set; }
+	private const string FakeTokenPrefix = "fakeToken-";
+	private readonly IClock _clock;
+	private readonly Dictionary<string, Func<HttpRequestMessage, HttpResponseMessage>> _responses = new();
+
+	public FakeAwsImds(IClock clock)
 	{
-		_mock = new(MockBehavior.Strict);
-		_mock
-			.Protected()
-			.Setup<Task<HttpResponseMessage>>(
-				"SendAsync",
-				ItExpr.IsAny<HttpRequestMessage>(),
-				ItExpr.IsAny<CancellationToken>()
-			)
-			.Returns(SendAsync);
-
-		// Ignore Dispose() calls
-		_mock.Protected().Setup("Dispose", ItExpr.IsAny<bool>());
-	}
-
-	public HttpClient GetHttpClient()
-	{
-		return new(_mock.Object);
+		_clock = clock;
+		_responses[AwsInstanceLifecycleService.ImdsTokenPath] = req =>
+		{
+			if (req.Method != HttpMethod.Put)
+			{
+				return Res(HttpStatusCode.MethodNotAllowed);
+			}
+
+			if (!req.Headers.TryGetValues(AwsInstanceLifecycleService.ImdsTokenTtlHeader, out IEnumerable<string>? values))
+			{
+				return Res(HttpStatusCode.BadRequest);
+			}
+			
+			string? tokenTtlSecs = values.FirstOrDefault();
+			TimeSpan tokenTtl = TimeSpan.FromSeconds(Convert.ToInt32(tokenTtlSecs));
+			long expireTime = new DateTimeOffset(_clock.UtcNow).Add(tokenTtl).ToUnixTimeSeconds();
+			return Res(HttpStatusCode.OK, $"{FakeTokenPrefix}{expireTime}");
+		};
+		
+		_responses[AwsInstanceLifecycleService.ImdsInstanceLifeCyclePath] = req => !IsValidImdsV2Request(req)
+			? Res(HttpStatusCode.Unauthorized)
+			: Res(HttpStatusCode.OK, IsSpotInstance ? "spot" : "on-demand");
+		
+		_responses[AwsInstanceLifecycleService.ImdsSpotActionPath] = req =>
+		{
+			if (!IsValidImdsV2Request(req))
+			{
+				return new HttpResponseMessage(HttpStatusCode.Unauthorized);
+			}
+			
+			HttpStatusCode statusCode = IsSpotTerminating ? HttpStatusCode.OK : HttpStatusCode.NotFound;
+			return Res(statusCode, IsSpotTerminating ? "{\"action\": \"terminate\", \"time\": \"2024-11-01T12:00:00Z\"}" : "");
+		};
+		
+		_responses[AwsInstanceLifecycleService.ImdsAsgTargetStatePath] = req => !IsValidImdsV2Request(req)
+			? new HttpResponseMessage(HttpStatusCode.Unauthorized)
+			: Res(HttpStatusCode.OK, IsAsgTerminating ? "Terminated" : "InService");
 	}
-
-	private Task<HttpResponseMessage> SendAsync(HttpRequestMessage request, CancellationToken cancellationToken)
+	
+	public HttpClient CreateHttpClient() => new (this) { BaseAddress = new Uri("http://169.254.169.254") };
+	
+	private static HttpResponseMessage Res(HttpStatusCode statusCode, string body = "") => new (statusCode) { Content = new StringContent(body) };
+	
+	private bool IsValidImdsV2Request(HttpRequestMessage request)
 	{
-		switch (request.RequestUri!.ToString())
+		if (!request.Headers.TryGetValues(AwsInstanceLifecycleService.ImdsTokenHeader, out IEnumerable<string>? tokens))
 		{
-			case $"{BaseUrl}/": return Task.FromResult(CreateResponse(HttpStatusCode.OK, "(list of sub-paths)"));
-			case $"{BaseUrl}/autoscaling/target-lifecycle-state": return Task.FromResult(CreateResponse(HttpStatusCode.OK, TargetLifecycleState));
-			case $"{BaseUrl}/spot/instance-action": return Task.FromResult(CreateResponse(HttpStatusCode.OK, SpotInstanceAction));
-			case $"{BaseUrl}/instance-life-cycle": return Task.FromResult(CreateResponse(HttpStatusCode.OK, InstanceLifeCycle));
-			default: return Task.FromResult(CreateResponse(HttpStatusCode.InternalServerError, "Invalid state"));
+			return false;
 		}
+		string[] parts = tokens.First().Split(FakeTokenPrefix);
+		return parts.Length == 2 && Convert.ToInt64(parts[1]) > new DateTimeOffset(_clock.UtcNow).ToUnixTimeSeconds();
 	}
-
-	private static HttpResponseMessage CreateResponse(HttpStatusCode status, string? data)
+	
+	/// <inheritdoc />
+	protected override Task<HttpResponseMessage> SendAsync(HttpRequestMessage request, CancellationToken cancellationToken)
 	{
-		return data != null
-			? new HttpResponseMessage { StatusCode = status, Content = new StringContent(data) }
-			: new HttpResponseMessage { StatusCode = HttpStatusCode.NotFound, Content = new StringContent("Not found") };
+		string path = request.RequestUri!.PathAndQuery;
+		if (_responses.TryGetValue(path, out Func<HttpRequestMessage, HttpResponseMessage>? handler))
+		{
+			return Task.FromResult(handler(request));
+		}
+		
+		return Task.FromResult(Res(HttpStatusCode.NotFound, $"FakeAwsImds: path {path} not found"));
 	}
 }
 
@@ -107,7 +130,8 @@ public sealed class AwsInstanceLifecycleServiceTests : IDisposable
 {
 	private readonly WorkerServiceStub _workerService = new();
 	private readonly HttpClient _httpClient;
-	private readonly FakeAwsImds _fakeImds = new();
+	private readonly StubClock _clock = new();
+	private readonly FakeAwsImds _fakeImds;
 	private readonly AwsInstanceLifecycleService _service;
 	private readonly FileReference _terminationSignalFile;
 	private Ec2TerminationInfo? _info;
@@ -124,12 +148,14 @@ public sealed class AwsInstanceLifecycleServiceTests : IDisposable
 		AgentSettings settings = new() { WorkingDir = new DirectoryReference(tempDir) };
 		_terminationSignalFile = settings.GetTerminationSignalFile();
 
-		_httpClient = _fakeImds.GetHttpClient();
-		_service = new AwsInstanceLifecycleService(_workerService, _httpClient, new TestOptionsMonitor<AgentSettings>(settings), loggerFactory.CreateLogger<AwsInstanceLifecycleService>());
+		_fakeImds = new(_clock);
+		_httpClient = _fakeImds.CreateHttpClient();
+		_service = new AwsInstanceLifecycleService(_workerService, _httpClient, _clock, new TestOptionsMonitor<AgentSettings>(settings),
+			loggerFactory.CreateLogger<AwsInstanceLifecycleService>());
 		_service._timeToLiveAsg = TimeSpan.FromMilliseconds(10);
 		_service._timeToLiveSpot = TimeSpan.FromMilliseconds(20);
 		_service._terminationBufferTime = TimeSpan.FromMilliseconds(2);
-
+		
 		AwsInstanceLifecycleService.TerminationWarningDelegate origWarningCallback = _service._terminationWarningCallback;
 		_service._terminationWarningCallback = (info, ct) =>
 		{
@@ -143,49 +169,59 @@ public sealed class AwsInstanceLifecycleServiceTests : IDisposable
 			return Task.CompletedTask;
 		};
 	}
-
+	
 	[TestMethod]
 	[Timeout(5000)]
 	public async Task Terminate_Asg_CallbackHasCorrectParametersAsync()
 	{
-		_fakeImds.TargetLifecycleState = "Terminated";
+		_fakeImds.IsAsgTerminating = true;
 		await _service.MonitorInstanceLifecycleAsync(CancellationToken.None);
 		Assert.AreEqual(Ec2InstanceState.TerminatingAsg, _info!.State);
 		Assert.IsFalse(_info!.IsSpot);
 		Assert.AreEqual(8, _info!.TimeToLive.TotalMilliseconds); // 10 ms for ASG, minus 2 ms for termination buffer
 	}
-
+	
 	[TestMethod]
 	[Timeout(5000)]
 	public async Task Terminate_Spot_CallbackHasCorrectParametersAsync()
 	{
-		_fakeImds.SpotInstanceAction = FakeAwsImds.SpotInstanceData;
-		_fakeImds.InstanceLifeCycle = FakeAwsImds.Spot;
+		_fakeImds.IsSpotInstance = true;
+		_fakeImds.IsSpotTerminating = true;
 		await _service.MonitorInstanceLifecycleAsync(CancellationToken.None);
 		Assert.AreEqual(Ec2InstanceState.TerminatingSpot, _info!.State);
 		Assert.IsTrue(_info!.IsSpot);
 		Assert.AreEqual(18, _info!.TimeToLive.TotalMilliseconds); // 20 ms for spot, minus 2 ms for termination buffer
 	}
-
+	
 	[TestMethod]
 	[Timeout(5000)]
 	public async Task Terminate_Spot_WritesSignalFileAsync()
 	{
-		_fakeImds.SpotInstanceAction = FakeAwsImds.SpotInstanceData;
-		_fakeImds.InstanceLifeCycle = FakeAwsImds.Spot;
+		_fakeImds.IsSpotInstance = true;
+		_fakeImds.IsSpotTerminating = true;
 		Assert.IsFalse(File.Exists(_terminationSignalFile.FullName));
-
+		
 		Assert.IsFalse(_workerService.TerminateSession);
 		await _service.MonitorInstanceLifecycleAsync(CancellationToken.None);
 		Assert.IsTrue(_workerService.TerminateSession);
-
+		
 		string data = await File.ReadAllTextAsync(_terminationSignalFile.FullName);
 		Assert.AreEqual("v1\n18\n2222000\nAWS EC2 Spot interruption\n", data); // 20 ms for spot, minus 2 ms for termination buffer
 	}
-
+	
+	[TestMethod]
+	public async Task Imds_TokenExpired_RefreshesToken_Async()
+	{
+		_fakeImds.IsSpotInstance = true;
+		Assert.AreEqual(true, await _service.IsSpotInstanceAsync(CancellationToken.None));
+		_clock.Advance(TimeSpan.FromHours(12));
+		Assert.AreEqual(true, await _service.IsSpotInstanceAsync(CancellationToken.None));
+	}
+	
 	public void Dispose()
 	{
 		_httpClient.Dispose();
 		_service.Dispose();
+		_fakeImds.Dispose();
 	}
 }
\ No newline at end of file
diff --git a/Engine/Source/Programs/Horde/HordeAgent/Services/AwsInstanceLifecycleService.cs b/Engine/Source/Programs/Horde/HordeAgent/Services/AwsInstanceLifecycleService.cs
--- a/Engine/Source/Programs/Horde/HordeAgent/Services/AwsInstanceLifecycleService.cs
+++ b/Engine/Source/Programs/Horde/HordeAgent/Services/AwsInstanceLifecycleService.cs
@@ -38,22 +38,13 @@ public enum Ec2InstanceState
 	TerminatingAsg
 }
 
-internal class Ec2TerminationInfo
+internal class Ec2TerminationInfo(Ec2InstanceState state, bool isSpot, TimeSpan timeToLive, DateTime terminateAt, string reason)
 {
-	public Ec2InstanceState State { get; init; }
-	public bool IsSpot { get; init; }
-	public TimeSpan TimeToLive { get; init; }
-	public DateTime TerminateAt { get; init; }
-	public string Reason { get; init; }
-
-	public Ec2TerminationInfo(Ec2InstanceState state, bool isSpot, TimeSpan timeToLive, DateTime terminateAt, string reason)
-	{
-		State = state;
-		IsSpot = isSpot;
-		TimeToLive = timeToLive;
-		TerminateAt = terminateAt;
-		Reason = reason;
-	}
+	public Ec2InstanceState State { get; } = state;
+	public bool IsSpot { get; } = isSpot;
+	public TimeSpan TimeToLive { get; } = timeToLive;
+	public DateTime TerminateAt { get; } = terminateAt;
+	public string Reason { get; } = reason;
 }
 
 /// <summary>
@@ -66,12 +57,24 @@ class AwsInstanceLifecycleService : BackgroundService
 	/// </summary>
 	public const string HttpClientName = "Horde.HttpAwsInstanceClient";
 
-	private const string BaseUri = "http://169.254.169.254/latest/meta-data";
+	
+	public const string ImdsTokenPath = "/latest/api/token";
+	public const string ImdsInstanceLifeCyclePath = "/latest/meta-data/instance-life-cycle";
+	public const string ImdsSpotActionPath = "/latest/meta-data/spot/instance-action";
+	public const string ImdsAsgTargetStatePath = "/latest/meta-data/autoscaling/target-lifecycle-state";
+	public const string ImdsTokenHeader = "X-aws-ec2-metadata-token";
+	public const string ImdsTokenTtlHeader = "X-aws-ec2-metadata-token-ttl-seconds";
+	public const string ImdsTokenTtl = "21600"; // 6 hours
+	
+	private const string BaseUri = "http://169.254.169.254"; // TODO: remove prefix
 	private readonly IWorkerService _workerService;
 	private readonly HttpClient _httpClient;
+	private readonly IClock _clock;
 	private readonly FileReference _terminationSignalFile;
 	private readonly ILogger<AwsInstanceLifecycleService> _logger;
 	private readonly TimeSpan _pollInterval = TimeSpan.FromSeconds(5);
+	private string? _imdsToken;
+	private DateTime? _imdsTokenExpireTime = DateTime.UnixEpoch;
 
 	internal delegate Task TerminationWarningDelegate(Ec2TerminationInfo info, CancellationToken cancellationToken);
 	internal delegate Task TerminationDelegate(Ec2TerminationInfo info, CancellationToken cancellationToken);
@@ -99,62 +102,74 @@ class AwsInstanceLifecycleService : BackgroundService
 	/// <summary>
 	/// Constructor
 	/// </summary>
-	public AwsInstanceLifecycleService(IWorkerService workerService, HttpClient httpClient, IOptionsMonitor<AgentSettings> settings, ILogger<AwsInstanceLifecycleService> logger)
+	public AwsInstanceLifecycleService(IWorkerService workerService, HttpClient httpClient, IClock clock, IOptionsMonitor<AgentSettings> settings, ILogger<AwsInstanceLifecycleService> logger)
 	{
 		_workerService = workerService;
 		_httpClient = httpClient;
 		_httpClient.Timeout = TimeSpan.FromSeconds(10);
+		_clock = clock;
 		_terminationWarningCallback = OnTerminationWarningAsync;
 		_terminationCallback = OnTerminationAsync;
 		_logger = logger;
 		_terminationSignalFile = settings.CurrentValue.GetTerminationSignalFile();
 	}
-
+	
+	private async Task<HttpResponseMessage> GetMetadataAsync(string path, CancellationToken cancellationToken)
+	{
+		if (_imdsToken == null || _clock.UtcNow >= _imdsTokenExpireTime)
+		{
+			// Get a token
+			using HttpRequestMessage tokenReq = new (HttpMethod.Put, new Uri($"{BaseUri}{ImdsTokenPath}"));
+			tokenReq.Headers.Add(ImdsTokenTtlHeader, ImdsTokenTtl);
+			using HttpResponseMessage tokenRes = await _httpClient.SendAsync(tokenReq, cancellationToken);
+			tokenRes.EnsureSuccessStatusCode();
+			_imdsToken = await tokenRes.Content.ReadAsStringAsync(cancellationToken);
+			_imdsTokenExpireTime = _clock.UtcNow.AddSeconds(Convert.ToInt32(ImdsTokenTtl));
+		}
+		
+		// Perform actual metadata request
+		using HttpRequestMessage request = new (HttpMethod.Get, $"{BaseUri}{path}");
+		request.Headers.Add(ImdsTokenHeader, _imdsToken);
+		return await _httpClient.SendAsync(request, cancellationToken);
+	}
+	
+	internal async Task<bool> IsSpotInstanceAsync(CancellationToken cancellationToken)
+	{
+		HttpResponseMessage res = await GetMetadataAsync(ImdsInstanceLifeCyclePath, cancellationToken);
+		if (res.StatusCode == HttpStatusCode.OK)
+		{
+			return await res.Content.ReadAsStringAsync(cancellationToken) == "spot";
+		}
+		
+		throw new Exception($"Request to IMDS server failed. Status={res.StatusCode}");
+	}
+	
 	private async Task<Ec2InstanceState> GetStateAsync(CancellationToken cancellationToken)
 	{
 		try
 		{
-			HttpResponseMessage spotRes = await _httpClient.GetAsync(new Uri(BaseUri + "/spot/instance-action"), cancellationToken);
+			HttpResponseMessage spotRes = await GetMetadataAsync(ImdsSpotActionPath, cancellationToken);
 			if (spotRes.StatusCode == HttpStatusCode.OK)
 			{
 				return Ec2InstanceState.TerminatingSpot;
 			}
-			
-			HttpResponseMessage asgRes = await _httpClient.GetAsync(new Uri(BaseUri + "/autoscaling/target-lifecycle-state"), cancellationToken);
-			if (asgRes.StatusCode == HttpStatusCode.OK)
-			{
-				string state = await asgRes.Content.ReadAsStringAsync(cancellationToken);
-				if (state == "Terminated")
-				{
-					return Ec2InstanceState.TerminatingAsg;
-				}
-			}
-			
-			return Ec2InstanceState.InService;
+
+			HttpResponseMessage asgRes = await GetMetadataAsync(ImdsAsgTargetStatePath, cancellationToken);
+			string state = await asgRes.Content.ReadAsStringAsync(cancellationToken);
+			return state.Equals("Terminated", StringComparison.OrdinalIgnoreCase)
+				? Ec2InstanceState.TerminatingAsg
+				: Ec2InstanceState.InService;
 		}
 		catch (Exception e)
 		{
 			// HTTP timeouts from the IMDS server can occur, in addition to unknown exceptions.
 			// Unclear why timeouts are seen in the first place from the metadata server.
 			// It should not be under load at all but it's possible an agent performing heavy work can affect this.
-			
-			_logger.LogWarning("Request to IMDS server failed. Reason={Reason}", e.Message);
+			_logger.LogWarning(e, "Request to IMDS server failed. Reason={Reason}", e.Message);
 			return Ec2InstanceState.Unknown;
 		}
 	}
 
-	private async Task<bool> IsSpotInstanceAsync(CancellationToken cancellationToken)
-	{
-		HttpResponseMessage res = await _httpClient.GetAsync(new Uri(BaseUri + "/instance-life-cycle"), cancellationToken);
-		if (res.StatusCode == HttpStatusCode.OK)
-		{
-			string state = await res.Content.ReadAsStringAsync(cancellationToken);
-			return state == "spot";
-		}
-
-		return false;
-	}
-
 	internal async Task MonitorInstanceLifecycleAsync(CancellationToken cancellationToken)
 	{
 		_logger.LogInformation("Monitoring EC2 instance lifecycle state...");
@@ -249,5 +264,4 @@ class AwsInstanceLifecycleService : BackgroundService
 			// Ignore any exceptions if cancellation has been requested
 		}
 	}
-}
-
+}
\ No newline at end of file
