Horde: Fix missing IClock for dependency injection

This missing dependency caused agent startup failure when EC2 mode was enabled.

diff --git a/Engine/Source/Programs/Horde/HordeAgent/AgentApp.cs b/Engine/Source/Programs/Horde/HordeAgent/AgentApp.cs
--- a/Engine/Source/Programs/Horde/HordeAgent/AgentApp.cs
+++ b/Engine/Source/Programs/Horde/HordeAgent/AgentApp.cs
@@ -176,7 +176,8 @@ namespace HordeAgent
 			}
 
 			OpenTelemetryHelper.Configure(services, settings.OpenTelemetry);
-
+			
+			services.AddSingleton<IClock, DefaultClock>();
 			services.AddHorde(options =>
 			{
 				options.ServerUrl = serverProfile.Url;
diff --git a/Engine/Source/Programs/Horde/HordeAgent/Services/AwsInstanceLifecycleService.cs b/Engine/Source/Programs/Horde/HordeAgent/Services/AwsInstanceLifecycleService.cs
--- a/Engine/Source/Programs/Horde/HordeAgent/Services/AwsInstanceLifecycleService.cs
+++ b/Engine/Source/Programs/Horde/HordeAgent/Services/AwsInstanceLifecycleService.cs
@@ -198,6 +198,7 @@ class AwsInstanceLifecycleService : BackgroundService
 			catch (Exception e)
 			{
 				_logger.LogError(e, "Unhandled exception during EC2 instance monitoring");
+				await Task.Delay(TimeSpan.FromMinutes(1), cancellationToken);
 			}
 
 			await Task.Delay(_pollInterval, cancellationToken);
