Horde: Fix legacy method for specifying which changes to trigger schedules on not being respected.

#rnx

diff --git a/Engine/Source/Programs/Horde/Plugins/Build/HordeServer.Build/Jobs/Schedules/ScheduleConfig.cs b/Engine/Source/Programs/Horde/Plugins/Build/HordeServer.Build/Jobs/Schedules/ScheduleConfig.cs
--- a/Engine/Source/Programs/Horde/Plugins/Build/HordeServer.Build/Jobs/Schedules/ScheduleConfig.cs
+++ b/Engine/Source/Programs/Horde/Plugins/Build/HordeServer.Build/Jobs/Schedules/ScheduleConfig.cs
@@ -54,29 +54,6 @@ namespace HordeServer.Jobs.Schedules
 		[Obsolete("Use Commits instead")]
 		public List<ChangeContentFlags>? Filter { get; set; }
 
-		/// <summary>
-		/// Merges the list of commit tags with the legacy filter field
-		/// </summary>
-#pragma warning disable CS0618 // Type or member is obsolete
-		internal List<CommitTag> GetAllCommitTags()
-		{
-			List<CommitTag> commits = Commits;
-			if (Filter != null && Filter.Count > 0)
-			{
-				commits = new List<CommitTag>(commits);
-				if (Filter.Contains(ChangeContentFlags.ContainsCode))
-				{
-					commits.Add(CommitTag.Code);
-				}
-				if (Filter.Contains(ChangeContentFlags.ContainsContent))
-				{
-					commits.Add(CommitTag.Content);
-				}
-			}
-			return commits;
-		}
-#pragma warning restore CS0618 // Type or member is obsolete
-
 		/// <summary>
 		/// Files that should cause the job to trigger
 		/// </summary>
diff --git a/Engine/Source/Programs/Horde/Plugins/Build/HordeServer.Build/Streams/StreamConfig.cs b/Engine/Source/Programs/Horde/Plugins/Build/HordeServer.Build/Streams/StreamConfig.cs
--- a/Engine/Source/Programs/Horde/Plugins/Build/HordeServer.Build/Streams/StreamConfig.cs
+++ b/Engine/Source/Programs/Horde/Plugins/Build/HordeServer.Build/Streams/StreamConfig.cs
@@ -325,6 +325,22 @@ namespace HordeServer.Streams
 				ScheduleConfig? schedule = template.Schedule;
 				if (schedule != null)
 				{
+#pragma warning disable CS0618 // Type or member is obsolete
+					List<ChangeContentFlags>? filter = schedule.Filter;
+					if (filter != null && filter.Count > 0)
+					{
+						List<CommitTag> commits = schedule.Commits;
+						if (filter.Contains(ChangeContentFlags.ContainsCode) && !commits.Contains(CommitTag.Code))
+						{
+							commits.Add(CommitTag.Code);
+						}
+						if (filter.Contains(ChangeContentFlags.ContainsContent) && !commits.Contains(CommitTag.Content))
+						{
+							commits.Add(CommitTag.Content);
+						}
+					}
+#pragma warning restore CS0618 // Type or member is obsolete
+
 					foreach (CommitTag commitTag in schedule.Commits)
 					{
 						if (!TryGetCommitTag(commitTag, out _))
