Horde: Fix issue with failed batches not being able to retry steps

diff --git a/Engine/Source/Programs/Horde/Plugins/Build/HordeServer.Build.Tests/Jobs/JobCollectionTests.cs b/Engine/Source/Programs/Horde/Plugins/Build/HordeServer.Build.Tests/Jobs/JobCollectionTests.cs
--- a/Engine/Source/Programs/Horde/Plugins/Build/HordeServer.Build.Tests/Jobs/JobCollectionTests.cs
+++ b/Engine/Source/Programs/Horde/Plugins/Build/HordeServer.Build.Tests/Jobs/JobCollectionTests.cs
@@ -755,12 +755,12 @@ namespace HordeServer.Tests.Jobs
 			job = await RunBatchAsync(job, 0);
 			job = Deref(await job.TryUpdateBatchAsync(initialBatch.Id, null, JobStepBatchState.Complete, JobStepBatchError.Cancelled));
 
-			UpdateStepResponse response = ResultToValue(await JobsController.UpdateStepAsync(job.Id, initialBatch.Id, initialBatch.Steps[0].Id, new UpdateStepRequest { Retry = true }));
+			UpdateStepResponse response = ResultToValue(await JobsController.UpdateStepAsync(job.Id, initialBatch.Id, initialBatch.Steps[2].Id, new UpdateStepRequest { Retry = true }));
 
 			job = Deref(await JobCollection.GetAsync(job.Id));
 
-			// Expecting 2 batches here, and new steps and stuff, commented out for Ben to look at, as it fails
-			// Assert.AreEqual(2, job.Batches.Count);
+			Assert.AreEqual(2, job.Batches.Count);
+			Assert.AreEqual(3, job.Batches[1].Steps.Count);
 		}
 
 		[TestMethod]
diff --git a/Engine/Source/Programs/Horde/Plugins/Build/HordeServer.Build/Jobs/JobCollection.cs b/Engine/Source/Programs/Horde/Plugins/Build/HordeServer.Build/Jobs/JobCollection.cs
--- a/Engine/Source/Programs/Horde/Plugins/Build/HordeServer.Build/Jobs/JobCollection.cs
+++ b/Engine/Source/Programs/Horde/Plugins/Build/HordeServer.Build/Jobs/JobCollection.cs
@@ -1390,7 +1390,7 @@ namespace HordeServer.Jobs
 					NodeRef nodeRef = new NodeRef(batch.GroupIdx, step.NodeIdx);
 					if (retryNodes.Remove(nodeRef))
 					{
-						if (step.State == JobStepState.Skipped)
+						if (step.State == JobStepState.Skipped && batch.Error == JobStepBatchError.None)
 						{
 							// Add all the dependencies to be retried
 							NodeRef[] dependencies = graph.GetNode(nodeRef).InputDependencies;
