Horde Analytics Dashboard Tweaks

diff --git a/Engine/Source/Programs/Horde/HordeServer/Defaults/default-views.telemetry.json b/Engine/Source/Programs/Horde/HordeServer/Defaults/default-views.telemetry.json
--- a/Engine/Source/Programs/Horde/HordeServer/Defaults/default-views.telemetry.json
+++ b/Engine/Source/Programs/Horde/HordeServer/Defaults/default-views.telemetry.json
@@ -7,11 +7,18 @@
 			"variables": [
 				{
 					"name": "Stream",
-					"group": "Build_BranchName"
+					"group": "Build_BranchName",
+					"defaults": [
+						"++ue5+main"
+					]
 				},
 				{
 					"name": "Project",
-					"group": "ProjectName"
+					"group": "ProjectName",
+					"defaults": [
+						"Lyra",
+						"CitySample"
+					]
 				},
 				{
 					"name": "Map",
@@ -21,8 +28,7 @@
 					"name": "LoadingName",
 					"group": "LoadingName",
 					"defaults": [
-						"BootEditor",
-						"LoadMap",
+						"TimeToEdit",
 						"TimeToPIE"
 					]
 				},
