Added Macros to Horde Analytics metrics

diff --git a/Engine/Source/Programs/Horde/HordeServer/Defaults/default-metrics.telemetry.json b/Engine/Source/Programs/Horde/HordeServer/Defaults/default-metrics.telemetry.json
--- a/Engine/Source/Programs/Horde/HordeServer/Defaults/default-metrics.telemetry.json
+++ b/Engine/Source/Programs/Horde/HordeServer/Defaults/default-metrics.telemetry.json
@@ -1,86 +1,114 @@
 {
+	"macros": [
+		{
+			"name": "IsBuildMachine",
+			"value": "@.Payload.Config_IsBuildMachine==true"
+		},
+		{
+			"name": "IsUserMachine",
+			"value": "@.Payload.Config_IsBuildMachine==false && @.Payload.Config_IsRunningCommandlet==false"
+		},
+		{
+			"name": "IsLoadTimeInRange",
+			"value": "@.Payload.LoadingSeconds>5 && @.Payload.LoadingSeconds<3600"
+		},
+		{
+			"name": "IsUBAExecutor",
+			"value": "@.Payload.EventName == 'Core.UBT.Executor' && search(@.Payload.Executor, \"Unreal Build Accelerator\")"
+		}
+	],
 	"metrics": [
 		{
 			"id": "core-loading-user-loadingtime",
-			"filter": "$[?(@.Payload.EventName == 'Core.Loading' && @.Payload.Config_IsBuildMachine==false && @.Payload.Config_IsRunningCommandlet==false && @.Payload.LoadingSeconds>5 && @.Payload.LoadingSeconds<3600)]",
+			"filter": "$[?(@.Payload.EventName == 'Core.Loading' && $(IsUserMachine) && $(IsLoadTimeInRange) )]",
 			"property": "$.Payload.LoadingSeconds",
 			"groupBy": "$.Payload.Build_BranchName,$.Payload.ProjectName,$.Payload.LoadingName",
 			"function": "Average",
-			"interval": "03:00:00"
+			"interval": "03:00:00",
+			"topN": 10
 		},
 		{
 			"id": "core-loading-user-loadingcount",
-			"filter": "$[?(@.Payload.EventName == 'Core.Loading' &&@.Payload.Config_IsBuildMachine==false )]",
+			"filter": "$[?(@.Payload.EventName == 'Core.Loading' && $(IsUserMachine) )]",
 			"groupBy": "$.Payload.Build_BranchName,$.Payload.ProjectName,$.Payload.LoadingName",
 			"function": "Count",
-			"interval": "03:00:00"
+			"interval": "03:00:00",
+			"topN": 10
 		},
 		{
 			"id": "core-loading-user-editorboottime",
-			"filter": "$[?(@.Payload.EventName == 'Core.Loading' && @.Payload.LoadingName=='Editor.Boot' && @.Payload.Config_IsBuildMachine==false && @.Payload.Config_IsRunningCommandlet==false && @.Payload.LoadingSeconds>5 && @.Payload.LoadingSeconds<3600)]",
+			"filter": "$[?(@.Payload.EventName == 'Core.Loading' && @.Payload.LoadingName=='Editor.Boot' && $(IsUserMachine) && $(IsLoadTimeInRange))]",
 			"property": "$.Payload.LoadingSeconds",
 			"groupBy": "$.Payload.Build_BranchName,$.Payload.ProjectName,$.Payload.ProjectName",
 			"function": "Average",
-			"interval": "03:00:00"
+			"interval": "03:00:00",
+			"topN": 10
 		},
 		{
 			"id": "core-loading-user-editorstartuptime",
-			"filter": "$[?(@.Payload.EventName == 'Core.Loading' && @.Payload.LoadingName=='Editor.Initialize' && @.Payload.Config_IsBuildMachine==false && @.Payload.Config_IsRunningCommandlet==false && @.Payload.LoadingSeconds>5 && @.Payload.LoadingSeconds<3600)]",
+			"filter": "$[?(@.Payload.EventName == 'Core.Loading' && @.Payload.LoadingName=='Editor.Initialize' && $(IsUserMachine) && $(IsLoadTimeInRange))]",
 			"property": "$.Payload.LoadingSeconds",
 			"groupBy": "$.Payload.Build_BranchName,$.Payload.ProjectName,$.Payload.ProjectName",
 			"function": "Average",
-			"interval": "03:00:00"
+			"interval": "03:00:00",
+			"topN": 10
 		},
 		{
 			"id": "core-loading-user-loadmaptime",
-			"filter": "$[?(@.Payload.EventName == 'Core.Loading' && @.Payload.LoadingName=='Editor.LoadMap' && @.Payload.Config_IsBuildMachine==false && @.Payload.Config_IsRunningCommandlet==false && @.Payload.LoadingSeconds>5 && @.Payload.LoadingSeconds<3600)]",
+			"filter": "$[?(@.Payload.EventName == 'Core.Loading' && @.Payload.LoadingName=='Editor.LoadMap' && $(IsUserMachine) && $(IsLoadTimeInRange))]",
 			"property": "$.Payload.LoadingSeconds",
 			"groupBy": "$.Payload.Build_BranchName,$.Payload.ProjectName,$.Payload.MapName",
 			"function": "Average",
-			"interval": "03:00:00"
+			"interval": "03:00:00",
+			"topN": 10
 		},
 		{
 			"id": "core-loading-user-loadmap-count",
-			"filter": "$[?(@.Payload.EventName == 'Core.Loading' && @.Payload.LoadingName=='Editor.LoadMap' && @.Payload.Config_IsBuildMachine==false && @.Payload.Config_IsRunningCommandlet==false && @.Payload.LoadingSeconds>5 && @.Payload.LoadingSeconds<3600)]",
+			"filter": "$[?(@.Payload.EventName == 'Core.Loading' && @.Payload.LoadingName=='Editor.LoadMap' && $(IsUserMachine) && $(IsLoadTimeInRange))]",
 			"groupBy": "$.Payload.Build_BranchName,$.Payload.ProjectName,$.Payload.MapName",
 			"function": "Count",
-			"interval": "03:00:00"
+			"interval": "03:00:00",
+			"topN": 10
 		},
 		{
 			"id": "core-loading-user-pieloadmaptime",
-			"filter": "$[?(@.Payload.EventName == 'Core.Loading' && @.Payload.LoadingName=='PIE.LoadMap' && @.Payload.Config_IsBuildMachine==false && @.Payload.Config_IsRunningCommandlet==false && @.Payload.LoadingSeconds>5 && @.Payload.LoadingSeconds<3600)]",
+			"filter": "$[?(@.Payload.EventName == 'Core.Loading' && @.Payload.LoadingName=='PIE.LoadMap' && $(IsUserMachine) && $(IsLoadTimeInRange))]",
 			"property": "$.Payload.LoadingSeconds",
 			"groupBy": "$.Payload.Build_BranchName,$.Payload.ProjectName,$.Payload.MapName",
 			"function": "Average",
-			"interval": "03:00:00"
+			"interval": "03:00:00",
+			"topN": 10
 		},
 		{
 			"id": "core-loading-user-timetoedit",
-			"filter": "$[?(@.Payload.EventName == 'Core.Loading' && @.Payload.LoadingName=='TimeToEdit' && @.Payload.Config_IsBuildMachine==false && @.Payload.Config_IsRunningCommandlet==false && @.Payload.LoadingSeconds>5 && @.Payload.LoadingSeconds<3600)]",
+			"filter": "$[?(@.Payload.EventName == 'Core.Loading' && @.Payload.LoadingName=='TimeToEdit' && $(IsUserMachine) && $(IsLoadTimeInRange))]",
 			"property": "$.Payload.LoadingSeconds",
-			"groupBy": "$.Payload.Build_BranchName,$.Payload.ProjectName,$.Payload.MapName",
+			"groupBy": "$.Payload.Build_BranchName,$.Payload.ProjectName",
 			"function": "Average",
-			"interval": "03:00:00"
+			"interval": "03:00:00",
+			"topN": 10
 		},
 		{
 			"id": "core-loading-user-timetopie",
-			"filter": "$[?(@.Payload.EventName == 'Core.Loading' && @.Payload.LoadingName=='TimeToPIE' && @.Payload.Config_IsBuildMachine==false && @.Payload.Config_IsRunningCommandlet==false && @.Payload.LoadingSeconds>5 && @.Payload.LoadingSeconds<3600)]",
+			"filter": "$[?(@.Payload.EventName == 'Core.Loading' && @.Payload.LoadingName=='TimeToPIE' && $(IsUserMachine) && $(IsLoadTimeInRange))]",
 			"property": "$.Payload.LoadingSeconds",
 			"groupBy": "$.Payload.Build_BranchName,$.Payload.ProjectName,$.Payload.MapName",
 			"function": "Average",
-			"interval": "03:00:00"
+			"interval": "03:00:00",
+			"topN": 10
 		},
 		{
 			"id": "core-loading-user-piestartuptime",
-			"filter": "$[?(@.Payload.EventName == 'Core.Loading' && @.Payload.LoadingName=='PIE.Startup' && @.Payload.Config_IsBuildMachine==false && @.Payload.Config_IsRunningCommandlet==false && @.Payload.LoadingSeconds>5 && @.Payload.LoadingSeconds<3600)]",
+			"filter": "$[?(@.Payload.EventName == 'Core.Loading' && @.Payload.LoadingName=='PIE.Startup' && $(IsUserMachine) && $(IsLoadTimeInRange))]",
 			"property": "$.Payload.LoadingSeconds",
 			"groupBy": "$.Payload.Build_BranchName,$.Payload.ProjectName,$.Payload.MapName",
 			"function": "Average",
-			"interval": "03:00:00"
+			"interval": "03:00:00",
+			"topN": 10
 		},
 		{
 			"id": "core-loading-user-editorboottime-kpi",
-			"filter": "$[?(@.Payload.EventName == 'Core.Loading' && @.Payload.LoadingName=='Editor.Boot' && @.Payload.Config_IsBuildMachine==false && @.Payload.Config_IsRunningCommandlet==false && @.Payload.LoadingSeconds>5 && @.Payload.LoadingSeconds<3600)]",
+			"filter": "$[?(@.Payload.EventName == 'Core.Loading' && @.Payload.LoadingName=='Editor.Boot' && $(IsUserMachine) && $(IsLoadTimeInRange))]",
 			"property": "$.Payload.LoadingSeconds",
 			"groupBy": "$.Payload.Build_BranchName,$.Payload.ProjectName",
 			"function": "Average",
@@ -88,7 +116,7 @@
 		},
 		{
 			"id": "core-loading-user-editorstartuptime-kpi",
-			"filter": "$[?(@.Payload.EventName == 'Core.Loading' && @.Payload.LoadingName=='Editor.Initialize' && @.Payload.Config_IsBuildMachine==false && @.Payload.Config_IsRunningCommandlet==false && @.Payload.LoadingSeconds>5 && @.Payload.LoadingSeconds<3600)]",
+			"filter": "$[?(@.Payload.EventName == 'Core.Loading' && @.Payload.LoadingName=='Editor.Initialize' && $(IsUserMachine) && $(IsLoadTimeInRange))]",
 			"property": "$.Payload.LoadingSeconds",
 			"groupBy": "$.Payload.Build_BranchName,$.Payload.ProjectName",
 			"function": "Average",
@@ -96,7 +124,7 @@
 		},
 		{
 			"id": "core-loading-user-loadmaptime-kpi",
-			"filter": "$[?(@.Payload.EventName == 'Core.Loading' && @.Payload.LoadingName=='Editor.LoadMap' && @.Payload.Config_IsBuildMachine==false && @.Payload.Config_IsRunningCommandlet==false && @.Payload.LoadingSeconds>5 && @.Payload.LoadingSeconds<3600)]",
+			"filter": "$[?(@.Payload.EventName == 'Core.Loading' && @.Payload.LoadingName=='Editor.LoadMap' && $(IsUserMachine) && $(IsLoadTimeInRange))]",
 			"property": "$.Payload.LoadingSeconds",
 			"groupBy": "$.Payload.Build_BranchName,$.Payload.ProjectName",
 			"function": "Average",
@@ -104,7 +132,7 @@
 		},
 		{
 			"id": "core-loading-user-piestartuptime-kpi",
-			"filter": "$[?(@.Payload.EventName == 'Core.Loading' && @.Payload.LoadingName=='PIE.Startup' && @.Payload.Config_IsBuildMachine==false && @.Payload.Config_IsRunningCommandlet==false && @.Payload.LoadingSeconds>5 && @.Payload.LoadingSeconds<3600)]",
+			"filter": "$[?(@.Payload.EventName == 'Core.Loading' && @.Payload.LoadingName=='PIE.Startup' && $(IsUserMachine) && $(IsLoadTimeInRange))]",
 			"property": "$.Payload.LoadingSeconds",
 			"groupBy": "$.Payload.Build_BranchName,$.Payload.ProjectName",
 			"function": "Average",
@@ -112,7 +140,7 @@
 		},
 		{
 			"id": "core-loading-user-timetoedit-kpi",
-			"filter": "$[?(@.Payload.EventName == 'Core.Loading' && @.Payload.LoadingName=='TimeToEdit' && @.Payload.Config_IsBuildMachine==false && @.Payload.Config_IsRunningCommandlet==false && @.Payload.LoadingSeconds>5 && @.Payload.LoadingSeconds<3600)]",
+			"filter": "$[?(@.Payload.EventName == 'Core.Loading' && @.Payload.LoadingName=='TimeToEdit' && $(IsUserMachine) && $(IsLoadTimeInRange))]",
 			"property": "$.Payload.LoadingSeconds",
 			"groupBy": "$.Payload.Build_BranchName,$.Payload.ProjectName",
 			"function": "Average",
@@ -120,7 +148,7 @@
 		},
 		{
 			"id": "core-loading-user-timetopie-kpi",
-			"filter": "$[?(@.Payload.EventName == 'Core.Loading' && @.Payload.LoadingName=='TimeToPIE' && @.Payload.Config_IsBuildMachine==false && @.Payload.Config_IsRunningCommandlet==false && @.Payload.LoadingSeconds>5 && @.Payload.LoadingSeconds<3600)]",
+			"filter": "$[?(@.Payload.EventName == 'Core.Loading' && @.Payload.LoadingName=='TimeToPIE' && $(IsUserMachine) && $(IsLoadTimeInRange))]",
 			"property": "$.Payload.LoadingSeconds",
 			"groupBy": "$.Payload.Build_BranchName,$.Payload.ProjectName",
 			"function": "Average",
@@ -128,7 +156,7 @@
 		},
 		{
 			"id": "core-loading-user-cloudlocalddchitratio",
-			"filter": "$[?(@.Payload.EventName == 'Core.Loading' && @.Payload.Config_IsBuildMachine==false && @.Payload.DDC_Summary_CloudGetTotal>0 && @.Payload.LoadingSeconds>5 && @.Payload.LoadingSeconds<3600)]",
+			"filter": "$[?(@.Payload.EventName == 'Core.Loading' && $(IsUserMachine) && @.Payload.DDC_Summary_CloudGetTotal>0 && $(IsLoadTimeInRange))]",
 			"property": "$.Payload.DDC_Summary_CloudGetHitPct",
 			"groupBy": "$.Payload.Build_BranchName,$.Payload.ProjectName",
 			"function": "Average",
@@ -136,7 +164,7 @@
 		},
 		{
 			"id": "core-loading-user-zenlocalddchitratio",
-			"filter": "$[?(@.Payload.EventName == 'Core.Loading' && @.Payload.Config_IsBuildMachine==false && @.Payload.Config_IsRunningCommandlet==false && @.Payload.DDC_Summary_ZenLocalGetTotal>0 && @.Payload.LoadingSeconds>5 && @.Payload.LoadingSeconds<3600)]",
+			"filter": "$[?(@.Payload.EventName == 'Core.Loading' && $(IsUserMachine) && @.Payload.DDC_Summary_ZenLocalGetTotal>0 && $(IsLoadTimeInRange))]",
 			"property": "$.Payload.DDC_Summary_ZenLocalGetHitPct",
 			"groupBy": "$.Payload.Build_BranchName,$.Payload.ProjectName",
 			"function": "Average",
@@ -144,7 +172,7 @@
 		},
 		{
 			"id": "core-loading-user-sharedddchitratio",
-			"filter": "$[?(@.Payload.EventName == 'Core.Loading' && @.Payload.Config_IsBuildMachine==false && @.Payload.Config_IsRunningCommandlet==false && @.Payload.DDC_Summary_SharedGetTotal>0 && @.Payload.LoadingSeconds>5 && @.Payload.LoadingSeconds<3600)]",
+			"filter": "$[?(@.Payload.EventName == 'Core.Loading' && $(IsUserMachine) && @.Payload.DDC_Summary_SharedGetTotal>0 && $(IsLoadTimeInRange))]",
 			"property": "$.Payload.DDC_Summary_SharedGetHitPct",
 			"groupBy": "$.Payload.Build_BranchName,$.Payload.ProjectName",
 			"function": "Average",
@@ -152,7 +180,7 @@
 		},
 		{
 			"id": "core-loading-user-totalddchitratio",
-			"filter": "$[?(@.Payload.EventName == 'Core.Loading' && @.Payload.Config_IsBuildMachine==false && @.Payload.Config_IsRunningCommandlet==false && @.Payload.LoadingSeconds>5 && @.Payload.LoadingSeconds<3600)]",
+			"filter": "$[?(@.Payload.EventName == 'Core.Loading' && $(IsUserMachine) && $(IsLoadTimeInRange))]",
 			"property": "$.Payload.DDC_Summary_TotalGetHitPct",
 			"groupBy": "$.Payload.Build_BranchName,$.Payload.ProjectName",
 			"function": "Average",
@@ -160,7 +188,7 @@
 		},
 		{
 			"id": "core-loading-user-ddc-asset-efficiency",
-			"filter": "$[?(@.Payload.EventName=='Core.DDC.Resource' && @.Payload.Config_IsBuildMachine==false && @.Payload.Config_IsEditor==true && @.Payload.AssetType!='Total')]",
+			"filter": "$[?(@.Payload.EventName=='Core.DDC.Resource' && $(IsUserMachine) && @.Payload.AssetType!='Total')]",
 			"property": "$.Payload.Efficiency",
 			"groupBy": "$.Payload.Build_BranchName,$.Payload.ProjectName,$.Payload.AssetType",
 			"function": "Average",
@@ -168,8 +196,8 @@
 		},
 		{
 			"id": "core-loading-user-timetoedit-percentile",
-			"filter": "$[?(@.Payload.EventName == 'Core.Loading' && @.Payload.LoadingName=='TimeToEdit' && @.Payload.Config_IsBuildMachine==false && @.Payload.Config_IsRunningCommandlet==false && @.Payload.LoadingSeconds>5 && @.Payload.LoadingSeconds<3600)]",
-			"groupBy": "$.Payload.Build_BranchName,$.Payload.ProjectName,$.Payload.MapName",
+			"filter": "$[?(@.Payload.EventName == 'Core.Loading' && @.Payload.LoadingName=='TimeToEdit' && $(IsUserMachine) && $(IsLoadTimeInRange))]",
+			"groupBy": "$.Payload.Build_BranchName,$.Payload.ProjectName",
 			"property": "$.Payload.LoadingSeconds",
 			"function": "Percentile",
 			"percentile": 75,
@@ -177,7 +205,7 @@
 		},
 		{
 			"id": "core-loading-user-timetopie-percentile",
-			"filter": "$[?(@.Payload.EventName == 'Core.Loading' && @.Payload.LoadingName=='TimeToPIE' && @.Payload.Config_IsBuildMachine==false && @.Payload.Config_IsRunningCommandlet==false && @.Payload.LoadingSeconds>5 && @.Payload.LoadingSeconds<3600)]",
+			"filter": "$[?(@.Payload.EventName == 'Core.Loading' && @.Payload.LoadingName=='TimeToPIE' && $(IsUserMachine) && $(IsLoadTimeInRange))]",
 			"groupBy": "$.Payload.Build_BranchName,$.Payload.ProjectName,$.Payload.MapName",
 			"property": "$.Payload.LoadingSeconds",
 			"function": "Percentile",
@@ -186,7 +214,7 @@
 		},
 		{
 			"id": "core-loading-user-editorboottime-percentile",
-			"filter": "$[?(@.Payload.EventName == 'Core.Loading' && @.Payload.LoadingName=='Editor.Boot' && @.Payload.Config_IsBuildMachine==false && @.Payload.Config_IsRunningCommandlet==false && @.Payload.LoadingSeconds>5 && @.Payload.LoadingSeconds<3600 )]",
+			"filter": "$[?(@.Payload.EventName == 'Core.Loading' && @.Payload.LoadingName=='Editor.Boot' && $(IsUserMachine) && $(IsLoadTimeInRange) )]",
 			"groupBy": "$.Payload.Build_BranchName,$.Payload.ProjectName",
 			"property": "$.Payload.LoadingSeconds",
 			"function": "Percentile",
@@ -195,7 +223,7 @@
 		},
 		{
 			"id": "core-loading-user-editorstartuptime-percentile",
-			"filter": "$[?(@.Payload.EventName == 'Core.Loading' && @.Payload.LoadingName=='Editor.Initialize' && @.Payload.Config_IsBuildMachine==false && @.Payload.Config_IsRunningCommandlet==false && @.Payload.LoadingSeconds>5 && @.Payload.LoadingSeconds<3600 )]",
+			"filter": "$[?(@.Payload.EventName == 'Core.Loading' && @.Payload.LoadingName=='Editor.Initialize' && $(IsUserMachine) && $(IsLoadTimeInRange) )]",
 			"groupBy": "$.Payload.Build_BranchName,$.Payload.ProjectName",
 			"property": "$.Payload.LoadingSeconds",
 			"function": "Percentile",
@@ -204,7 +232,7 @@
 		},
 		{
 			"id": "core-loading-user-loadmaptime-percentile",
-			"filter": "$[?(@.Payload.EventName == 'Core.Loading' && @.Payload.LoadingName=='Editor.LoadMap' && @.Payload.Config_IsBuildMachine==false && @.Payload.LoadingSeconds>5 && @.Payload.LoadingSeconds<3600)]",
+			"filter": "$[?(@.Payload.EventName == 'Core.Loading' && @.Payload.LoadingName=='Editor.LoadMap' && $(IsUserMachine) && $(IsLoadTimeInRange))]",
 			"groupBy": "$.Payload.Build_BranchName,$.Payload.ProjectName,$.Payload.MapName",
 			"property": "$.Payload.LoadingSeconds",
 			"function": "Percentile",
@@ -213,7 +241,7 @@
 		},
 		{
 			"id": "core-loading-user-piestartuptime-percentile",
-			"filter": "$[?(@.Payload.EventName == 'Core.Loading' && @.Payload.LoadingName=='PIE.Startup' && @.Payload.Config_IsBuildMachine==false && @.Payload.LoadingSeconds>5 && @.Payload.LoadingSeconds<3600)]",
+			"filter": "$[?(@.Payload.EventName == 'Core.Loading' && @.Payload.LoadingName=='PIE.Startup' && $(IsUserMachine) && $(IsLoadTimeInRange))]",
 			"groupBy": "$.Payload.Build_BranchName,$.Payload.ProjectName,$.Payload.MapName",
 			"property": "$.Payload.LoadingSeconds",
 			"function": "Percentile",
@@ -426,7 +454,8 @@
 			"groupBy": "$.Payload.Build_BranchName,$.Payload.ProjectName,$.Payload.Hardware_Platform,$.Payload.Context,$.Payload.Name",
 			"property": "$.Payload.Size",
 			"function": "Average",
-			"interval": "00:15:00"
+			"interval": "00:15:00",
+			"topN": 25
 		},
 		{
 			"id": "core-memory-llm-assetclasses",
@@ -434,7 +463,8 @@
 			"groupBy": "$.Payload.Build_BranchName,$.Payload.ProjectName,$.Payload.Hardware_Platform,$.Payload.Context,$.Payload.Name",
 			"property": "$.Payload.Size",
 			"function": "Average",
-			"interval": "00:15:00"
+			"interval": "00:15:00",
+			"topN": 25
 		},
 		{
 			"id": "core-memory-llm-assets",
@@ -442,7 +472,8 @@
 			"groupBy": "$.Payload.Build_BranchName,$.Payload.ProjectName,$.Payload.Hardware_Platform,$.Payload.Context,$.Payload.Name",
 			"property": "$.Payload.Size",
 			"function": "Average",
-			"interval": "00:15:00"
+			"interval": "00:15:00",
+			"topN": 25
 		},
 		{
 			"id": "horde-incremental-build-jobs",
@@ -502,7 +533,7 @@
 		},
 		{
 			"id": "ubt-uba-user-duration-avg",
-			"filter": "$[?(@.Payload.EventName == 'Core.UBT.Executor' && $.Payload.Config_IsBuildMachine==false && search(@.Payload.Executor, \"Unreal Build Accelerator\") )]",
+			"filter": "$[?( $(IsUserMachine) && $(IsUBAExecutor) )]",
 			"groupBy": "$.Payload.Build_BranchName,$.Payload.Target,$.Payload.Mode,$.Payload.Configuration,$.Payload.Platform",
 			"property": "$.Payload.Duration",
 			"function": "Average",
@@ -510,7 +541,7 @@
 		},
 		{
 			"id": "ubt-uba-user-remote-usage-avg",
-			"filter": "$[?(@.Payload.EventName == 'Core.UBT.Executor' && $.Payload.Config_IsBuildMachine==false && search(@.Payload.Executor, \"Unreal Build Accelerator\") )]",
+			"filter": "$[?( $(IsUserMachine) && $(IsUBAExecutor) )]",
 			"groupBy": "$.Payload.Build_BranchName,$.Payload.Target,$.Payload.Mode,$.Payload.Configuration,$.Payload.Platform",
 			"property": "$.Payload.RemoteUsage",
 			"function": "Average",
@@ -518,7 +549,7 @@
 		},
 		{
 			"id": "ubt-uba-user-remote-actions-sum",
-			"filter": "$[?(@.Payload.EventName == 'Core.UBT.Executor' && $.Payload.Config_IsBuildMachine==false && search(@.Payload.Executor, \"Unreal Build Accelerator\") )]",
+			"filter": "$[?( $(IsUserMachine) && $(IsUBAExecutor) )]",
 			"groupBy": "$.Payload.Build_BranchName,$.Payload.Target,$.Payload.Mode,$.Payload.Configuration,$.Payload.Platform",
 			"property": "$.Payload.RemoteActions",
 			"function": "Sum",
@@ -526,7 +557,7 @@
 		},
 		{
 			"id": "ubt-uba-user-total-actions-sum",
-			"filter": "$[?(@.Payload.EventName == 'Core.UBT.Executor' && $.Payload.Config_IsBuildMachine==false && search(@.Payload.Executor, \"Unreal Build Accelerator\") )]",
+			"filter": "$[?( $(IsUserMachine) && $(IsUBAExecutor) )]",
 			"groupBy": "$.Payload.Build_BranchName,$.Payload.Target,$.Payload.Mode,$.Payload.Configuration,$.Payload.Platform",
 			"property": "$.Payload.TotalActions",
 			"function": "Sum",
@@ -534,7 +565,7 @@
 		},
 		{
 			"id": "ubt-uba-build-duration-avg",
-			"filter": "$[?(@.Payload.EventName == 'Core.UBT.Executor' && $.Payload.Config_IsBuildMachine==true && search(@.Payload.Executor, \"Unreal Build Accelerator\") )]",
+			"filter": "$[?( $(IsBuildMachine) && $(IsUBAExecutor) )]",
 			"groupBy": "$.Payload.Build_BranchName,$.Payload.Target,$.Payload.Mode,$.Payload.Configuration,$.Payload.Platform",
 			"property": "$.Payload.Duration",
 			"function": "Average",
@@ -542,7 +573,7 @@
 		},
 		{
 			"id": "ubt-uba-build-remote-usage-avg",
-			"filter": "$[?(@.Payload.EventName == 'Core.UBT.Executor' && $.Payload.Config_IsBuildMachine==true && search(@.Payload.Executor, \"Unreal Build Accelerator\") )]",
+			"filter": "$[?( $(IsBuildMachine) && $(IsUBAExecutor) )]",
 			"groupBy": "$.Payload.Build_BranchName,$.Payload.Target,$.Payload.Mode,$.Payload.Configuration,$.Payload.Platform",
 			"property": "$.Payload.RemoteUsage",
 			"function": "Average",
@@ -550,7 +581,7 @@
 		},
 		{
 			"id": "ubt-uba-build-remote-actions-sum",
-			"filter": "$[?(@.Payload.EventName == 'Core.UBT.Executor' && $.Payload.Config_IsBuildMachine==true && search(@.Payload.Executor, \"Unreal Build Accelerator\") )]",
+			"filter": "$[?( $(IsBuildMachine) && $(IsUBAExecutor) )]",
 			"groupBy": "$.Payload.Build_BranchName,$.Payload.Target,$.Payload.Mode,$.Payload.Configuration,$.Payload.Platform",
 			"property": "$.Payload.RemoteActions",
 			"function": "Sum",
@@ -558,7 +589,7 @@
 		},
 		{
 			"id": "ubt-uba-build-total-actions-sum",
-			"filter": "$[?(@.Payload.EventName == 'Core.UBT.Executor' && $.Payload.Config_IsBuildMachine==true && search(@.Payload.Executor, \"Unreal Build Accelerator\") )]",
+			"filter": "$[?( $(IsBuildMachine) && $(IsUBAExecutor) )]",
 			"groupBy": "$.Payload.Build_BranchName,$.Payload.Target,$.Payload.Mode,$.Payload.Configuration,$.Payload.Platform",
 			"property": "$.Payload.TotalActions",
 			"function": "Sum",
@@ -566,7 +597,7 @@
 		},
 		{
 			"id": "ubt-uba-duration",
-			"filter": "$[?(@.Payload.EventName == 'Core.UBT.Executor' && search(@.Payload.Executor, \"Unreal Build Accelerator\") )]",
+			"filter": "$[?( $(IsUBAExecutor) )]",
 			"groupBy": "$.Payload.Build_BranchName,$.Payload.Target,$.Payload.Mode,$.Payload.Configuration,$.Payload.Platform",
 			"property": "$.Payload.Duration",
 			"function": "Average",
@@ -574,7 +605,7 @@
 		},
 		{
 			"id": "ubt-uba-remote-usage",
-			"filter": "$[?(@.Payload.EventName == 'Core.UBT.Executor' && search(@.Payload.Executor, \"Unreal Build Accelerator\") )]",
+			"filter": "$[?( $(IsUBAExecutor) )]",
 			"groupBy": "$.Payload.Build_BranchName,$.Payload.Target,$.Payload.Mode,$.Payload.Configuration,$.Payload.Platform",
 			"property": "$.Payload.RemoteUsage",
 			"function": "Average",
@@ -582,7 +613,7 @@
 		},
 		{
 			"id": "ubt-uba-remote-actions",
-			"filter": "$[?(@.Payload.EventName == 'Core.UBT.Executor' && search(@.Payload.Executor, \"Unreal Build Accelerator\") )]",
+			"filter": "$[?( $(IsUBAExecutor) )]",
 			"groupBy": "$.Payload.Build_BranchName,$.Payload.Target,$.Payload.Mode,$.Payload.Configuration,$.Payload.Platform",
 			"property": "$.Payload.RemoteActions",
 			"function": "Sum",
@@ -590,7 +621,7 @@
 		},
 		{
 			"id": "ubt-uba-local-usage",
-			"filter": "$[?(@.Payload.EventName == 'Core.UBT.Executor' && search(@.Payload.Executor, \"Unreal Build Accelerator\") )]",
+			"filter": "$[?( $(IsUBAExecutor) )]",
 			"groupBy": "$.Payload.Build_BranchName,$.Payload.Target,$.Payload.Mode,$.Payload.Configuration,$.Payload.Platform",
 			"property": "$.Payload.LocalUsage",
 			"function": "Average",
@@ -598,7 +629,7 @@
 		},
 		{
 			"id": "ubt-uba-local-actions",
-			"filter": "$[?(@.Payload.EventName == 'Core.UBT.Executor' && search(@.Payload.Executor, \"Unreal Build Accelerator\") )]",
+			"filter": "$[?( $(IsUBAExecutor) )]",
 			"groupBy": "$.Payload.Build_BranchName,$.Payload.Target,$.Payload.Mode,$.Payload.Configuration,$.Payload.Platform",
 			"property": "$.Payload.LocalActions",
 			"function": "Sum",
@@ -606,7 +637,7 @@
 		},
 		{
 			"id": "ubt-uba-total-actions",
-			"filter": "$[?(@.Payload.EventName == 'Core.UBT.Executor' && search(@.Payload.Executor, \"Unreal Build Accelerator\") )]",
+			"filter": "$[?( $(IsUBAExecutor) )]",
 			"groupBy": "$.Payload.Build_BranchName,$.Payload.Target,$.Payload.Mode,$.Payload.Configuration,$.Payload.Platform",
 			"property": "$.Payload.TotalActions",
 			"function": "Sum",
@@ -614,7 +645,7 @@
 		},
 		{
 			"id": "ubt-uba-remote-actions-avg",
-			"filter": "$[?(@.Payload.EventName == 'Core.UBT.Executor' && search(@.Payload.Executor, \"Unreal Build Accelerator\") )]",
+			"filter": "$[?( $(IsUBAExecutor) )]",
 			"groupBy": "$.Payload.Build_BranchName,$.Payload.Target,$.Payload.Mode,$.Payload.Configuration,$.Payload.Platform",
 			"property": "$.Payload.RemoteActions",
 			"function": "Average",
@@ -622,7 +653,7 @@
 		},
 		{
 			"id": "ubt-uba-local-actions-avg",
-			"filter": "$[?(@.Payload.EventName == 'Core.UBT.Executor' && search(@.Payload.Executor, \"Unreal Build Accelerator\") )]",
+			"filter": "$[?( $(IsUBAExecutor) )]",
 			"groupBy": "$.Payload.Build_BranchName,$.Payload.Target,$.Payload.Mode,$.Payload.Configuration,$.Payload.Platform",
 			"property": "$.Payload.LocalActions",
 			"function": "Average",
@@ -630,7 +661,7 @@
 		},
 		{
 			"id": "ubt-uba-total-actions-avg",
-			"filter": "$[?(@.Payload.EventName == 'Core.UBT.Executor' && search(@.Payload.Executor, \"Unreal Build Accelerator\") )]",
+			"filter": "$[?( $(IsUBAExecutor) )]",
 			"groupBy": "$.Payload.Build_BranchName,$.Payload.Target,$.Payload.Mode,$.Payload.Configuration,$.Payload.Platform",
 			"property": "$.Payload.TotalActions",
 			"function": "Average",
@@ -638,7 +669,7 @@
 		},
 		{
 			"id": "ubt-uba-failure-rate",
-			"filter": "$[?(@.Payload.EventName == 'Core.UBT.Executor' && search(@.Payload.Executor, \"Unreal Build Accelerator\") )]",
+			"filter": "$[?( $(IsUBAExecutor) )]",
 			"groupBy": "$.Payload.Build_BranchName,$.Payload.Target,$.Payload.Mode,$.Payload.Configuration,$.Payload.Platform",
 			"property": "$.Payload.FailureRate",
 			"function": "Average",
@@ -646,7 +677,7 @@
 		},
 		{
 			"id": "ubt-uba-remote-usage-all",
-			"filter": "$[?(@.Payload.EventName == 'Core.UBT.Executor' && search(@.Payload.Executor, \"Unreal Build Accelerator\") )]",
+			"filter": "$[?( $(IsUBAExecutor) )]",
 			"groupBy": "$.Payload.Build_BranchName",
 			"property": "$.Payload.RemoteUsage",
 			"function": "Average",
@@ -654,7 +685,7 @@
 		},
 		{
 			"id": "ubt-uba-total-actions-all",
-			"filter": "$[?(@.Payload.EventName == 'Core.UBT.Executor' && search(@.Payload.Executor, \"Unreal Build Accelerator\") )]",
+			"filter": "$[?( $(IsUBAExecutor) )]",
 			"groupBy": "$.Payload.Build_BranchName",
 			"property": "$.Payload.TotalActions",
 			"function": "Sum",
@@ -662,7 +693,7 @@
 		},
 		{
 			"id": "ubt-uba-remote-actions-all",
-			"filter": "$[?(@.Payload.EventName == 'Core.UBT.Executor' && search(@.Payload.Executor, \"Unreal Build Accelerator\") )]",
+			"filter": "$[?( $(IsUBAExecutor) )]",
 			"groupBy": "$.Payload.Build_BranchName",
 			"property": "$.Payload.RemoteActions",
 			"function": "Sum",
@@ -670,7 +701,7 @@
 		},
 		{
 			"id": "ubt-uba-failure-rate-all",
-			"filter": "$[?(@.Payload.EventName == 'Core.UBT.Executor' && search(@.Payload.Executor, \"Unreal Build Accelerator\") )]",
+			"filter": "$[?( $(IsUBAExecutor) )]",
 			"groupBy": "$.Payload.Build_BranchName",
 			"property": "$.Payload.FailureRate",
 			"function": "Average",
