Horde: Tweaks to Analytics TopN/BottomN support

diff --git a/Engine/Source/Programs/Horde/Plugins/Analytics/HordeServer.Analytics.Tests/Telemetry/TelemetryTests.cs b/Engine/Source/Programs/Horde/Plugins/Analytics/HordeServer.Analytics.Tests/Telemetry/TelemetryTests.cs
--- a/Engine/Source/Programs/Horde/Plugins/Analytics/HordeServer.Analytics.Tests/Telemetry/TelemetryTests.cs
+++ b/Engine/Source/Programs/Horde/Plugins/Analytics/HordeServer.Analytics.Tests/Telemetry/TelemetryTests.cs
@@ -443,7 +443,6 @@ namespace HordeServer.Analytics.Tests.Telemetry
 			}
 		}
 
-		/*
 		[TestMethod]
 		public async Task TopAndBottomAsync()
 		{
@@ -508,7 +507,7 @@ namespace HordeServer.Analytics.Tests.Telemetry
 
 			// BottomN
 			{
-				Clock.UtcNow = new DateTime(2023, 6, 8, 4, 30, 0, DateTimeKind.Utc);
+				Clock.UtcNow = new DateTime(2024, 6, 8, 4, 30, 0, DateTimeKind.Utc);
 
 				MetricConfig metricConfig = new MetricConfig();
 				metricConfig.Id = new MetricId("test-metric-1");
@@ -551,7 +550,7 @@ namespace HordeServer.Analytics.Tests.Telemetry
 				await sink.FlushAsync(CancellationToken.None);
 				await collection.FlushAsync(CancellationToken.None);
 
-				List<IMetric> metrics = await collection.FindAsync(telemetryStoreConfig.Id, [metricConfig.Id]);
+				List<IMetric> metrics = await collection.FindAsync(telemetryStoreConfig.Id, [metricConfig.Id], new DateTime(2024, 1, 1, 1, 1, 1, DateTimeKind.Utc));
 				metrics = metrics.OrderBy(x => x.Group).ToList();
 				Assert.AreEqual(3, metrics.Count);
 
@@ -567,7 +566,6 @@ namespace HordeServer.Analytics.Tests.Telemetry
 				Assert.AreEqual(1, metrics[2].Count);
 				Assert.AreEqual(5, metrics[2].Value);
 			}
-		}
-		*/
+		}		
 	}
 }
diff --git a/Engine/Source/Programs/Horde/Plugins/Analytics/HordeServer.Analytics/Telemetry/Metrics/MetricCollection.cs b/Engine/Source/Programs/Horde/Plugins/Analytics/HordeServer.Analytics/Telemetry/Metrics/MetricCollection.cs
--- a/Engine/Source/Programs/Horde/Plugins/Analytics/HordeServer.Analytics/Telemetry/Metrics/MetricCollection.cs
+++ b/Engine/Source/Programs/Horde/Plugins/Analytics/HordeServer.Analytics/Telemetry/Metrics/MetricCollection.cs
@@ -247,8 +247,6 @@ namespace HordeServer.Telemetry.Metrics
 				_queuedSamples = new Dictionary<SampleKey, List<double>>();
 			}
 
-			Dictionary<SampleKey, MetricConfig> tbConstraint = new Dictionary<SampleKey, MetricConfig>();
-
 			// Add the samples to the database
 			try
 			{
@@ -258,11 +256,6 @@ namespace HordeServer.Telemetry.Metrics
 					if (_analyticsConfig.CurrentValue.TryGetTelemetryStore(sampleKey.Store, out TelemetryStoreConfig? telemetryStoreConfig) && telemetryStoreConfig.TryGetMetric(sampleKey.Metric, out metricConfig))
 					{
 						await CombineValuesAsync(sampleKey.Store, metricConfig, sampleKey.Group, sampleKey.Time, sampleValues, cancellationToken);
-
-						if (metricConfig.TopN > 0 || metricConfig.BottomN > 0)
-						{
-							tbConstraint.Add(sampleKey, metricConfig);
-						}
 					}
 				}
 			}
@@ -277,50 +270,6 @@ namespace HordeServer.Telemetry.Metrics
 				}
 				throw;
 			}
-
-			// Constrain to TopN or BottomN
-			/*
-			HashSet<ObjectId> removeIds = new HashSet<ObjectId>();			
-			if (tbConstraint.Count > 0)
-			{
-				foreach ((SampleKey sampleKey, MetricConfig metricConfig) in tbConstraint)
-				{
-					FilterDefinition < MetricDocument > filter = FilterDefinition<MetricDocument>.Empty;
-					filter &= Builders<MetricDocument>.Filter.Eq(x => x.TelemetryStoreId, sampleKey.Store);
-					filter &= Builders<MetricDocument>.Filter.Eq(x => x.Time, sampleKey.Time);
-					filter &= Builders<MetricDocument>.Filter.Eq(x => x.MetricId, sampleKey.Metric);
-
-					List<MetricDocument> results = await _metrics.Find(filter).SortByDescending(x => x.Value).ToListAsync(cancellationToken);
-
-					if (results.Count <= metricConfig.TopN && results.Count <= metricConfig.BottomN)
-					{
-						continue;
-					}
-
-					if (metricConfig.TopN > 0)
-					{
-						for (int i = metricConfig.TopN; i < results.Count; i++)
-						{
-							removeIds.Add(results[i].Id);
-						}
-					}
-					else
-					{
-						for (int i = 0; i < results.Count - metricConfig.BottomN; i++)
-						{
-							removeIds.Add(results[i].Id);
-						}
-					}
-				}			
-			}
-
-			if (removeIds.Count > 0)
-			{
-				FilterDefinition<MetricDocument> filter = FilterDefinition<MetricDocument>.Empty;
-				filter &= Builders<MetricDocument>.Filter.In(x => x.Id, removeIds);
-				await _metrics.DeleteManyAsync(filter, cancellationToken);
-			}
-			*/
 		}
 
 		async Task CombineValuesAsync(TelemetryStoreId telemetryStoreId, MetricConfig metricConfig, string group, DateTime time, List<double> values, CancellationToken cancellationToken)
@@ -402,7 +351,46 @@ namespace HordeServer.Telemetry.Metrics
 			}
 
 			List<MetricDocument> results = await _metrics.Find(filter).SortByDescending(x => x.Time).Limit(maxResults).ToListAsync(cancellationToken);
-			return results.ConvertAll<MetricDocument, IMetric>(x => x);
+
+			HashSet<ObjectId> filterIds = new HashSet<ObjectId>();
+			if (_analyticsConfig.CurrentValue.TryGetTelemetryStore(telemetryStoreId, out TelemetryStoreConfig? telemetryStoreConfig))
+			{
+				foreach (MetricId metricId in metricIds)
+				{
+					MetricConfig? metricConfig;
+					if (telemetryStoreConfig.TryGetMetric(metricId, out metricConfig))
+					{
+						if (metricConfig.TopN > 0 || metricConfig.BottomN > 0)
+						{
+							foreach (IGrouping<DateTime, MetricDocument> metrics in results.Where(x =>  x.MetricId == metricId).GroupBy(x => x.Time, x => x))
+							{
+								int count = metrics.Count();
+								if (count <= metricConfig.TopN && count <= metricConfig.BottomN)
+								{
+									continue;
+								}
+
+								if (metricConfig.TopN > 0)
+								{
+									foreach (MetricDocument m in metrics.OrderBy(x => x.Value).Take(count - metricConfig.TopN))
+									{
+										filterIds.Add(m.Id);
+									}
+								}
+								else
+								{
+									foreach (MetricDocument m in metrics.OrderBy(x => x.Value).TakeLast(count - metricConfig.BottomN))
+									{
+										filterIds.Add(m.Id);
+									}
+								}
+							}
+						}
+					}
+				}				 
+			}
+
+			return results.Where(x => !filterIds.Contains(x.Id)).ToList().ConvertAll<MetricDocument, IMetric>(x => x);
 		}
 	}
 }
