A few improvements to the dashboards.

diff --git a/Engine/Source/Programs/Horde/HordeServer/Defaults/default-metrics.telemetry.json b/Engine/Source/Programs/Horde/HordeServer/Defaults/default-metrics.telemetry.json
--- a/Engine/Source/Programs/Horde/HordeServer/Defaults/default-metrics.telemetry.json
+++ b/Engine/Source/Programs/Horde/HordeServer/Defaults/default-metrics.telemetry.json
@@ -192,7 +192,8 @@
 			"property": "$.Payload.Efficiency",
 			"groupBy": "$.Payload.Build_BranchName,$.Payload.ProjectName,$.Payload.AssetType",
 			"function": "Average",
-			"interval": "03:00:00"
+			"interval": "03:00:00",
+			"bottomN": 25
 		},
 		{
 			"id": "core-loading-user-timetoedit-percentile",
diff --git a/Engine/Source/Programs/Horde/HordeServer/Defaults/default-views.telemetry.json b/Engine/Source/Programs/Horde/HordeServer/Defaults/default-views.telemetry.json
--- a/Engine/Source/Programs/Horde/HordeServer/Defaults/default-views.telemetry.json
+++ b/Engine/Source/Programs/Horde/HordeServer/Defaults/default-views.telemetry.json
@@ -62,12 +62,19 @@
 				},
 				{
 					"name": "Target",
-					"group": "Target"
+					"group": "Target",
+					"defaults": [
+						"CitySample",
+						"CityEditor",
+						"LyraEditor",
+						"LyraGame"
+					]
 				},
 				{
 					"name": "Mode",
 					"group": "Mode",
 					"defaults": [
+						"Build",
 						"Rebuild"
 					]
 				},
@@ -107,21 +114,6 @@
 									"alias": "TimeToPIE",
 									"threshold": 360
 								}
-								/*,{
-									"id": "core-loading-user-editorboottime-percentile",
-									"alias": "BootEditor",
-									"threshold": 120
-								},
-								{
-									"id": "core-loading-user-loadmaptime-percentile",
-									"alias": "LoadMap",
-									"threshold": 240
-								},
-								{
-									"id": "core-loading-user-piestartuptime-percentile",
-									"alias": "PIEStartup",
-									"threshold": 360
-								}*/
 							],
 							"graph": "indicator",
 							"max": 720
@@ -139,21 +131,6 @@
 									"alias": "TimeToPIE",
 									"threshold": 360
 								}
-								/*,{
-									"id": "core-loading-user-editorboottime",
-									"alias": "BootEditor",
-									"threshold": 120
-								},
-								{
-									"id": "core-loading-user-loadmaptime",
-									"alias": "LoadMap",
-									"threshold": 240
-								},
-								{
-									"id": "core-loading-user-piestartuptime",
-									"alias": "PIEStartup",
-									"threshold": 360
-								}*/
 							],
 							"graph": "indicator",
 							"max": 900
@@ -169,19 +146,6 @@
 									"id": "core-loading-user-timetopie-kpi",
 									"alias": "TimeToPIE"
 								}
-								/*,
-								{
-									"id": "core-loading-user-editorboottime-kpi",
-									"alias": "BootEditor"
-								},
-								{
-									"id": "core-loading-user-loadmaptime-kpi",
-									"alias": "LoadMap"
-								},
-								{
-									"id": "core-loading-user-piestartuptime-kpi",
-									"alias": "PIEStartup"
-								}*/
 							]
 						},
 						{
