Horde: CpuCount and CpuMultiplier in AgentSettings used in LogicalCores property

- Limit the number of logical processors to CpuCount or number of logical processors reported by the OS, whichever is the smallest
- Scale the number of logical processors by CpuMultiplier
- Change CpuCount to a nullable value type to know when it is to be used
- Find the number of logical processors on Windows
- Set UE_HORDE_CPU_COUNT to CpuCount when CpuCount is set

diff --git a/Engine/Source/Programs/Horde/HordeAgent/AgentSettings.cs b/Engine/Source/Programs/Horde/HordeAgent/AgentSettings.cs
--- a/Engine/Source/Programs/Horde/HordeAgent/AgentSettings.cs
+++ b/Engine/Source/Programs/Horde/HordeAgent/AgentSettings.cs
@@ -198,12 +198,12 @@ namespace HordeAgent
 		/// Maximum size of the bundle cache, in megabytes.
 		/// </summary>
 		public long BundleCacheSize { get; set; } = 1024;
-		
+
 		/// <summary>
 		/// Maximum number of logical CPU cores workloads should use
 		/// Currently this is only provided as a hint and requires leases to respect this value as it's set via an env variable (UE_HORDE_CPU_COUNT).
 		/// </summary>
-		public int CpuCount { get; set; } = Environment.ProcessorCount;
+		public int? CpuCount { get; set; } = null;
 		
 		/// <summary>
 		/// CPU core multiplier applied to CPU core count setting
diff --git a/Engine/Source/Programs/Horde/HordeAgent/Leases/Handlers/ComputeHandler.cs b/Engine/Source/Programs/Horde/HordeAgent/Leases/Handlers/ComputeHandler.cs
--- a/Engine/Source/Programs/Horde/HordeAgent/Leases/Handlers/ComputeHandler.cs
+++ b/Engine/Source/Programs/Horde/HordeAgent/Leases/Handlers/ComputeHandler.cs
@@ -84,7 +84,7 @@ namespace HordeAgent.Leases.Handlers
 						{
 							["UE_HORDE_SHARED_DIR"] = sharedDir.FullName,
 							["UE_HORDE_TERMINATION_SIGNAL_FILE"] = _settings.CurrentValue.GetTerminationSignalFile().FullName,
-							["UE_HORDE_CPU_COUNT"] = Convert.ToString(_settings.CurrentValue.CpuCount),
+							["UE_HORDE_CPU_COUNT"] = Convert.ToString(_settings.CurrentValue.CpuCount ?? Environment.ProcessorCount),
 							["UE_HORDE_CPU_MULTIPLIER"] = Convert.ToString(_settings.CurrentValue.CpuMultiplier, CultureInfo.InvariantCulture),
 						};
 						
diff --git a/Engine/Source/Programs/Horde/HordeAgent/Services/CapabilitiesService.cs b/Engine/Source/Programs/Horde/HordeAgent/Services/CapabilitiesService.cs
--- a/Engine/Source/Programs/Horde/HordeAgent/Services/CapabilitiesService.cs
+++ b/Engine/Source/Programs/Horde/HordeAgent/Services/CapabilitiesService.cs
@@ -156,10 +156,11 @@ namespace HordeAgent.Services
 										}
 									case "NUMBEROFCORES" when property.Value is uint c: totalPhysicalCores += (int)c; break;
 									case "NUMBEROFENABLEDCORE" when property.Value is uint c: totalEnabledPhysicalCores += (int)c; break;
+									case "NUMBEROFLOGICALPROCESSORS" when property.Value is uint c: totalLogicalCores += (int)c; break;
 								}
 							}
 						}
-						
+
 						AddCpuInfo(capabilities, cpuNameToCount, totalLogicalCores, Math.Min(totalEnabledPhysicalCores, totalPhysicalCores));
 					}
 					
@@ -350,7 +351,13 @@ namespace HordeAgent.Services
 						capabilities.Properties.Add((cpuPackagesNode.InnerText != "1") ? $"CPU={cpuPackagesNode.InnerText} x {cpuTypeNode.InnerText} @ {cpuSpeedNode.InnerText}" : $"CPU={cpuTypeNode.InnerText} @ {cpuSpeedNode.InnerText}");
 					}
 
-					capabilities.Resources.Add(KnownResourceNames.LogicalCores, Environment.ProcessorCount);
+					int numLogicalCores = Environment.ProcessorCount;
+					if (_settings.CurrentValue.CpuCount is > 0)
+					{
+						numLogicalCores = Math.Min(numLogicalCores, _settings.CurrentValue.CpuCount.Value);
+					}
+					numLogicalCores = (int)(numLogicalCores * _settings.CurrentValue.CpuMultiplier);
+					capabilities.Resources.Add(KnownResourceNames.LogicalCores, numLogicalCores);
 
 					XmlNode? cpuCountNode = hardwareNode.SelectSingleNode("key[. = 'number_processors']/following-sibling::integer");
 					if (cpuCountNode != null)
@@ -543,13 +550,19 @@ namespace HordeAgent.Services
 			}
 		}
 
-		static void AddCpuInfo(RpcAgentCapabilities capabilities, Dictionary<string, int> nameToCount, int numLogicalCores, int numPhysicalCores)
+		void AddCpuInfo(RpcAgentCapabilities capabilities, Dictionary<string, int> nameToCount, int numLogicalCores, int numPhysicalCores)
 		{
 			if (nameToCount.Count > 0)
 			{
 				capabilities.Properties.Add("CPU=" + String.Join(", ", nameToCount.Select(x => (x.Value > 1) ? $"{x.Key} x {x.Value}" : x.Key)));
 			}
 
+			if (_settings.CurrentValue.CpuCount is > 0)
+			{
+				numLogicalCores = Math.Min(numLogicalCores, _settings.CurrentValue.CpuCount.Value);
+			}
+			numLogicalCores = (int)(numLogicalCores * _settings.CurrentValue.CpuMultiplier);
+
 			if (numLogicalCores > 0)
 			{
 				capabilities.Resources.Add("LogicalCores", numLogicalCores);
