Added Editor Perf metrics to Horde Analytics

diff --git a/Engine/Source/Programs/Horde/HordeServer/Defaults/default-metrics.telemetry.json b/Engine/Source/Programs/Horde/HordeServer/Defaults/default-metrics.telemetry.json
--- a/Engine/Source/Programs/Horde/HordeServer/Defaults/default-metrics.telemetry.json
+++ b/Engine/Source/Programs/Horde/HordeServer/Defaults/default-metrics.telemetry.json
@@ -15,9 +15,65 @@
 		{
 			"name": "IsUBAExecutor",
 			"value": "@.Payload.EventName == 'Core.UBT.Executor' && search(@.Payload.Executor, \"Unreal Build Accelerator\")"
+		},
+		{
+			"name": "IsEditorPerf",
+			"value": "search(@.Payload.Horde_TemplateName, \"EditorPerf Test\")"
 		}
 	],
 	"metrics": [
+		{
+			"id": "core-loading-build-timetoedit-kpi",
+			"filter": "$[?(@.Payload.EventName == 'Core.Loading' && $(IsEditorPerf) && $(IsLoadTimeInRange) && @.Payload.LoadingName=='TimeToEdit')]",
+			"property": "$.Payload.LoadingSeconds",
+			"groupBy": "$.Payload.Build_BranchName,$.Payload.Horde_StepName",
+			"function": "Average",
+			"interval": "00:15:00"
+		},
+		{
+			"id": "core-loading-build-timetopie-kpi",
+			"filter": "$[?(@.Payload.EventName == 'Core.Loading' && $(IsEditorPerf) && $(IsLoadTimeInRange) && @.Payload.LoadingName=='TimeToPIE')]",
+			"property": "$.Payload.LoadingSeconds",
+			"groupBy": "$.Payload.Build_BranchName,$.Payload.Horde_StepName",
+			"function": "Average",
+			"interval": "00:15:00"
+		},
+		{
+			"id": "core-loading-build-timetoedit-percentile",
+			"filter": "$[?(@.Payload.EventName == 'Core.Loading' && $(IsEditorPerf) && $(IsLoadTimeInRange) && @.Payload.LoadingName=='TimeToEdit')]",
+			"groupBy": "$.Payload.Build_BranchName,$.Payload.Horde_StepName",
+			"property": "$.Payload.LoadingSeconds",
+			"function": "Percentile",
+			"percentile": 75,
+			"interval": "00:15:00"
+		},
+		{
+			"id": "core-loading-build-timetopie-percentile",
+			"filter": "$[?(@.Payload.EventName == 'Core.Loading' && $(IsEditorPerf) && $(IsLoadTimeInRange) && @.Payload.LoadingName=='TimeToPIE')]",
+			"groupBy": "$.Payload.Build_BranchName,$.Payload.Horde_StepName",
+			"property": "$.Payload.LoadingSeconds",
+			"function": "Percentile",
+			"percentile": 75,
+			"interval": "00:15:00"
+		},
+		{
+			"id": "core-loading-build-timetoedit",
+			"filter": "$[?(@.Payload.EventName == 'Core.Loading' && $(IsEditorPerf) && $(IsLoadTimeInRange) && @.Payload.LoadingName=='TimeToEdit')]",
+			"property": "$.Payload.LoadingSeconds",
+			"groupBy": "$.Payload.Build_BranchName,$.Payload.Horde_StepName",
+			"function": "Average",
+			"interval": "00:15:00",
+			"topN": 10
+		},
+		{
+			"id": "core-loading-build-timetopie",
+			"filter": "$[?(@.Payload.EventName == 'Core.Loading' && $(IsEditorPerf) && $(IsLoadTimeInRange) && @.Payload.LoadingName=='TimeToPIE')]",
+			"property": "$.Payload.LoadingSeconds",
+			"groupBy": "$.Payload.Build_BranchName,$.Payload.Horde_StepName",
+			"function": "Average",
+			"interval": "00:15:00",
+			"topN": 10
+		},
 		{
 			"id": "core-loading-user-loadingtime",
 			"filter": "$[?(@.Payload.EventName == 'Core.Loading' && $(IsUserMachine) && $(IsLoadTimeInRange) )]",
@@ -452,7 +508,7 @@
 		{
 			"id": "core-memory-llm-systems",
 			"filter": "$[?(@.Payload.EventName == 'Core.Memory.LLM' && @.Payload.TagSet=='None')]",
-			"groupBy": "$.Payload.Build_BranchName,$.Payload.ProjectName,$.Payload.Hardware_Platform,$.Payload.Context,$.Payload.Name",
+			"groupBy": "$.Payload.Build_BranchName,$.Payload.ProjectName,$.Payload.Hardware_Platform,$.Payload.LoadingName,$.Payload.Name",
 			"property": "$.Payload.Size",
 			"function": "Average",
 			"interval": "00:15:00",
@@ -461,7 +517,7 @@
 		{
 			"id": "core-memory-llm-assetclasses",
 			"filter": "$[?(@.Payload.EventName == 'Core.Memory.LLM' && @.Payload.TagSet=='AssetClasses')]",
-			"groupBy": "$.Payload.Build_BranchName,$.Payload.ProjectName,$.Payload.Hardware_Platform,$.Payload.Context,$.Payload.Name",
+			"groupBy": "$.Payload.Build_BranchName,$.Payload.ProjectName,$.Payload.Hardware_Platform,$.Payload.LoadingName,$.Payload.Name",
 			"property": "$.Payload.Size",
 			"function": "Average",
 			"interval": "00:15:00",
@@ -470,7 +526,7 @@
 		{
 			"id": "core-memory-llm-assets",
 			"filter": "$[?(@.Payload.EventName == 'Core.Memory.LLM' && @.Payload.TagSet=='Assets')]",
-			"groupBy": "$.Payload.Build_BranchName,$.Payload.ProjectName,$.Payload.Hardware_Platform,$.Payload.Context,$.Payload.Name",
+			"groupBy": "$.Payload.Build_BranchName,$.Payload.ProjectName,$.Payload.Hardware_Platform,$.Payload.LoadingName,$.Payload.Name",
 			"property": "$.Payload.Size",
 			"function": "Average",
 			"interval": "00:15:00",
diff --git a/Engine/Source/Programs/Horde/HordeServer/Defaults/default-views.telemetry.json b/Engine/Source/Programs/Horde/HordeServer/Defaults/default-views.telemetry.json
--- a/Engine/Source/Programs/Horde/HordeServer/Defaults/default-views.telemetry.json
+++ b/Engine/Source/Programs/Horde/HordeServer/Defaults/default-views.telemetry.json
@@ -101,6 +101,54 @@
 				{
 					"name": "Editor",
 					"charts": [
+
+						{
+							"name": "Build Load Time (75th Percentile)",
+							"metrics": [
+								{
+									"id": "core-loading-build-timetoedit-percentile",
+									"alias": "TimeToEdit",
+									"threshold": 120
+								},
+								{
+									"id": "core-loading-build-timetopie-percentile",
+									"alias": "TimeToPIE",
+									"threshold": 360
+								}
+							],
+							"graph": "indicator",
+							"max": 720
+						},
+						{
+							"name": "Build Load Time (Average)",
+							"metrics": [
+								{
+									"id": "core-loading-build-timetoedit",
+									"alias": "TimeToEdit",
+									"threshold": 120
+								},
+								{
+									"id": "core-loading-build-timetopie",
+									"alias": "TimeToPIE",
+									"threshold": 360
+								}
+							],
+							"graph": "indicator",
+							"max": 900
+						},
+						{
+							"name": "Build Loading Time (KPIs)",
+							"metrics": [
+								{
+									"id": "core-loading-build-timetoedit-kpi",
+									"alias": "TimeToEdit"
+								},
+								{
+									"id": "core-loading-build-timetopie-kpi",
+									"alias": "TimeToPIE"
+								}
+							]
+						},
 						{
 							"name": "User Load Time (75th Percentile)",
 							"metrics": [
@@ -123,7 +171,7 @@
 							"metrics": [
 								{
 									"id": "core-loading-user-timetoedit",
-									"alias": "TimeToPIE",
+									"alias": "TimeToEdit",
 									"threshold": 120
 								},
 								{
@@ -147,7 +195,8 @@
 									"alias": "TimeToPIE"
 								}
 							]
-						},
+						}
+						,
 						{
 							"name": "User Loading Count (Live)",
 							"display": "Value",
