Tweak to Editor Perf metrics in Horde Analytics
Added grouping by [REDACTED_PERSON]

diff --git a/Engine/Source/Programs/Horde/HordeServer/Defaults/default-metrics.telemetry.json b/Engine/Source/Programs/Horde/HordeServer/Defaults/default-metrics.telemetry.json
--- a/Engine/Source/Programs/Horde/HordeServer/Defaults/default-metrics.telemetry.json
+++ b/Engine/Source/Programs/Horde/HordeServer/Defaults/default-metrics.telemetry.json
@@ -26,7 +26,7 @@
 			"id": "core-loading-build-timetoedit-kpi",
 			"filter": "$[?(@.Payload.EventName == 'Core.Loading' && $(IsEditorPerf) && $(IsLoadTimeInRange) && @.Payload.LoadingName=='TimeToEdit')]",
 			"property": "$.Payload.LoadingSeconds",
-			"groupBy": "$.Payload.Build_BranchName,$.Payload.Horde_StepName",
+			"groupBy": "$.Payload.Build_BranchName,$.Payload.ProjectName,$.Payload.Horde_StepName",
 			"function": "Average",
 			"interval": "00:15:00"
 		},
@@ -34,14 +34,14 @@
 			"id": "core-loading-build-timetopie-kpi",
 			"filter": "$[?(@.Payload.EventName == 'Core.Loading' && $(IsEditorPerf) && $(IsLoadTimeInRange) && @.Payload.LoadingName=='TimeToPIE')]",
 			"property": "$.Payload.LoadingSeconds",
-			"groupBy": "$.Payload.Build_BranchName,$.Payload.Horde_StepName",
+			"groupBy": "$.Payload.Build_BranchName,$.Payload.ProjectName,$.Payload.Horde_StepName",
 			"function": "Average",
 			"interval": "00:15:00"
 		},
 		{
 			"id": "core-loading-build-timetoedit-percentile",
 			"filter": "$[?(@.Payload.EventName == 'Core.Loading' && $(IsEditorPerf) && $(IsLoadTimeInRange) && @.Payload.LoadingName=='TimeToEdit')]",
-			"groupBy": "$.Payload.Build_BranchName,$.Payload.Horde_StepName",
+			"groupBy": "$.Payload.Build_BranchName,$.Payload.ProjectName,$.Payload.Horde_StepName",
 			"property": "$.Payload.LoadingSeconds",
 			"function": "Percentile",
 			"percentile": 75,
@@ -50,7 +50,7 @@
 		{
 			"id": "core-loading-build-timetopie-percentile",
 			"filter": "$[?(@.Payload.EventName == 'Core.Loading' && $(IsEditorPerf) && $(IsLoadTimeInRange) && @.Payload.LoadingName=='TimeToPIE')]",
-			"groupBy": "$.Payload.Build_BranchName,$.Payload.Horde_StepName",
+			"groupBy": "$.Payload.Build_BranchName,$.Payload.ProjectName,$.Payload.Horde_StepName",
 			"property": "$.Payload.LoadingSeconds",
 			"function": "Percentile",
 			"percentile": 75,
@@ -60,19 +60,17 @@
 			"id": "core-loading-build-timetoedit",
 			"filter": "$[?(@.Payload.EventName == 'Core.Loading' && $(IsEditorPerf) && $(IsLoadTimeInRange) && @.Payload.LoadingName=='TimeToEdit')]",
 			"property": "$.Payload.LoadingSeconds",
-			"groupBy": "$.Payload.Build_BranchName,$.Payload.Horde_StepName",
+			"groupBy": "$.Payload.Build_BranchName,$.Payload.ProjectName,$.Payload.Horde_StepName",
 			"function": "Average",
 			"interval": "00:15:00",
-			"topN": 10
 		},
 		{
 			"id": "core-loading-build-timetopie",
 			"filter": "$[?(@.Payload.EventName == 'Core.Loading' && $(IsEditorPerf) && $(IsLoadTimeInRange) && @.Payload.LoadingName=='TimeToPIE')]",
 			"property": "$.Payload.LoadingSeconds",
-			"groupBy": "$.Payload.Build_BranchName,$.Payload.Horde_StepName",
+			"groupBy": "$.Payload.Build_BranchName,$.Payload.ProjectName,$.Payload.Horde_StepName",
 			"function": "Average",
 			"interval": "00:15:00",
-			"topN": 10
 		},
 		{
 			"id": "core-loading-user-loadingtime",
