Fixed JSON Formatting

diff --git a/Engine/Source/Programs/Horde/HordeServer/Defaults/default-metrics.telemetry.json b/Engine/Source/Programs/Horde/HordeServer/Defaults/default-metrics.telemetry.json
--- a/Engine/Source/Programs/Horde/HordeServer/Defaults/default-metrics.telemetry.json
+++ b/Engine/Source/Programs/Horde/HordeServer/Defaults/default-metrics.telemetry.json
@@ -62,7 +62,7 @@
 			"property": "$.Payload.LoadingSeconds",
 			"groupBy": "$.Payload.Build_BranchName,$.Payload.ProjectName,$.Payload.Horde_StepName",
 			"function": "Average",
-			"interval": "00:15:00",
+			"interval": "00:15:00"
 		},
 		{
 			"id": "core-loading-build-timetopie",
@@ -70,7 +70,7 @@
 			"property": "$.Payload.LoadingSeconds",
 			"groupBy": "$.Payload.Build_BranchName,$.Payload.ProjectName,$.Payload.Horde_StepName",
 			"function": "Average",
-			"interval": "00:15:00",
+			"interval": "00:15:00"
 		},
 		{
 			"id": "core-loading-user-loadingtime",
