Reconfigured Horde Analytics dashboard

diff --git a/Engine/Source/Programs/Horde/HordeServer/Defaults/default-views.telemetry.json b/Engine/Source/Programs/Horde/HordeServer/Defaults/default-views.telemetry.json
--- a/Engine/Source/Programs/Horde/HordeServer/Defaults/default-views.telemetry.json
+++ b/Engine/Source/Programs/Horde/HordeServer/Defaults/default-views.telemetry.json
@@ -99,56 +99,8 @@
 			],
 			"categories": [
 				{
-					"name": "Editor",
+					"name": "Live",
 					"charts": [
-
 						{
-							"name": "Build Load Time (75th Percentile)",
-							"metrics": [
-								{
-									"id": "core-loading-build-timetoedit-percentile",
-									"alias": "TimeToEdit",
-									"threshold": 120
-								},
-								{
-									"id": "core-loading-build-timetopie-percentile",
-									"alias": "TimeToPIE",
-									"threshold": 360
-								}
-							],
-							"graph": "indicator",
-							"max": 720
-						},
-						{
-							"name": "Build Load Time (Average)",
-							"metrics": [
-								{
-									"id": "core-loading-build-timetoedit",
-									"alias": "TimeToEdit",
-									"threshold": 120
-								},
-								{
-									"id": "core-loading-build-timetopie",
-									"alias": "TimeToPIE",
-									"threshold": 360
-								}
-							],
-							"graph": "indicator",
-							"max": 900
-						},
-						{
-							"name": "Build Loading Time (KPIs)",
-							"metrics": [
-								{
-									"id": "core-loading-build-timetoedit-kpi",
-									"alias": "TimeToEdit"
-								},
-								{
-									"id": "core-loading-build-timetopie-kpi",
-									"alias": "TimeToPIE"
-								}
-							]
-						},
-						{
 							"name": "User Load Time (75th Percentile)",
 							"metrics": [
@@ -195,8 +147,7 @@
 									"alias": "TimeToPIE"
 								}
 							]
-						}
-						,
+						},
 						{
 							"name": "User Loading Count (Live)",
 							"display": "Value",
@@ -259,32 +210,52 @@
 					]
 				},
 				{
-					"name": "Memory",
+					"name": "Load",
 					"charts": [
 						{
-							"name": "Memory By System",
-							"display": "Value",
+							"name": "Load Time (75th Percentile)",
 							"metrics": [
 								{
-									"id": "core-memory-llm-systems"
+									"id": "core-loading-build-timetoedit-percentile",
+									"alias": "TimeToEdit",
+									"threshold": 120
+								},
+								{
+									"id": "core-loading-build-timetopie-percentile",
+									"alias": "TimeToPIE",
+									"threshold": 360
+								}
+							],
+							"graph": "indicator",
+							"max": 720
+						},
+						{
+							"name": "Load Time (Average)",
+							"metrics": [
+								{
+									"id": "core-loading-build-timetoedit",
+									"alias": "TimeToEdit",
+									"threshold": 120
+								},
+								{
+									"id": "core-loading-build-timetopie",
+									"alias": "TimeToPIE",
+									"threshold": 360
+								}
+							],
+							"graph": "indicator",
+							"max": 900
+						},
+						{
+							"name": "Load Time (KPIs)",
+							"metrics": [
+								{
+									"id": "core-loading-build-timetoedit-kpi",
+									"alias": "TimeToEdit"
+								},
+								{
+									"id": "core-loading-build-timetopie-kpi",
+									"alias": "TimeToPIE"
 								}
 							]
-						},
-						{
-							"name": "Memory By AssetClass",
-							"display": "Value",
-							"metrics": [
-								{
-									"id": "core-memory-llm-assetclasses"
-								}
-							]
-						},
-						{
-							"name": "Memory By Asset",
-							"display": "Value",
-							"metrics": [
-								{
-									"id": "core-memory-llm-assets"
-								}
-							]
 						}
@@ -526,40 +497,32 @@
 					]
 				},
 				{
-					"name": "Build",
+					"name": "Memory",
 					"charts": [
 						{
-							"name": "Incremental Build",
+							"name": "Memory By System",
+							"display": "Value",
 							"metrics": [
 								{
-									"id": "horde-incremental-build-jobs"
-								},
-								{
-									"id": "horde-incremental-ci-jobs"
+									"id": "core-memory-llm-systems"
 								}
 							]
 						},
 						{
-							"name": "AutoTest Build",
+							"name": "Memory By AssetClass",
+							"display": "Value",
 							"metrics": [
 								{
-									"id": "horde-autotest-ci-jobs"
+									"id": "core-memory-llm-assetclasses"
 								}
 							]
 						},
 						{
-							"name": "Package And Test Build",
+							"name": "Memory By Asset",
+							"display": "Value",
 							"metrics": [
 								{
-									"id": "horde-package-and-test-jobs"
+									"id": "core-memory-llm-assets"
 								}
 							]
-						},
-						{
-							"name": "Nighty Build",
-							"metrics": [
-								{
-									"id": "horde-nightly-build-jobs"
-								}
-							]
 						}
@@ -583,12 +546,7 @@
 									"id": "horde-incremental-compile-jobs"
 								}
 							]
-						}
-					]
-				},
-				{
-					"name": "UBA",
-					"charts": [
+						},
 						{
 							"name": "Remote Distibution All (Average)",
 							"display": "Ratio",
@@ -696,14 +654,40 @@
 					]
 				},
 				{
-					"name": "UBT",
+					"name": "Build",
 					"charts": [
 						{
-							"name": "Wall Time (Average)",
+							"name": "Incremental Build",
 							"metrics": [
 								{
-									"id": "ubt-completed-duration"
-
+									"id": "horde-incremental-build-jobs"
+								},
+								{
+									"id": "horde-incremental-ci-jobs"
 								}
 							]
+						},
+						{
+							"name": "AutoTest Build",
+							"metrics": [
+								{
+									"id": "horde-autotest-ci-jobs"
+								}
+							]
+						},
+						{
+							"name": "Package And Test Build",
+							"metrics": [
+								{
+									"id": "horde-package-and-test-jobs"
+								}
+							]
+						},
+						{
+							"name": "Nighty Build",
+							"metrics": [
+								{
+									"id": "horde-nightly-build-jobs"
+								}
+							]
 						}
