Added Content Metrics to Horde Analytics
Added Content vizualization to Horde Analytics
Records the total sizes and virtualiszation level of all source content

diff --git a/Engine/Source/Programs/Horde/HordeServer/Defaults/default-metrics.telemetry.json b/Engine/Source/Programs/Horde/HordeServer/Defaults/default-metrics.telemetry.json
--- a/Engine/Source/Programs/Horde/HordeServer/Defaults/default-metrics.telemetry.json
+++ b/Engine/Source/Programs/Horde/HordeServer/Defaults/default-metrics.telemetry.json
@@ -530,6 +530,42 @@
 			"interval": "00:15:00",
 			"topN": 25
 		},
+		{
+			"id": "core-content-totalsize",
+			"filter": "$[?(@.Payload.EventName == 'Core.VirtualAssets.MountPoint')]",
+			"property": "$.Payload.TotalFileSize",
+			"groupBy": "$.Payload.Build_BranchName,$.Payload.ProjectName",
+			"function": "Sum",
+			"interval": "00:15:00",
+			"topN": 25
+		},
+		{
+			"id": "core-content-totalsize-mp",
+			"filter": "$[?(@.Payload.EventName == 'Core.VirtualAssets.MountPoint')]",
+			"property": "$.Payload.TotalFileSize",
+			"groupBy": "$.Payload.Build_BranchName,$.Payload.ProjectName,$.Payload.Name",
+			"function": "Sum",
+			"interval": "00:15:00",
+			"topN": 25
+		},
+		{
+			"id": "core-content-virtualization",
+			"filter": "$[?(@.Payload.EventName == 'Core.VirtualAssets.MountPoint')]",
+			"property": "$.Payload.VirtualizedPercent",
+			"groupBy": "$.Payload.Build_BranchName,$.Payload.ProjectName",
+			"function": "Average",
+			"interval": "00:15:00",
+			"bottomN": 25
+		},
+		{
+			"id": "core-content-virtualization-mp",
+			"filter": "$[?(@.Payload.EventName == 'Core.VirtualAssets.MountPoint')]",
+			"property": "$.Payload.VirtualizedPercent",
+			"groupBy": "$.Payload.Build_BranchName,$.Payload.ProjectName,$.Payload.Name",
+			"function": "Average",
+			"interval": "00:15:00",
+			"bottomN": 25
+		},
 		{
 			"id": "horde-incremental-build-jobs",
 			"filter": "$[?(@.Payload.EventName == 'State.JobStepRef' && @.Payload.JobName=='Incremental Build' && search(@.Payload.StepName, \"Incremental\") && @.Payload.Duration>0 && @.Payload.State=='Completed')]",
diff --git a/Engine/Source/Programs/Horde/HordeServer/Defaults/default-views.telemetry.json b/Engine/Source/Programs/Horde/HordeServer/Defaults/default-views.telemetry.json
--- a/Engine/Source/Programs/Horde/HordeServer/Defaults/default-views.telemetry.json
+++ b/Engine/Source/Programs/Horde/HordeServer/Defaults/default-views.telemetry.json
@@ -528,6 +528,47 @@
 						}
 					]
 				},
+				{
+					"name": "Content",
+					"charts": [
+						{
+							"name": "Content Size Total",
+							"display": "Value",
+							"metrics": [
+								{
+									"id": "core-content-totalsize"
+								}
+							]
+						},
+						{
+							"name": "Virtualization",
+							"display": "Value",
+							"metrics": [
+								{
+									"id": "core-content-virtualization"
+								}
+							]
+						},
+						{
+							"name": "Content Size By Mount Point",
+							"display": "Value",
+							"metrics": [
+								{
+									"id": "core-content-totalsize-mp"
+								}
+							]
+						},
+						{
+							"name": "Virtualization By Mount Point",
+							"display": "Value",
+							"metrics": [
+								{
+									"id": "core-content-virtualization-mp"
+								}
+							]
+						}
+					]
+				},
 				{
 					"name": "Compile",
 					"charts": [
