Horde: Make OIDC scopes and mappings overridable

- Due to .NET config only merging array entries when combining config sources, hardcoded OIDC defaults (scopes and claim mappings) cannot be overridden.
- Added OidcAddDefaultScopesAndMappings flag to control whether defaults are applied

diff --git a/Engine/Source/Programs/Horde/HordeServer/ServerSettings.cs b/Engine/Source/Programs/Horde/HordeServer/ServerSettings.cs
--- a/Engine/Source/Programs/Horde/HordeServer/ServerSettings.cs
+++ b/Engine/Source/Programs/Horde/HordeServer/ServerSettings.cs
@@ -1,6 +1,7 @@
 // Copyright Epic Games, Inc. All Rights Reserved.
 
 using System;
+using System.Collections.Generic;
 using System.Linq;
 using System.Net;
 using EpicGames.Horde.Server;
@@ -191,27 +192,60 @@ namespace HordeServer
 		/// <summary>
 		/// OpenID Connect scopes to request when signing in
 		/// </summary>
-		public string[] OidcRequestedScopes { get; set; } = { "profile", "email", "openid" };
+		public string[] OidcRequestedScopes { get; set; } = [];
 
 		/// <summary>
 		/// List of fields in /userinfo endpoint to try map to the standard name claim (see System.Security.Claims.ClaimTypes.Name)
 		/// </summary>
-		public string[] OidcClaimNameMapping { get; set; } = { "preferred_username", "email" };
+		public string[] OidcClaimNameMapping { get; set; } = [];
 
 		/// <summary>
 		/// List of fields in /userinfo endpoint to try map to the standard email claim (see System.Security.Claims.ClaimTypes.Email)
 		/// </summary>
-		public string[] OidcClaimEmailMapping { get; set; } = { "email" };
+		public string[] OidcClaimEmailMapping { get; set; } = [];
 
 		/// <summary>
 		/// List of fields in /userinfo endpoint to try map to the Horde user claim (see HordeClaimTypes.User)
 		/// </summary>
-		public string[] OidcClaimHordeUserMapping { get; set; } = { "preferred_username", "email" };
+		public string[] OidcClaimHordeUserMapping { get; set; } = [];
 
 		/// <summary>
 		/// List of fields in /userinfo endpoint to try map to the Horde Perforce user claim (see HordeClaimTypes.PerforceUser)
 		/// </summary>
-		public string[] OidcClaimHordePerforceUserMapping { get; set; } = { "preferred_username", "email" };
+		public string[] OidcClaimHordePerforceUserMapping { get; set; } = [];
+		
+		/// <summary>
+		/// Add common scopes and mappings to above OIDC config fields
+		/// Provided as a workaround since .NET config will only *merge* array entries when combining multiple config sources.
+		/// Due to this unwanted behavior, having hard-coded defaults makes such fields unchangeable.
+		/// See [REDACTED_PROFILE_URL]/runtime/issues/36569 
+		/// <see cref="OidcRequestedScopes" />
+		/// <see cref="OidcClaimNameMapping" />
+		/// <see cref="OidcClaimEmailMapping" />
+		/// <see cref="OidcClaimHordeUserMapping" />
+		/// <see cref="OidcClaimHordePerforceUserMapping" />
+		/// </summary>
+		public bool OidcAddDefaultScopesAndMappings { get; set; } = true;
+		
+		/// <summary>
+		/// Adds default scopes and claim mappings to the OIDC configuration if enabled.
+		/// This method should be called after the initial configuration is loaded.
+		/// </summary>
+		public void AddDefaultOidcScopesAndMappings()
+		{
+			IReadOnlyList<string> defaultScopes = ["profile", "email", "openid"];
+			IReadOnlyList<string> defaultMappings = ["preferred_username", "email"];
+			IReadOnlyList<string> defaultEmailMapping = ["email"];
+			
+			if (OidcAddDefaultScopesAndMappings)
+			{
+				OidcRequestedScopes = OidcRequestedScopes.Concat(defaultScopes).ToArray();
+				OidcClaimNameMapping = OidcClaimNameMapping.Concat(defaultMappings).ToArray();
+				OidcClaimEmailMapping = OidcClaimEmailMapping.Concat(defaultEmailMapping).ToArray();
+				OidcClaimHordeUserMapping = OidcClaimHordeUserMapping.Concat(defaultMappings).ToArray();
+				OidcClaimHordePerforceUserMapping = OidcClaimHordePerforceUserMapping.Concat(defaultMappings).ToArray();
+			}
+		}
 
 		/// <summary>
 		/// Base URL this Horde server is accessible from
diff --git a/Engine/Source/Programs/Horde/HordeServer/Startup.cs b/Engine/Source/Programs/Horde/HordeServer/Startup.cs
--- a/Engine/Source/Programs/Horde/HordeServer/Startup.cs
+++ b/Engine/Source/Programs/Horde/HordeServer/Startup.cs
@@ -219,6 +219,7 @@ namespace HordeServer
 			IConfigurationSection hordeSection = configuration.GetSection("Horde");
 			hordeSection.Bind(settings);
 			settings.Validate();
+			settings.AddDefaultOidcScopesAndMappings();
 		}
 
 		/// <summary>
