Tweak to content metrics

diff --git a/Engine/Source/Programs/Horde/HordeServer/Defaults/default-metrics.telemetry.json b/Engine/Source/Programs/Horde/HordeServer/Defaults/default-metrics.telemetry.json
--- a/Engine/Source/Programs/Horde/HordeServer/Defaults/default-metrics.telemetry.json
+++ b/Engine/Source/Programs/Horde/HordeServer/Defaults/default-metrics.telemetry.json
@@ -534,7 +534,7 @@
 			"id": "core-content-totalsize",
 			"filter": "$[?(@.Payload.EventName == 'Core.VirtualAssets.MountPoint')]",
 			"property": "$.Payload.TotalFileSize",
-			"groupBy": "$.Payload.Build_BranchName,$.Payload.ProjectName",
+			"groupBy": "$.Payload.Build_BranchName",
 			"function": "Sum",
 			"interval": "00:15:00",
 			"topN": 25
@@ -543,7 +543,7 @@
 			"id": "core-content-totalsize-mp",
 			"filter": "$[?(@.Payload.EventName == 'Core.VirtualAssets.MountPoint')]",
 			"property": "$.Payload.TotalFileSize",
-			"groupBy": "$.Payload.Build_BranchName,$.Payload.ProjectName,$.Payload.Name",
+			"groupBy": "$.Payload.Build_BranchName,$.Payload.Name",
 			"function": "Sum",
 			"interval": "00:15:00",
 			"topN": 25
@@ -552,7 +552,7 @@
 			"id": "core-content-virtualization",
 			"filter": "$[?(@.Payload.EventName == 'Core.VirtualAssets.MountPoint')]",
 			"property": "$.Payload.VirtualizedPercent",
-			"groupBy": "$.Payload.Build_BranchName,$.Payload.ProjectName",
+			"groupBy": "$.Payload.Build_BranchName",
 			"function": "Average",
 			"interval": "00:15:00",
 			"bottomN": 25
@@ -561,7 +561,7 @@
 			"id": "core-content-virtualization-mp",
 			"filter": "$[?(@.Payload.EventName == 'Core.VirtualAssets.MountPoint')]",
 			"property": "$.Payload.VirtualizedPercent",
-			"groupBy": "$.Payload.Build_BranchName,$.Payload.ProjectName,$.Payload.Name",
+			"groupBy": "$.Payload.Build_BranchName,$.Payload.Name",
 			"function": "Average",
 			"interval": "00:15:00",
 			"bottomN": 25
