Implementation of the Gauntlet Slack Notification redesign as an experimental Horde plugin

#rnx

diff --git a/Engine/Source/Programs/Horde/Plugins/Build/HordeServer.Build.Tests/Notifications/NotificationServiceTests.cs b/Engine/Source/Programs/Horde/Plugins/Build/HordeServer.Build.Tests/Notifications/NotificationServiceTests.cs
--- a/Engine/Source/Programs/Horde/Plugins/Build/HordeServer.Build.Tests/Notifications/NotificationServiceTests.cs
+++ b/Engine/Source/Programs/Horde/Plugins/Build/HordeServer.Build.Tests/Notifications/NotificationServiceTests.cs
@@ -1,26 +1,60 @@
 // Copyright Epic Games, Inc. All Rights Reserved.
 
+using EpicGames.Core;
+using EpicGames.Horde.Agents.Leases;
+using EpicGames.Horde.Agents.Pools;
+using EpicGames.Horde.Agents;
+using EpicGames.Horde.Agents.Sessions;
 using EpicGames.Horde.Commits;
 using EpicGames.Horde.Jobs;
 using EpicGames.Horde.Jobs.Graphs;
+using EpicGames.Horde.Jobs.Templates;
+using EpicGames.Horde.Logs;
+using EpicGames.Horde.Projects;
 using EpicGames.Horde.Streams;
+using EpicGames.Horde.Users;
 using HordeServer.Agents;
 using HordeServer.Agents.Pools;
 using HordeServer.Configuration;
 using HordeServer.Devices;
 using HordeServer.Issues;
+using HordeServer.Jobs;
 using HordeServer.Logs;
 using HordeServer.Notifications;
+using HordeServer.Projects;
 using HordeServer.Streams;
 using HordeServer.Users;
+using HordeServer.Utilities;
 using Microsoft.Extensions.Caching.Memory;
 using Microsoft.Extensions.DependencyInjection;
 using Moq;
 
 namespace HordeServer.Tests.Notifications
 {
+	public class JobNotificationInformation
+	{
+		public IJob Job { get; }
+		public IJobStepBatch? Batch { get; } = null;
+		public IJobStep? Step { get; } = null;
+		public INode? Node { get; } = null;
+		public List<ILogEventData>? JobStepEventData { get; } = null;
+		public IEnumerable<IUser>? UsersToNotify { get; } = null;
+		public JobNotificationInformation(IJob job, IJobStepBatch? batch = null, IJobStep? step = null, INode? node = null, List<ILogEventData>? jobStepEventData = null, IEnumerable<IUser>? usersToNotify = null)
+		{
+			Job = job;
+			Batch = batch;
+			Step = step;
+			Node = node;
+			JobStepEventData = jobStepEventData;
+			UsersToNotify = usersToNotify;
+		}
+	}
+
 	public class FakeNotificationSink : INotificationSink
 	{
+		public JobNotificationInformation? JobNotification { get; set; } = null;
+		public List<JobNotificationInformation> JobStepNotifications { get; } = new();
+
 		public List<JobScheduledNotification> JobScheduledNotifications { get; } = new();
 		public int JobScheduledCallCount { get; set; }
 
@@ -31,9 +65,27 @@ namespace HordeServer.Tests.Notifications
 			return Task.CompletedTask;
 		}
 
-		public Task NotifyJobCompleteAsync(IJob job, IGraph graph, LabelOutcome outcome, CancellationToken cancellationToken) { throw new NotImplementedException(); }
+		public Task NotifyJobStepCompleteAsync(IEnumerable<IUser>? usersToNotify, IJob job, IJobStepBatch batch, IJobStep step, INode node, List<ILogEventData> jobStepEventData, CancellationToken cancellationToken)
+		{
+			JobNotificationInformation jobNotificationInformation = new JobNotificationInformation(job, batch, step, node, jobStepEventData, usersToNotify);
+			JobStepNotifications.Add(jobNotificationInformation);
+			return Task.CompletedTask;
+		}
+
+		public Task NotifyJobStepAbortedAsync(IEnumerable<IUser>? usersToNotify, IJob job, IJobStepBatch batch, IJobStep step, INode node, List<ILogEventData> jobStepEventData, CancellationToken cancellationToken)
+		{
+			JobNotificationInformation jobNotificationInformation = new JobNotificationInformation(job, batch, step, node, jobStepEventData, usersToNotify);
+			JobStepNotifications.Add(jobNotificationInformation);
+			return Task.CompletedTask;
+		}
+
+		public Task NotifyJobCompleteAsync(IJob job, IGraph graph, LabelOutcome outcome, CancellationToken cancellationToken)
+		{
+			JobNotification = new JobNotificationInformation(job);
+			return Task.CompletedTask;
+		}
+
 		public Task NotifyJobCompleteAsync(IUser user, IJob job, IGraph graph, LabelOutcome outcome, CancellationToken cancellationToken) { throw new NotImplementedException(); }
-		public Task NotifyJobStepCompleteAsync(IEnumerable<IUser> usersToNotify, IJob job, IJobStepBatch batch, IJobStep step, INode node, List<ILogEventData> jobStepEventData, CancellationToken cancellationToken) { throw new NotImplementedException(); }
 		public Task NotifyLabelCompleteAsync(IUser user, IJob job, ILabel label, int labelIdx, LabelOutcome outcome, List<(string, JobStepOutcome, Uri)> stepData, CancellationToken cancellationToken) { throw new NotImplementedException(); }
 		public Task NotifyIssueUpdatedAsync(IIssue issue, CancellationToken cancellationToken) { throw new NotImplementedException(); }
 		public Task NotifyConfigUpdateAsync(ConfigUpdateInfo info, CancellationToken cancellationToken) => Task.CompletedTask;
@@ -47,6 +99,15 @@ namespace HordeServer.Tests.Notifications
 	[TestClass]
 	public class NotificationServiceTests : BuildTestSetup
 	{
+		readonly StreamId _streamId = new("ue4-main");
+		readonly StreamConfig _streamConfig;
+		readonly TemplateId _templateId1 = new("template1");
+
+		public NotificationServiceTests()
+		{
+			_streamConfig = new StreamConfig { Id = _streamId };
+		}
+
 		protected override void ConfigureServices(IServiceCollection services)
 		{
 			base.ConfigureServices(services);
@@ -94,6 +155,155 @@ namespace HordeServer.Tests.Notifications
 			return job.Object;
 		}
 
+		static NewGroup AddGroup(List<NewGroup> groups)
+		{
+			NewGroup group = new NewGroup("win64", new List<NewNode>());
+			groups.Add(group);
+			return group;
+		}
+
+		static NewNode AddNode(NewGroup group, string name, string[]? inputDependencies, Action<NewNode>? action = null)
+		{
+			NewNode node = new NewNode(name, inputDependencies: inputDependencies?.ToList(), orderDependencies: inputDependencies?.ToList());
+			action?.Invoke(node);
+			group.Nodes.Add(node);
+			return node;
+		}
+
+		void UpdateGlobalConfig()
+		{
+			UpdateConfig(globalConfig =>
+			{
+				ProjectConfig projectConfig = new() { Id = new ProjectId("ue4") };
+				projectConfig.Streams.Add(new StreamConfig { Id = _streamId });
+				globalConfig.Plugins.GetBuildConfig().Projects.Add(projectConfig);
+			});
+		}
+
+		async Task<IJob> CreateJobAsync()
+		{
+			Mock<ITemplate> templateMock = new Mock<ITemplate>(MockBehavior.Strict);
+			templateMock.SetupGet(x => x.InitialAgentType).Returns((string?)null);
+
+			IGraph baseGraph = await GraphCollection.AddAsync(templateMock.Object);
+
+			CreateJobOptions options = new CreateJobOptions();
+			options.Arguments.Add("-Target=Publish Client");
+			options.Arguments.Add("-Target=Post-Publish Client");
+
+			List<NewGroup> newGroups = new List<NewGroup>();
+
+			NewGroup initialGroup = AddGroup(newGroups);
+			AddNode(initialGroup, "Update Version Files", null);
+			AddNode(initialGroup, "Compile Editor", new[] { "Update Version Files" });
+
+			NewGroup compileGroup = AddGroup(newGroups);
+			AddNode(compileGroup, "Compile Client", new[] { "Update Version Files" });
+
+			NewGroup publishGroup = AddGroup(newGroups);
+			AddNode(publishGroup, "Cook Client", new[] { "Compile Editor" }, x => x.RunEarly = true);
+			AddNode(publishGroup, "Publish Client", new[] { "Compile Client", "Cook Client" });
+			AddNode(publishGroup, "Post-Publish Client", null, x => x.OrderDependencies = new List<string> { "Publish Client" });
+
+			IGraph graph = await GraphCollection.AppendAsync(baseGraph, newGroups);
+
+			ITemplate template1 = new TemplateStub(ContentHash.MD5("graphHash"), "template1");
+			return await JobService.CreateJobAsync(null, _streamConfig!, _templateId1, template1.Hash, graph, "Test job", CommitIdWithOrder.FromPerforceChange(123), CommitIdWithOrder.FromPerforceChange(123), options);
+		}
+
+		async Task<IJob> AbortJobAsync(IJob job)
+		{
+			job = Deref(await JobService.UpdateJobAsync(job, abortedByUserId: UserId.Anonymous));
+			Assert.AreEqual(JobState.Complete, job.GetState());
+			return job;
+		}
+
+		Task<IJob> RunBatchAsync(IJob job, int batchIdx)
+		{
+			return UpdateBatchAsync(job, batchIdx, JobStepBatchState.Running);
+		}
+
+		async Task<IJob> UpdateBatchAsync(IJob job, int batchIdx, JobStepBatchState state)
+		{
+			Assert.AreEqual(JobStepBatchState.Ready, job.Batches[batchIdx].State);
+
+			LogId logId = LogIdUtils.GenerateNewId();
+			job = Deref(await job.TryAssignLeaseAsync(batchIdx, new PoolId("foo"), new AgentId("agent"), new SessionId(BinaryIdUtils.CreateNew()), new LeaseId(BinaryIdUtils.CreateNew()), logId));
+			if (job.Batches[batchIdx].State != state)
+			{
+				job = Deref(await JobService.UpdateBatchAsync(job, job.Batches[batchIdx].Id, _streamConfig, logId, state));
+				Assert.AreEqual(state, job.Batches[batchIdx].State);
+			}
+			return job;
+		}
+
+		async Task<IJob> AbortStepAsync(IJob job, int batchIdx, int stepIdx)
+		{
+			LogId logId = LogIdUtils.GenerateNewId();
+			ILog log = await LogCollection.AddAsync(job.Id, job.Batches[batchIdx].LeaseId, job.Batches[batchIdx].SessionId, LogType.Json, logId);
+			job = Deref(await JobService.UpdateStepAsync(job, job.Batches[batchIdx].Id, job.Batches[batchIdx].Steps[stepIdx].Id, _streamConfig, newAbortRequested: true, newAbortByUserId: UserId.Anonymous, newLogId: logId));
+			Assert.AreEqual(JobStepState.Aborted, job.Batches[batchIdx].Steps[stepIdx].State);
+			return job;
+		}
+
+		async Task<IJob> RunStepAsync(IJob job, int batchIdx, int stepIdx, JobStepOutcome outcome)
+		{
+			Assert.AreEqual(JobStepState.Ready, job.Batches[batchIdx].Steps[stepIdx].State);
+
+			LogId logId = LogIdUtils.GenerateNewId();
+			ILog log = await LogCollection.AddAsync(job.Id, job.Batches[batchIdx].LeaseId, job.Batches[batchIdx].SessionId, LogType.Json, logId);
+			job = Deref(await JobService.UpdateStepAsync(job, job.Batches[batchIdx].Id, job.Batches[batchIdx].Steps[stepIdx].Id, _streamConfig, JobStepState.Running, JobStepOutcome.Success, newLogId: logId));
+			Assert.AreEqual(JobStepState.Running, job.Batches[batchIdx].Steps[stepIdx].State);
+			job = Deref(await JobService.UpdateStepAsync(job, job.Batches[batchIdx].Id, job.Batches[batchIdx].Steps[stepIdx].Id, _streamConfig, JobStepState.Completed, outcome, newLogId: logId));
+			Assert.AreEqual(JobStepState.Completed, job.Batches[batchIdx].Steps[stepIdx].State);
+			Assert.AreEqual(outcome, job.Batches[batchIdx].Steps[stepIdx].Outcome);
+			return job;
+		}
+
+		static async Task ValidateJobNotificationAsync(FakeNotificationSink sink, IJob job, JobState state, UserId? abortedByUserId = null, int timeoutMs = 30000)
+		{
+			DateTime endTime = DateTime.Now + TimeSpan.FromMilliseconds(timeoutMs);
+
+			while (endTime >= DateTime.Now)
+			{
+				if (sink.JobNotification is not null)
+				{
+					Assert.IsNotNull(sink.JobNotification);
+					Assert.AreEqual(job.Id, sink.JobNotification.Job.Id);
+					Assert.AreEqual(state, sink.JobNotification.Job.GetState());
+					Assert.AreEqual(abortedByUserId, sink.JobNotification.Job.AbortedByUserId);
+					sink.JobNotification = null;
+					return;
+				}
+				await Task.Delay(100);
+			}
+			Assert.Fail("Unable to validate job notification for job '{JobId}'");
+		}
+
+		static async Task ValidateJobStepNotificationAsync(FakeNotificationSink sink, IJob job, int batchIdx, int stepIdx, JobStepState state, JobStepOutcome outcome, int timeoutMs = 30000)
+		{
+			DateTime endTime = DateTime.Now + TimeSpan.FromMilliseconds(timeoutMs);
+
+			while (endTime >= DateTime.Now)
+			{
+				if (sink.JobStepNotifications.Any())
+				{
+					JobNotificationInformation jobNotificationInformation = sink.JobStepNotifications[0];
+					Assert.AreEqual(job.Id, jobNotificationInformation.Job.Id);
+					Assert.IsNotNull(jobNotificationInformation.Step);
+
+					IJobStep step = job.Batches[batchIdx].Steps[stepIdx];
+					Assert.AreEqual(step.Id, jobNotificationInformation.Step.Id);
+					Assert.AreEqual(state, jobNotificationInformation.Step.State);
+					Assert.AreEqual(outcome, jobNotificationInformation.Step.Outcome);
+					sink.JobStepNotifications.RemoveAt(0);
+					return;
+				}
+				await Task.Delay(100);
+			}
+			Assert.Fail("Unable to validate job step notification for job '{JobId}'");
+		}
+
 		[TestMethod]
 		public async Task NotifyJobScheduledAsync()
 		{
@@ -148,6 +358,68 @@ namespace HordeServer.Tests.Notifications
 			Assert.AreEqual(2, fakeSink.JobScheduledNotifications.Count);
 		}
 
+		[TestMethod]
+		public async Task NotifyJobCompleteAsync()
+		{
+			UpdateGlobalConfig();
+
+			FakeNotificationSink fakeSink = ServiceProvider.GetRequiredService<FakeNotificationSink>();
+
+			NotificationService service = (NotificationService)ServiceProvider.GetRequiredService<INotificationService>();
+			await service._ticker.StartAsync();
+
+			IJob job = await CreateJobAsync();
+			job = await RunBatchAsync(job, 0);
+			job = await RunStepAsync(job, 0, 0, JobStepOutcome.Success); // Setup Build
+			await ValidateJobStepNotificationAsync(fakeSink, job, 0, 0, JobStepState.Completed, JobStepOutcome.Success);
+
+			// Trigger a cancellation of the 'Cook Client' step
+			job = await AbortStepAsync(job, 3, 0);
+			await ValidateJobStepNotificationAsync(fakeSink, job, 3, 0, JobStepState.Aborted, JobStepOutcome.Failure);
+
+			job = await RunBatchAsync(job, 1);
+			job = await RunStepAsync(job, 1, 0, JobStepOutcome.Success); // Update Version Files
+			await ValidateJobStepNotificationAsync(fakeSink, job, 1, 0, JobStepState.Completed, JobStepOutcome.Success);
+
+			job = await RunStepAsync(job, 1, 1, JobStepOutcome.Success); // Compile Editor
+			await ValidateJobStepNotificationAsync(fakeSink, job, 1, 1, JobStepState.Completed, JobStepOutcome.Success);
+
+			job = await RunBatchAsync(job, 2);
+			job = await RunStepAsync(job, 2, 0, JobStepOutcome.Failure); // Compile Client
+			await ValidateJobStepNotificationAsync(fakeSink, job, 2, 0, JobStepState.Completed, JobStepOutcome.Failure);
+
+			// Validate that the states did not change after running the other steps
+			Assert.AreEqual(JobStepBatchState.Complete, job.Batches[3].State);
+			Assert.AreEqual(JobStepState.Aborted, job.Batches[3].Steps[0].State); // Cook Client
+			Assert.AreEqual(JobStepState.Skipped, job.Batches[3].Steps[1].State); // Publish Client
+			Assert.AreEqual(JobStepState.Skipped, job.Batches[3].Steps[2].State); // Post-Publish Client
+
+			// Validate that we received our job completion notification
+			await ValidateJobNotificationAsync(fakeSink, job, JobState.Complete);
+		}
+
+		[TestMethod]
+		public async Task NotifyJobAbortedAsync()
+		{
+			UpdateGlobalConfig();
+
+			FakeNotificationSink fakeSink = ServiceProvider.GetRequiredService<FakeNotificationSink>();
+
+			NotificationService service = (NotificationService)ServiceProvider.GetRequiredService<INotificationService>();
+			await service._ticker.StartAsync();
+
+			IJob job = await CreateJobAsync();
+			job = await RunBatchAsync(job, 0);
+			job = await RunStepAsync(job, 0, 0, JobStepOutcome.Success); // Setup Build
+			await ValidateJobStepNotificationAsync(fakeSink, job, 0, 0, JobStepState.Completed, JobStepOutcome.Success);
+
+			// Trigger a cancellation of the job
+			job = await AbortJobAsync(job);
+
+			// Validate that we received our job completion notification
+			await ValidateJobNotificationAsync(fakeSink, job, JobState.Complete, UserId.Anonymous);
+		}
+
 		//public void NotifyLabelUpdate(IJob Job, IReadOnlyList<(LabelState, LabelOutcome)> OldLabelStates, IReadOnlyList<(LabelState, LabelOutcome)> NewLabelStates)
 		//{
 		//	// If job has any label trigger IDs, send label complete notifications if needed
diff --git a/Engine/Source/Programs/Horde/Plugins/Build/HordeServer.Build/Jobs/JobService.cs b/Engine/Source/Programs/Horde/Plugins/Build/HordeServer.Build/Jobs/JobService.cs
--- a/Engine/Source/Programs/Horde/Plugins/Build/HordeServer.Build/Jobs/JobService.cs
+++ b/Engine/Source/Programs/Horde/Plugins/Build/HordeServer.Build/Jobs/JobService.cs
@@ -96,6 +96,16 @@ namespace HordeServer.Jobs
 		/// </summary>
 		public event LabelUpdateEvent? OnLabelUpdate;
 
+		/// <summary>
+		/// Delegate for job aborted events
+		/// </summary>
+		public delegate void JobAbortedEvent(IJob job, IGraph graph);
+
+		/// <summary>
+		/// Event triggered when a job has been aborted
+		/// </summary>
+		public event JobAbortedEvent? OnJobAborted;
+
 		/// <summary>
 		/// Delegate for job step complete events
 		/// </summary>
@@ -370,6 +380,14 @@ namespace HordeServer.Jobs
 							await CancelLeaseAsync(batch.AgentId.Value, batch.LeaseId.Value, cancellationToken);
 						}
 					}
+
+					// Notify subscribers that the job has been canceled
+					if (newJob.GetState() == JobState.Complete && abortedByUserId is not null)
+					{
+						using TelemetrySpan _ = _tracer.StartActiveSpan($"{nameof(JobService)}.{nameof(UpdateJobAsync)}.NotifySubscribers");
+						OnJobAborted?.Invoke(newJob, graph);
+					}
+
 					return newJob;
 				}
 			}
@@ -939,9 +957,19 @@ namespace HordeServer.Jobs
 				updateSpan.SetAttribute("BatchId", batchId.ToString());
 				updateSpan.SetAttribute("StepId", stepId.ToString());
 
-				if (oldState != newState || oldOutcome != newOutcome)
+				// Get the current state after updating the step
+				// To minimize impact we only want to perform this if newAbortRequested has been set
+				JobStepState updatedState = newState;
+				JobStepOutcome updatedOutcome = newOutcome;
+				IJobStep? step;
+				if (newAbortRequested.HasValue && newAbortRequested.Value && job.TryGetStep(batchId, stepId, out step))
+				{
+					updatedState = step.State;
+					updatedOutcome = step.Outcome;
+				}
+				if (oldState != updatedState || oldOutcome != updatedOutcome)
 				{
-					_logger.LogInformation("Transitioned job {JobId}, batch {BatchId}, step {StepId} from {OldState} to {NewState}", job.Id, batchId, stepId, oldState, newState);
+					_logger.LogInformation("Transitioned job {JobId}, batch {BatchId}, step {StepId} from {OldState} to {NewState}", job.Id, batchId, stepId, oldState, updatedState);
 
 					// Send any updates for modified badges
 					if (oldLabelStates != null)
@@ -953,7 +981,7 @@ namespace HordeServer.Jobs
 					}
 
 					// Submit the change if auto-submit is enabled
-					if (job.PreflightCommitId != null && newState == JobStepState.Completed)
+					if (job.PreflightCommitId != null && updatedState == JobStepState.Completed)
 					{
 						(JobStepState state, JobStepOutcome outcome) = job.GetTargetState();
 						if (state == JobStepState.Completed)
@@ -975,10 +1003,10 @@ namespace HordeServer.Jobs
 					}
 
 					// Update the jobstep ref if it completed
-					if (newState == JobStepState.Running || newState == JobStepState.Completed || newState == JobStepState.Aborted)
+					if (updatedState == JobStepState.Running || updatedState == JobStepState.Completed || updatedState == JobStepState.Aborted)
 					{
 						using TelemetrySpan _ = _tracer.StartActiveSpan($"{nameof(JobService)}.{nameof(TryUpdateStepAsync)}.UpdateJobStepRef");
-						if (job.TryGetBatch(batchId, out IJobStepBatch? batch) && batch.TryGetStep(stepId, out IJobStep? step) && step.StartTimeUtc != null)
+						if (job.TryGetBatch(batchId, out IJobStepBatch? batch) && batch.TryGetStep(stepId, out step) && step.StartTimeUtc != null)
 						{
 							await _jobStepRefs.UpdateAsync(job, batch, step, graph);
 							await _bisectTasks.UpdateAsync(job, batch, step, graph, cancellationToken: cancellationToken);
@@ -986,14 +1014,14 @@ namespace HordeServer.Jobs
 					}
 
 					// Notify subscribers
-					if (newState == JobStepState.Completed)
+					if (updatedState == JobStepState.Completed || updatedState == JobStepState.Aborted)
 					{
 						using TelemetrySpan _ = _tracer.StartActiveSpan($"{nameof(JobService)}.{nameof(TryUpdateStepAsync)}.NotifySubscribers");
 						OnJobStepComplete?.Invoke(job, graph, batchId, stepId);
 					}
 
 					// Create any downstream jobs
-					if (newState == JobStepState.Completed && newOutcome != JobStepOutcome.Failure)
+					if (updatedState == JobStepState.Completed && updatedOutcome != JobStepOutcome.Failure)
 					{
 						using TelemetrySpan _ = _tracer.StartActiveSpan($"{nameof(JobService)}.{nameof(TryUpdateStepAsync)}.CreateDownstreamJobs");
 						for (int idx = 0; idx < job.ChainedJobs.Count; idx++)
diff --git a/Engine/Source/Programs/Horde/Plugins/Build/HordeServer.Build/Notifications/INotificationService.cs b/Engine/Source/Programs/Horde/Plugins/Build/HordeServer.Build/Notifications/INotificationService.cs
--- a/Engine/Source/Programs/Horde/Plugins/Build/HordeServer.Build/Notifications/INotificationService.cs
+++ b/Engine/Source/Programs/Horde/Plugins/Build/HordeServer.Build/Notifications/INotificationService.cs
@@ -112,6 +112,13 @@ namespace HordeServer.Notifications
 		/// <returns>Subscriptions for that user</returns>
 		Task<INotificationSubscription?> GetSubscriptionsAsync(NotificationTriggerId triggerId, ClaimsPrincipal user, CancellationToken cancellationToken);
 
+		/// <summary>
+		/// Notify all subscribers that a job has been cancelled
+		/// </summary>
+		/// <param name="job">The job that has finished</param>
+		/// <param name="graph">Graph for the job</param>
+		void NotifyJobAborted(IJob job, IGraph graph);
+
 		/// <summary>
 		/// Notify all subscribers that a job step has finished
 		/// </summary>
diff --git a/Engine/Source/Programs/Horde/Plugins/Build/HordeServer.Build/Notifications/INotificationSink.cs b/Engine/Source/Programs/Horde/Plugins/Build/HordeServer.Build/Notifications/INotificationSink.cs
--- a/Engine/Source/Programs/Horde/Plugins/Build/HordeServer.Build/Notifications/INotificationSink.cs
+++ b/Engine/Source/Programs/Horde/Plugins/Build/HordeServer.Build/Notifications/INotificationSink.cs
@@ -48,10 +48,23 @@ namespace HordeServer.Notifications
 		/// <returns>Async task</returns>
 		Task NotifyJobCompleteAsync(IUser user, IJob job, IGraph graph, LabelOutcome outcome, CancellationToken cancellationToken);
 
+		/// <summary>
+		/// Send notifications that a job step has been canceled
+		/// </summary>
+		/// <param name="usersToNotify">Optional list of users to notify</param>
+		/// <param name="job">The job containing the step</param>
+		/// <param name="batch">Unique id of the batch</param>
+		/// <param name="step">The step id</param>
+		/// <param name="node">Corresponding node for the step</param>
+		/// <param name="jobStepEventData"></param>
+		/// <param name="cancellationToken">Cancellation token for the operation</param>
+		/// <returns>Async task</returns>
+		Task NotifyJobStepAbortedAsync(IEnumerable<IUser>? usersToNotify, IJob job, IJobStepBatch batch, IJobStep step, INode node, List<ILogEventData> jobStepEventData, CancellationToken cancellationToken);
+
 		/// <summary>
 		/// Send notifications that a job step has completed
 		/// </summary>
-		/// <param name="usersToNotify">Users to notify</param>
+		/// <param name="usersToNotify">Optional list of users to notify</param>
 		/// <param name="job">The job containing the step</param>
 		/// <param name="batch">Unique id of the batch</param>
 		/// <param name="step">The step id</param>
@@ -59,7 +72,7 @@ namespace HordeServer.Notifications
 		/// <param name="jobStepEventData"></param>
 		/// <param name="cancellationToken">Cancellation token for the operation</param>
 		/// <returns>Async task</returns>
-		Task NotifyJobStepCompleteAsync(IEnumerable<IUser> usersToNotify, IJob job, IJobStepBatch batch, IJobStep step, INode node, List<ILogEventData> jobStepEventData, CancellationToken cancellationToken);
+		Task NotifyJobStepCompleteAsync(IEnumerable<IUser>? usersToNotify, IJob job, IJobStepBatch batch, IJobStep step, INode node, List<ILogEventData> jobStepEventData, CancellationToken cancellationToken);
 
 		/// <summary>
 		/// Send notifications that a job step has completed
diff --git a/Engine/Source/Programs/Horde/Plugins/Build/HordeServer.Build/Notifications/NotificationService.cs b/Engine/Source/Programs/Horde/Plugins/Build/HordeServer.Build/Notifications/NotificationService.cs
--- a/Engine/Source/Programs/Horde/Plugins/Build/HordeServer.Build/Notifications/NotificationService.cs
+++ b/Engine/Source/Programs/Horde/Plugins/Build/HordeServer.Build/Notifications/NotificationService.cs
@@ -165,6 +165,7 @@ namespace HordeServer.Notifications
 			_backgroundTask = new BackgroundTask(ExecuteAsync);
 
 			issueService.OnIssueUpdated += NotifyIssueUpdated;
+			jobService.OnJobAborted += NotifyJobAborted;
 			jobService.OnJobStepComplete += NotifyJobStepComplete;
 			jobService.OnJobScheduled += NotifyJobScheduled;
 			jobService.OnLabelUpdate += NotifyLabelUpdate;
@@ -193,6 +194,7 @@ namespace HordeServer.Notifications
 		public async ValueTask DisposeAsync()
 		{
 			_issueService.OnIssueUpdated -= NotifyIssueUpdated;
+			_jobService.OnJobAborted -= NotifyJobAborted;
 			_jobService.OnJobStepComplete -= NotifyJobStepComplete;
 			_jobService.OnJobScheduled += NotifyJobScheduled;
 			_jobService.OnLabelUpdate -= NotifyLabelUpdate;
@@ -234,6 +236,18 @@ namespace HordeServer.Notifications
 			return trigger.Subscriptions.FirstOrDefault(x => x.UserId == userId.Value);
 		}
 
+		/// <inheritdoc/>
+		public void NotifyJobAborted(IJob job, IGraph graph)
+		{
+			// Enqueue job complete notifications if needed
+			if (job.GetState() == JobState.Complete && job.AbortedByUserId is not null)
+			{
+				_logger.LogInformation("Queuing job notifications for {JobId}", job.Id);
+				EnqueueTask(ctx => SendJobNotificationsAsync(job, graph, ctx));
+				EnqueueTask(ctx => RecordJobCompleteMetricsAsync(job, ctx));
+			}
+		}
+
 		/// <inheritdoc/>
 		public void NotifyJobStepComplete(IJob job, IGraph graph, JobStepBatchId batchId, JobStepId stepId)
 		{
@@ -662,23 +676,96 @@ namespace HordeServer.Notifications
 				}
 			}
 
-			if (usersToNotify.Count == 0)
+			List<IUser> filteredUsersToNotify = new List<IUser>();
+			foreach (IUser slackUser in usersToNotify)
 			{
-				_logger.LogInformation("No users to notify for step {JobId}:{BatchId}:{StepId}", job.Id, batchId, stepId);
-				return;
+				if (job.PreflightCommitId != null)
+				{
+					if (slackUser.Id != job.StartedByUserId)
+					{
+						continue;
+					}
+				}
+				filteredUsersToNotify.Add(slackUser);
 			}
 
-			if (step.LogId == null)
+			// Fetch our event data which is mandatory for a completed step, but optional for canceled steps
+			ILog? log = null;
+			List<ILogEventData> jobStepEventData = new List<ILogEventData>();
+			if (step.LogId is not null)
 			{
-				_logger.LogError("Step does not have a log file");
-				return;
+				log = await _logCollection.GetAsync(step.LogId.Value, cancellationToken);
+				if (log is not null)
+				{
+					List<ILogAnchor> jobStepEvents = await log.GetAnchorsAsync(cancellationToken: cancellationToken);
+					foreach (ILogAnchor logEvent in jobStepEvents)
+					{
+						ILogEventData eventData = await logEvent.GetDataAsync(cancellationToken);
+						jobStepEventData.Add(eventData);
+					}
+				}
+			}
+
+			// Treat a canceled job step separately from a completed step as we may have different logic for each
+			if (step.State == JobStepState.Aborted)
+			{
+				if (filteredUsersToNotify.Count == 0)
+				{
+					_logger.LogInformation("No users to notify for step {JobId}:{BatchId}:{StepId}", job.Id, batchId, stepId);
+				}
+				else
+				{
+					EnqueueTasks((sink, ctx) => sink.NotifyJobStepAbortedAsync(filteredUsersToNotify, job, batch, step, node, jobStepEventData, ctx));
+				}
+
+				if (job.PreflightCommitId is null)
+				{
+					EnqueueTasks((sink, ctx) => sink.NotifyJobStepAbortedAsync(null, job, batch, step, node, jobStepEventData, ctx));
+				}
+			}
+			else
+			{
+				if (step.LogId == null)
+				{
+					_logger.LogError("Step does not have a log file");
+					return;
+				}
+
+				if (log == null)
+				{
+					_logger.LogError("Step does not have a log file");
+					return;
+				}
+
+				if (filteredUsersToNotify.Count == 0)
+				{
+					_logger.LogInformation("No users to notify for step {JobId}:{BatchId}:{StepId}", job.Id, batchId, stepId);
+				}
+				else
+				{
+					EnqueueTasks((sink, ctx) => sink.NotifyJobStepCompleteAsync(filteredUsersToNotify, job, batch, step, node, jobStepEventData, ctx));
+				}
+
+				if (job.PreflightCommitId is null)
+				{
+					EnqueueTasks((sink, ctx) => sink.NotifyJobStepCompleteAsync(null, job, batch, step, node, jobStepEventData, ctx));
+				}
+			}
+
+			_logger.LogDebug("Finished sending notifications for step {JobId}:{BatchId}:{StepId}", job.Id, batchId, stepId);
+		}
+
+		private async Task<List<ILogEventData>?> FetchLogEventsAsync(IJobStep step, CancellationToken cancellationToken)
+		{
+			if (step.LogId is null)
+			{
+				return null;
 			}
 
 			ILog? log = await _logCollection.GetAsync(step.LogId.Value, cancellationToken);
-			if (log == null)
+			if (log is null)
 			{
-				_logger.LogError("Step does not have a log file");
-				return;
+				return null;
 			}
 
 			List<ILogAnchor> jobStepEvents = await log.GetAnchorsAsync(cancellationToken: cancellationToken);
@@ -689,22 +776,7 @@ namespace HordeServer.Notifications
 				jobStepEventData.Add(eventData);
 			}
 
-			List<IUser> filteredUsersToNotify = new List<IUser>();
-			foreach (IUser slackUser in usersToNotify)
-			{
-				if (job.PreflightCommitId != null)
-				{
-					if (slackUser.Id != job.StartedByUserId)
-					{
-						continue;
-					}
-				}
-				filteredUsersToNotify.Add(slackUser);
-			}
-
-			EnqueueTasks((sink, ctx) => sink.NotifyJobStepCompleteAsync(filteredUsersToNotify, job, batch, step, node, jobStepEventData, ctx));
-
-			_logger.LogDebug("Finished sending notifications for step {JobId}:{BatchId}:{StepId}", job.Id, batchId, stepId);
+			return jobStepEventData;
 		}
 
 		private async Task SendAllLabelNotificationsAsync(IJob job, IReadOnlyList<(LabelState State, LabelOutcome Outcome)> oldLabelStates, IReadOnlyList<(LabelState, LabelOutcome)> newLabelStates, CancellationToken cancellationToken)
diff --git a/Engine/Source/Programs/Horde/Plugins/Build/HordeServer.Build/Notifications/Sinks/SlackNotificationSink.cs b/Engine/Source/Programs/Horde/Plugins/Build/HordeServer.Build/Notifications/Sinks/SlackNotificationSink.cs
--- a/Engine/Source/Programs/Horde/Plugins/Build/HordeServer.Build/Notifications/Sinks/SlackNotificationSink.cs
+++ b/Engine/Source/Programs/Horde/Plugins/Build/HordeServer.Build/Notifications/Sinks/SlackNotificationSink.cs
@@ -601,8 +601,19 @@ namespace HordeServer.Notifications.Sinks
 		#region Job step complete
 
 		/// <inheritdoc/>
-		public async Task NotifyJobStepCompleteAsync(IEnumerable<IUser> usersToNotify, IJob job, IJobStepBatch batch, IJobStep step, INode node, List<ILogEventData> jobStepEventData, CancellationToken cancellationToken)
+		public Task NotifyJobStepAbortedAsync(IEnumerable<IUser>? usersToNotify, IJob job, IJobStepBatch batch, IJobStep step, INode node, List<ILogEventData> jobStepEventData, CancellationToken cancellationToken)
+			=> Task.CompletedTask;
+
+
+		/// <inheritdoc/>
+		public async Task NotifyJobStepCompleteAsync(IEnumerable<IUser>? usersToNotify, IJob job, IJobStepBatch batch, IJobStep step, INode node, List<ILogEventData> jobStepEventData, CancellationToken cancellationToken)
 		{
+			// The Slack notification requires the list of user's to notify to be non-null
+			if (usersToNotify is null)
+			{
+				return;
+			}
+
 			if (step.Error == JobStepError.TimedOut)
 			{
 				if (_settings.JobNotificationChannel != null)
diff --git a/Engine/Source/Programs/Horde/Plugins/Experimental/HordeServer.Experimental/ExperimentalConfig.cs b/Engine/Source/Programs/Horde/Plugins/Experimental/HordeServer.Experimental/ExperimentalConfig.cs
--- a/Engine/Source/Programs/Horde/Plugins/Experimental/HordeServer.Experimental/ExperimentalConfig.cs
+++ b/Engine/Source/Programs/Horde/Plugins/Experimental/HordeServer.Experimental/ExperimentalConfig.cs
@@ -9,9 +9,12 @@ namespace HordeServer
 	/// </summary>
 	public class ExperimentalConfig : IPluginConfig
 	{
+		public SlackConfig Slack { get; set; } = new SlackConfig();
+
 		/// <inheritdoc/>
 		public void PostLoad(PluginConfigOptions configOptions)
 		{
+			Slack.PostLoad(configOptions);
 		}
 	}
 }
diff --git a/Engine/Source/Programs/Horde/Plugins/Experimental/HordeServer.Experimental/Notifications/Sinks/ExperimentalSlackNotificationSink.cs b/Engine/Source/Programs/Horde/Plugins/Experimental/HordeServer.Experimental/Notifications/Sinks/ExperimentalSlackNotificationSink.cs
--- a/Engine/Source/Programs/Horde/Plugins/Experimental/HordeServer.Experimental/Notifications/Sinks/ExperimentalSlackNotificationSink.cs
+++ b/Engine/Source/Programs/Horde/Plugins/Experimental/HordeServer.Experimental/Notifications/Sinks/ExperimentalSlackNotificationSink.cs
@@ -1,17 +1,26 @@
 // Copyright Epic Games, Inc. All Rights Reserved.
 
+using EpicGames.Core;
 using EpicGames.Horde.Jobs;
 using EpicGames.Horde.Jobs.Graphs;
+using EpicGames.Horde.Jobs.Templates;
+using EpicGames.Horde.Logs;
 using EpicGames.Slack;
+using EpicGames.Slack.Blocks;
+using EpicGames.Slack.Elements;
 using HordeServer.Agents;
 using HordeServer.Configuration;
 using HordeServer.Devices;
 using HordeServer.Issues;
 using HordeServer.Logs;
+using HordeServer.Server;
 using HordeServer.Streams;
 using HordeServer.Users;
 using Microsoft.Extensions.Logging;
 using Microsoft.Extensions.Options;
+using MongoDB.Bson.Serialization.Attributes;
+using MongoDB.Bson;
+using MongoDB.Driver;
 
 namespace HordeServer.Notifications.Sinks
 {
@@ -20,31 +29,155 @@ namespace HordeServer.Notifications.Sinks
 	/// </summary>
 	public sealed class ExperimentalSlackNotificationSink : INotificationSink, IDisposable
 	{
-		readonly BuildServerConfig _buildServerConfig;
+		/// <summary>
+		/// MongoDB document for job notifications
+		/// </summary>
+		class JobNotificationStateDocument
+		{
+			[BsonId]
+			public ObjectId Id { get; set; }
+
+			[BsonElement("jid")]
+			public JobId JobId { get; set; }
+
+			[BsonElement("tid")]
+			public TemplateId TemplateId { get; set; }
+
+			[BsonElement("rcp")]
+			public string Recipient { get; set; } = String.Empty;
+
+			[BsonElement("ch")]
+			public string Channel { get; set; } = String.Empty;
+
+			[BsonElement("ts")]
+			public string Ts { get; set; } = String.Empty;
+
+			[BsonIgnore]
+			public SlackMessageId MessageId => new SlackMessageId(Channel, null, Ts);
+		}
+
+		/// <summary>
+		/// MongoDB document for job step notifications
+		/// </summary>
+		class JobStepNotificationStateDocument
+		{
+			[BsonId]
+			public ObjectId Id { get; set; }
+
+			[BsonElement("jid")]
+			public JobId JobId { get; set; }
+
+			[BsonElement("tid")]
+			public TemplateId TemplateId { get; set; }
+
+			[BsonElement("rcp")]
+			public string Recipient { get; set; } = String.Empty;
+
+			[BsonElement("sid")]
+			public JobStepId JobStepId { get; set; }
+
+			[BsonElement("grp")]
+			public string Group { get; set; } = String.Empty;
+
+			[BsonElement("plt")]
+			public string TargetPlatform { get; set; } = String.Empty;
+
+			[BsonElement("bdg")]
+			public string Badge { get; set; } = String.Empty;
+
+			[BsonElement("ch")]
+			public string Channel { get; set; } = String.Empty;
+
+			[BsonElement("tts")]
+			public string? ThreadTs { get; set; }
+
+			[BsonElement("ts")]
+			public string Ts { get; set; } = String.Empty;
+
+			[BsonIgnore]
+			public SlackMessageId MessageId => new SlackMessageId(Channel, ThreadTs, Ts);
+		}
+
+		/// <summary>
+		/// Structure for generic use of log reporting
+		/// </summary>
+		/// <remarks>Cancellation of jobs and job steps may not always include log information. In these cases we want to still provide a message to the user for the reason.</remarks>
+		struct LogData
+		{
+			public LogEventSeverity Severity { get; }
+			public string Message { get; }
+			public LogId? LogId { get; } = null;
+			public int? LineIndex { get; } = null;
+			public LogData(LogEventSeverity severity, string message, LogId? logId = null, int? lineIndex = null)
+			{
+				Severity = severity;
+				Message = message;
+				LogId = logId;
+				LineIndex = lineIndex;
+			}
+		}
+
+		readonly IOptionsMonitor<BuildConfig> _buildConfig;
+		readonly IOptionsMonitor<ExperimentalConfig> _experimentalConfig;
+		readonly IMongoCollection<JobNotificationStateDocument> _jobNotificationStates;
+		readonly IMongoCollection<JobStepNotificationStateDocument> _jobStepNotificationStates;
+		readonly IUserCollection _userCollection;
+		readonly IServerInfo _serverInfo;
 		readonly ILogger _logger;
 
 		readonly HttpClient _httpClient;
 		readonly SlackClient _slackClient;
 
+		/// <summary>
+		/// Slack emoji used to emphasize state
+		/// </summary>
+		const string PendingBadge = ":large_blue_square:";
+		const string WarningBadge = ":large_yellow_square:";
+		const string FailureBadge = ":large_red_square:";
+		const string SkippedBadge = ":black_square:";
+		const string PassingBadge = ":large_green_square:";
+
+		/// <summary>
+		/// Slack emoji used to emphasize an aborted job
+		/// </summary>
+		const string AlertBadge = ":rotating_light:";
+
+		/// <summary>
+		/// Limit of logs to be provided in a Slack notification
+		/// </summary>
+		const int MaxEventMessages = 5;
+
 		/// <summary>
 		/// Constructor
 		/// </summary>
-		public ExperimentalSlackNotificationSink(IOptions<BuildServerConfig> buildServerConfig, ILogger<ExperimentalSlackNotificationSink> logger)
+		/// <param name="mongoService">MongoDB service</param>
+		/// <param name="buildServerConfig">Options from the build plugin's server config</param>
+		/// <param name="buildConfig">Monitored options from the build plugin's global config</param>
+		/// <param name="experimentalConfig">Monitored options from the experimental plugin's global config</param>
+		/// <param name="userCollection">Collection of users</param>
+		/// <param name="serverInfo">Information of the current server instance</param>
+		/// <param name="logger">Logger for output</param>
+		public ExperimentalSlackNotificationSink(IMongoService mongoService, IOptions<BuildServerConfig> buildServerConfig, IOptionsMonitor<BuildConfig> buildConfig, IOptionsMonitor<ExperimentalConfig> experimentalConfig, IUserCollection userCollection, IServerInfo serverInfo, ILogger<ExperimentalSlackNotificationSink> logger)
 		{
-			_buildServerConfig = buildServerConfig.Value;
+			_buildConfig = buildConfig;
+			_experimentalConfig = experimentalConfig;
+			_jobNotificationStates = mongoService.GetCollection<JobNotificationStateDocument>("SlackExperimentalJob", keys => keys.Ascending(x => x.JobId).Ascending(x => x.TemplateId).Ascending(x => x.Recipient), unique: true);
+			_jobStepNotificationStates = mongoService.GetCollection<JobStepNotificationStateDocument>("SlackExperimentalJobStep", keys => keys.Ascending(x => x.JobId).Ascending(x => x.TemplateId).Ascending(x => x.Recipient).Ascending(x => x.Group).Ascending(x => x.JobStepId), unique: true);
+			_userCollection = userCollection;
+			_serverInfo = serverInfo;
 			_logger = logger;
 
 			_httpClient = new HttpClient();
-			_httpClient.DefaultRequestHeaders.Add("Authorization", $"Bearer {_buildServerConfig.SlackToken ?? ""}");
+			_httpClient.DefaultRequestHeaders.Add("Authorization", $"Bearer {buildServerConfig.Value.SlackToken ?? ""}");
 			_slackClient = new SlackClient(_httpClient, _logger);
-
-			_ = _slackClient;
 		}
 
 		/// <inheritdoc/>
 		public void Dispose()
 			=> _httpClient.Dispose();
 
+		#region INotificationSink Methods
+
 		/// <inheritdoc/>
 		public Task NotifyConfigUpdateAsync(ConfigUpdateInfo info, CancellationToken cancellationToken)
 			=> Task.CompletedTask;
@@ -62,8 +195,43 @@ namespace HordeServer.Notifications.Sinks
 			=> Task.CompletedTask;
 
 		/// <inheritdoc/>
-		public Task NotifyJobCompleteAsync(IJob job, IGraph graph, LabelOutcome outcome, CancellationToken cancellationToken)
-			=> Task.CompletedTask;
+		public async Task NotifyJobCompleteAsync(IJob job, IGraph graph, LabelOutcome outcome, CancellationToken cancellationToken)
+		{
+			string hordeJobURL = new Uri(_serverInfo.DashboardUrl, $"job/{job.Id}").ToString();
+
+			string branch = "Unknown Branch";
+			if (_buildConfig.CurrentValue.TryGetStream(job.StreamId, out StreamConfig? config))
+			{
+				branch = config.Name;
+			}
+
+			string? abortedByUserName = null;
+			if (job.AbortedByUserId is not null)
+			{
+				IUser? user = await _userCollection.GetUserAsync(job.AbortedByUserId.Value, cancellationToken);
+				if (user is not null)
+				{
+					abortedByUserName = user.Name;
+				}
+			}
+
+			// Send the updated summary after the steps has been updated
+			List<JobNotificationStateDocument>? jobNotificationStates = await GetJobNotificationStatesAsync(job.Id, cancellationToken);
+			if (jobNotificationStates is null)
+			{
+				_logger.LogError("No job notifications found for job id '{JobId}'.", job.Id);
+				return;
+			}
+
+			foreach (JobNotificationStateDocument jobNotificationState in jobNotificationStates)
+			{
+				await SendSlackSummaryMessageAsync(job, branch, hordeJobURL, jobNotificationState.TemplateId, jobNotificationState.Recipient, abortedByUserName, cancellationToken);
+			}
+
+			// Delete all documents related to this job as no other notifications will be sent
+			_ = await DeleteJobNotificationStatesAsync(job.Id, cancellationToken);
+			_ = await DeleteJobStepNotificationStatesAsync(job.Id, cancellationToken);
+		}
 
 		/// <inheritdoc/>
 		public Task NotifyJobCompleteAsync(IUser user, IJob job, IGraph graph, LabelOutcome outcome, CancellationToken cancellationToken)
@@ -74,8 +242,34 @@ namespace HordeServer.Notifications.Sinks
 			=> Task.CompletedTask;
 
 		/// <inheritdoc/>
-		public Task NotifyJobStepCompleteAsync(IEnumerable<IUser> usersToNotify, IJob job, IJobStepBatch batch, IJobStep step, INode node, List<ILogEventData> jobStepEventData, CancellationToken cancellationToken)
-			=> Task.CompletedTask;
+		public async Task NotifyJobStepAbortedAsync(IEnumerable<IUser>? usersToNotify, IJob job, IJobStepBatch batch, IJobStep step, INode node, List<ILogEventData> jobStepEventData, CancellationToken cancellationToken)
+		{
+			List<LogData> logEvents = new List<LogData>();
+			string cancellationDetails = "Step was canceled.";
+			if (step.AbortedByUserId is not null)
+			{
+				IUser? user = await _userCollection.GetUserAsync(step.AbortedByUserId.Value, cancellationToken);
+				if (user is not null)
+				{
+					cancellationDetails = $"Step was canceled by {user.Name}.";
+				}
+			}
+
+			if (!String.IsNullOrEmpty(step.CancellationReason))
+			{
+				cancellationDetails += $" Reason for cancellation: '{step.CancellationReason}'.";
+			}
+			logEvents.Add(new LogData(LogEventSeverity.Error, cancellationDetails));
+
+			await ProcessJobStepCompleteNotificationAsync(job, step, logEvents, cancellationToken);
+		}
+
+		/// <inheritdoc/>
+		public async Task NotifyJobStepCompleteAsync(IEnumerable<IUser>? usersToNotify, IJob job, IJobStepBatch batch, IJobStep step, INode node, List<ILogEventData> jobStepEventData, CancellationToken cancellationToken)
+		{
+			List<LogData> logEvents = ExtractLogInformation(step, jobStepEventData);
+			await ProcessJobStepCompleteNotificationAsync(job, step, logEvents, cancellationToken);
+		}
 
 		/// <inheritdoc/>
 		public Task NotifyLabelCompleteAsync(IUser user, IJob job, ILabel label, int labelIdx, LabelOutcome outcome, List<(string, JobStepOutcome, Uri)> stepData, CancellationToken cancellationToken)
@@ -92,5 +286,880 @@ namespace HordeServer.Notifications.Sinks
 		/// <inheritdoc/>
 		public Task SendIssueReportAsync(IssueReportGroup report, CancellationToken cancellationToken)
 			=> Task.CompletedTask;
+
+		#endregion INotificationSink Methods
+
+		#region Job Processing
+
+		/// <summary>
+		/// Search the job for steps to be included in the notifications based on the experimental plugin's global config
+		/// </summary>
+		/// <param name="job">Information about the job</param>
+		/// <param name="config">Configuration to fetch notification information</param>
+		/// <param name="cancellationToken">Cancellation token for the operation</param>
+		async Task ProcessJobAsync(IJob job, JobNotificationConfig config, CancellationToken cancellationToken = default)
+		{
+			// Populate our job information needed for the notification
+			foreach (IJobStepBatch batch in job.Batches)
+			{
+				foreach (IJobStep step in batch.Steps)
+				{
+					if (config.GetStepDetails(step.Name, out string? group, out string? platform) && !String.IsNullOrEmpty(group) &&!String.IsNullOrEmpty(platform))
+					{
+						string badge = GetBadgeForStep(step);
+						foreach (string recipient in config.Channels)
+						{
+							await AddOrUpdateJobStepNotificationStateAsync(job.Id, job.TemplateId, recipient, step.Id, group, platform, badge, null, cancellationToken);
+						}
+					}
+				}
+			}
+		}
+
+		/// <summary>
+		/// Prepares a job summary and job step thread to be sent to Slack based on the experimental plugin's global configuration
+		/// </summary>
+		/// <param name="job">Information about the job</param>
+		/// <param name="step">Information about the step</param>
+		/// <param name="jobStepEventData">List of the step's logged events</param>
+		/// <param name="cancellationToken">Cancellation token for the operation</param>
+		async Task ProcessJobStepCompleteNotificationAsync(IJob job, IJobStep step, List<LogData> jobStepEventData, CancellationToken cancellationToken = default)
+		{
+			SlackConfig currentConfig = _experimentalConfig.CurrentValue.Slack;
+
+			List<JobNotificationConfig>? configs = currentConfig.GetNotificationConfigs(job.Name, job.TemplateId, job.StreamId);
+			if (configs is null || !configs.Any())
+			{
+				// Log that there is nothing to do here
+				_logger.LogInformation("No notification config found for '{JobName}' in stream '{StreamId}'", job.Name, job.StreamId);
+				return;
+			}
+
+			string hordeJobURL = new Uri(_serverInfo.DashboardUrl, $"job/{job.Id}").ToString();
+
+			string branch = "Unknown Branch";
+			if (_buildConfig.CurrentValue.TryGetStream(job.StreamId, out StreamConfig? streamConfig))
+			{
+				branch = streamConfig.Name;
+			}
+
+			// Check if the step completed was used to generate the details of the job so that the initial notification can be sent out
+			if (step.Name == IJob.SetupNodeName)
+			{
+				// Check if the Setup step failed to notify that the job can't be executed
+				if (step.Outcome != JobStepOutcome.Success)
+				{
+					await SendSetupFailedSlackMessageAsync(job, step, branch, jobStepEventData, configs, cancellationToken);
+					return;
+				}
+
+				foreach (JobNotificationConfig config in configs)
+				{
+					await ProcessJobAsync(job, config, cancellationToken);
+					foreach (string recipient in config.Channels)
+					{
+						await SendSlackSummaryMessageAsync(job, branch, hordeJobURL, config.Template, recipient, null, cancellationToken);
+					}
+				}
+				return;
+			}
+
+			// Send a message for our current step
+			bool needsSummaryUpdate = await UpdateAndSendSlackNotificationsForStepAsync(job.Id, step, jobStepEventData, null, cancellationToken);
+
+			// Fetch all dependent steps if there was an error and send notifications with this step as the cause for the dependent steps being skipped
+			if (step.Outcome != JobStepOutcome.Success)
+			{
+				HashSet<JobStepId>? dependentSteps = await FindDependentJobStepIdsAsync(job, step.Id, cancellationToken);
+				if (dependentSteps is not null)
+				{
+					foreach (JobStepId stepId in dependentSteps)
+					{
+						// Send a message for our dependent step
+						if (job.TryGetStep(stepId, out IJobStep? dependentStep))
+						{
+							needsSummaryUpdate |= await UpdateAndSendSlackNotificationsForStepAsync(job.Id, dependentStep, jobStepEventData, step, cancellationToken);
+						}
+					}
+				}
+			}
+
+			// Check to see if any of our steps were updated which requires the Slack summary to be updated
+			if (!needsSummaryUpdate)
+			{
+				return;
+			}
+
+			// Send the updated summary after the steps has been updated
+			List<JobNotificationStateDocument>? jobNotificationStates = await GetJobNotificationStatesAsync(job.Id, cancellationToken);
+			if (jobNotificationStates is null)
+			{
+				_logger.LogError("No job notifications found for job id '{JobId}'.", job.Id);
+				return;
+			}
+
+			foreach (JobNotificationStateDocument jobNotificationState in jobNotificationStates)
+			{
+				await SendSlackSummaryMessageAsync(job, branch, hordeJobURL, jobNotificationState.TemplateId, jobNotificationState.Recipient, null, cancellationToken);
+			}
+		}
+
+		/// <summary>
+		/// Finds all dependent steps of the <see cref="JobStepNotificationStateDocument"/>
+		/// </summary>
+		/// <param name="job">Information about the job</param>
+		/// <param name="underlyingStepId">Identifier of the possible job step dependency</param>
+		/// <param name="cancellationToken">Cancellation token for the operation</param>
+		/// <returns></returns>
+		async Task<HashSet<JobStepId>?> FindDependentJobStepIdsAsync(IJob job, JobStepId underlyingStepId, CancellationToken cancellationToken = default)
+		{
+			// Fetch all of our job steps sending notifications
+			FilterDefinition<JobStepNotificationStateDocument> filter = Builders<JobStepNotificationStateDocument>.Filter.Eq(x => x.JobId, job.Id);
+			List<JobStepNotificationStateDocument>? jobStepNotificationStateDocuments = await GetJobStepNotificationStatesAsync(filter, cancellationToken);
+			if (jobStepNotificationStateDocuments is null)
+			{
+				return null;
+			}
+
+			HashSet<JobStepId> dependencyJobStepNotifications = new HashSet<JobStepId>();
+			HashSet<JobStepId> cachedSteps = new HashSet<JobStepId>();
+			foreach (JobStepNotificationStateDocument jobStepNotificationState in jobStepNotificationStateDocuments)
+			{
+				if (cachedSteps.Contains(jobStepNotificationState.JobStepId))
+				{
+					continue;
+				}
+				cachedSteps.Add(jobStepNotificationState.JobStepId);
+
+				List<JobStepId> dependencySteps = new List<JobStepId>()
+				{
+					jobStepNotificationState.JobStepId
+				};
+
+				HashSet<JobStepId> visitedDependencySteps = new HashSet<JobStepId>()
+				{
+					jobStepNotificationState.JobStepId
+				};
+
+				bool isMatchingDependency = false;
+				while (dependencySteps.Count > 0 && !isMatchingDependency)
+				{
+					JobStepId jobStepId = dependencySteps[0];
+					dependencySteps.RemoveAt(0);
+
+					IJobStep? currentStep;
+					if (!job.TryGetStep(jobStepId, out currentStep))
+					{
+						continue;
+					}
+
+					// Iterate through all of the step's dependencies
+					// Add the step Id of the notification state if the dependencies match
+					foreach (JobStepId dependencyStepId in currentStep.InputDependencies)
+					{
+						if (underlyingStepId == dependencyStepId)
+						{
+							isMatchingDependency = true;
+							dependencyJobStepNotifications.Add(jobStepNotificationState.JobStepId);
+							break;
+						}
+						else if (!visitedDependencySteps.Contains(dependencyStepId))
+						{
+							dependencySteps.Add(dependencyStepId);
+							visitedDependencySteps.Add(dependencyStepId);
+						}
+					}
+
+					// Check to see if we already found a potential match before continuing searching through dependencies
+					if (isMatchingDependency)
+					{
+						continue;
+					}
+
+					foreach (JobStepId dependencyStepId in currentStep.OrderDependencies)
+					{
+						if (underlyingStepId == dependencyStepId)
+						{
+							isMatchingDependency = true;
+							dependencyJobStepNotifications.Add(jobStepNotificationState.JobStepId);
+							break;
+						}
+						else if (!visitedDependencySteps.Contains(dependencyStepId))
+						{
+							dependencySteps.Add(dependencyStepId);
+							visitedDependencySteps.Add(dependencyStepId);
+						}
+					}
+				}
+			}
+
+			return dependencyJobStepNotifications;
+		}
+
+		/// <summary>
+		/// Determines the badge based on the job step's state and outcome
+		/// </summary>
+		/// <param name="step">Information about the step</param>
+		/// <returns>String representation of the badge</returns>
+		static string GetBadgeForStep(IJobStep step)
+		{
+			// Check the state
+			switch (step.State)
+			{
+				case JobStepState.Waiting:
+				case JobStepState.Ready:
+				case JobStepState.Running:
+					{
+						// Step is either queued or currently being executed
+						return PendingBadge;
+					}
+				case JobStepState.Skipped:
+					{
+						return SkippedBadge;
+					}
+				case JobStepState.Aborted:
+					{
+						// Treat aborted as failure conditions
+						return FailureBadge;
+					}
+				case JobStepState.Completed:
+					{
+						// Continue to check the outcome of the job
+						break;
+					}
+				default:
+					{
+						// Log an error "Job Step in an Unspecified State"
+						return String.Empty;
+					}
+			}
+
+			// Check the outcome
+			switch (step.Outcome)
+			{
+				case JobStepOutcome.Failure:
+					{
+						return FailureBadge;
+					}
+				case JobStepOutcome.Warnings:
+					{
+						return WarningBadge;
+					}
+				case JobStepOutcome.Success:
+					{
+						return PassingBadge;
+					}
+				default:
+					{
+						// Log an error "Job Step in an Unspecified Outcome"
+						return String.Empty;
+					}
+			}
+		}
+
+		/// <summary>
+		/// Parses list of <see cref="ILogEventData"/> into the generic <see cref="LogData"/> structure
+		/// </summary>
+		/// <param name="step">Information about the step</param>
+		/// <param name="jobStepEventData">List of the step's logged events</param>
+		/// <returns>List of parsed logged events to be consumed by the Slack notifications</returns>
+		static List<LogData> ExtractLogInformation(IJobStep step, List<ILogEventData> jobStepEventData)
+		{
+			List<LogData> logData = new List<LogData>();
+			foreach (ILogEventData eventData in jobStepEventData)
+			{
+				// Extract the log information from the first log line
+				if (eventData.Lines.Count > 0)
+				{
+					LogEvent logEvent = LogEvent.Read(eventData.Lines[0].Data.Span);
+					logData.Add(new LogData(eventData.Severity, eventData.Message, step.LogId, logEvent.LineIndex));
+				}
+				else if (!String.IsNullOrEmpty(eventData.Message))
+				{
+					// Only store the relevant message for the event if there is no actual log information
+					logData.Add(new LogData(eventData.Severity, eventData.Message, step.LogId));
+				}
+			}
+			return logData;
+		}
+
+		#endregion Job Processing
+
+		#region MongoDB
+
+		/// <summary>
+		/// Add or update the <see cref="JobNotificationStateDocument"/> to the collection
+		/// </summary>
+		/// <param name="jobId">Identifier of the job</param>
+		/// <param name="templateId">Identifier of the stream's template</param>
+		/// <param name="recipient">Identifier of the channel or user in Slack or user's email address</param>
+		/// <param name="messageId">Slack message identifer of the sent notification</param>
+		/// <param name="cancellationToken">Cancellation token for the operation</param>
+		/// <returns>A tuple of the <see cref="JobNotificationStateDocument"/> and a bool used to determine if the record is considered to be new</returns>
+		async Task<(JobNotificationStateDocument, bool)> AddOrUpdateJobNotificationStateAsync(JobId jobId, TemplateId templateId, string recipient, SlackMessageId messageId, CancellationToken cancellationToken = default)
+		{
+			ObjectId newId = ObjectId.GenerateNewId();
+
+			FilterDefinition<JobNotificationStateDocument> filter = Builders<JobNotificationStateDocument>.Filter.Eq(x => x.JobId, jobId) &
+																		Builders<JobNotificationStateDocument>.Filter.Eq(x => x.TemplateId, templateId) &
+																		Builders<JobNotificationStateDocument>.Filter.Eq(x => x.Recipient, recipient);
+			UpdateDefinition<JobNotificationStateDocument> update = Builders<JobNotificationStateDocument>.Update.SetOnInsert(x => x.Id, newId).Set(x => x.Channel, messageId.Channel).Set(x => x.Ts, messageId.Ts);
+
+			JobNotificationStateDocument state = await _jobNotificationStates.FindOneAndUpdateAsync(filter, update, new FindOneAndUpdateOptions<JobNotificationStateDocument> { IsUpsert = true, ReturnDocument = ReturnDocument.After }, cancellationToken);
+			if (state.Id == newId)
+			{
+				_logger.LogInformation("Posted message {StateId} (jobId: {JobId}, templateId: {TemplateId}, recipient: {Recipient}, messageId: {MessageId})", state.Id, state.JobId, state.TemplateId, state.Recipient, state.MessageId);
+			}
+			else
+			{
+				_logger.LogInformation("Updated message {StateId} (jobId: {JobId}, templateId: {TemplateId}, recipient: {Recipient}, messageId: {MessageId})", state.Id, state.JobId, state.TemplateId, state.Recipient, state.MessageId);
+			}
+
+			return (state, state.Id == newId);
+		}
+
+		/// <summary>
+		/// Fetches a list of <see cref="JobNotificationStateDocument"/>
+		/// </summary>
+		/// <param name="jobId">Identifier of the job</param>
+		/// <param name="cancellationToken">Cancellation token for the operation</param>
+		/// <returns>A list of <see cref="JobNotificationStateDocument"/> if records were found, otherwise returns null for error checking</returns>
+		async Task<List<JobNotificationStateDocument>?> GetJobNotificationStatesAsync(JobId jobId, CancellationToken cancellationToken = default)
+		{
+			FilterDefinition<JobNotificationStateDocument> filter = Builders<JobNotificationStateDocument>.Filter.Eq(x => x.JobId, jobId);
+			List<JobNotificationStateDocument> states = await _jobNotificationStates.Find(filter).ToListAsync(cancellationToken);
+			if (states.Count == 0)
+			{
+				return null;
+			}
+
+			return states;
+		}
+
+		/// <summary>
+		/// Fetches a single entity of <see cref="JobNotificationStateDocument"/>
+		/// </summary>
+		/// <param name="jobId">Identifier of the job</param>
+		/// <param name="templateId">Identifier of the stream's template</param>
+		/// <param name="recipient">Identifier of the channel or user in Slack or user's email address</param>
+		/// <param name="cancellationToken">Cancellation token for the operation</param>
+		/// <returns>A <see cref="JobNotificationStateDocument"/> if records were found, otherwise returns null for error checking</returns>
+		/// <remarks>Only a single <see cref="JobNotificationStateDocument"/> should exist for the aggregate of the job identifier, template identifier, and recipient should ever exist. Will report an error if another record had been found.</remarks>
+		async Task<JobNotificationStateDocument?> GetJobNotificationStateAsync(JobId jobId, TemplateId templateId, string recipient, CancellationToken cancellationToken = default)
+		{
+			FilterDefinition<JobNotificationStateDocument> filter = Builders<JobNotificationStateDocument>.Filter.Eq(x => x.JobId, jobId) &
+																		Builders<JobNotificationStateDocument>.Filter.Eq(x => x.TemplateId, templateId) &
+																		Builders<JobNotificationStateDocument>.Filter.Eq(x => x.Recipient, recipient);
+			List<JobNotificationStateDocument> states = await _jobNotificationStates.Find(filter).ToListAsync(cancellationToken);
+			if (states.Count == 0)
+			{
+				return null;
+			}
+
+			if (states.Count != 1)
+			{
+				// Log error that multiple were found and could cause issues
+				_logger.LogError("Found {Count} records when fetching for {JobId}:{TemplateId}:{Recipient}. Using first record which could cause issues.", states.Count, jobId, templateId, recipient);
+			}
+
+			return states[0];
+		}
+
+		/// <summary>
+		/// Deletes all <see cref="JobNotificationStateDocument"/> from the collection
+		/// </summary>
+		/// <param name="jobId">Identifier of the job</param>
+		/// <param name="cancellationToken">Cancellation token for the operation</param>
+		/// <returns>True if records were deleted, otherwise false</returns>
+		async Task<bool> DeleteJobNotificationStatesAsync(JobId jobId, CancellationToken cancellationToken = default)
+		{
+			FilterDefinition<JobNotificationStateDocument> filter = Builders<JobNotificationStateDocument>.Filter.Eq(x => x.JobId, jobId);
+			DeleteResult result = await _jobNotificationStates.DeleteManyAsync(filter, cancellationToken);
+			return result.DeletedCount > 0;
+		}
+
+		/// <summary>
+		/// Add or update the <see cref="JobStepNotificationStateDocument"/> to the collection
+		/// </summary>
+		/// <param name="jobId">Identifier of the job</param>
+		/// <param name="templateId">Identifier of the stream's template</param>
+		/// <param name="recipient">Identifier of the channel or user in Slack or user's email address</param>
+		/// <param name="stepId">Identifier of the job step</param>
+		/// <param name="group">Group associated with the notification</param>
+		/// <param name="platform">Target platform associated with the notification</param>
+		/// <param name="badge">Slack emoji associated with the notification</param>
+		/// <param name="messageId">Optional Slack message identifer if the notification has been sent out</param>
+		/// <param name="cancellationToken">Cancellation token for the operation</param>
+		/// <returns>A tuple of the <see cref="JobStepNotificationStateDocument"/> and a bool used to determine if the record is considered to be new</returns>
+		async Task<(JobStepNotificationStateDocument, bool)> AddOrUpdateJobStepNotificationStateAsync(JobId jobId, TemplateId templateId, string recipient, JobStepId stepId, string group, string platform, string badge, SlackMessageId? messageId, CancellationToken cancellationToken = default)
+		{
+			ObjectId newId = ObjectId.GenerateNewId();
+
+			FilterDefinition<JobStepNotificationStateDocument> filter = Builders<JobStepNotificationStateDocument>.Filter.Eq(x => x.JobId, jobId) &
+																		Builders<JobStepNotificationStateDocument>.Filter.Eq(x => x.TemplateId, templateId) &
+																		Builders<JobStepNotificationStateDocument>.Filter.Eq(x => x.Recipient, recipient) &
+																		Builders<JobStepNotificationStateDocument>.Filter.Eq(x => x.JobStepId, stepId) &
+																		Builders<JobStepNotificationStateDocument>.Filter.Eq(x => x.Group, group) &
+																		Builders<JobStepNotificationStateDocument>.Filter.Eq(x => x.TargetPlatform, platform);
+			UpdateDefinition<JobStepNotificationStateDocument> update = Builders<JobStepNotificationStateDocument>.Update.SetOnInsert(x => x.Id, newId).Set(x => x.Badge, badge);
+
+			if (messageId != null)
+			{
+				update = update.Set(x => x.Channel, messageId.Channel).Set(x => x.Ts, messageId.Ts);
+			}
+
+			JobStepNotificationStateDocument state = await _jobStepNotificationStates.FindOneAndUpdateAsync(filter, update, new FindOneAndUpdateOptions<JobStepNotificationStateDocument> { IsUpsert = true, ReturnDocument = ReturnDocument.After }, cancellationToken);
+			if (state.Id == newId)
+			{
+				_logger.LogInformation("Posted message {StateId} (jobId: {JobId}, templateId: {TemplateId}, recipient: {Recipient}, jobStepId: {JobStepId}, group: {Group}, platform: {Platform}, badge: {Badge}, messageId: {MessageId})", state.Id, state.JobId, state.TemplateId, state.Recipient, state.JobStepId, state.Group, state.TargetPlatform, state.Badge, state.MessageId);
+			}
+			else
+			{
+				_logger.LogInformation("Updatedmessage {StateId} (jobId: {JobId}, templateId: {TemplateId}, recipient: {Recipient}, jobStepId: {JobStepId}, group: {Group}, platform: {Platform}, badge: {Badge}, messageId: {MessageId})", state.Id, state.JobId, state.TemplateId, state.Recipient, state.JobStepId, state.Group, state.TargetPlatform, state.Badge, state.MessageId);
+			}
+
+			return (state, state.Id == newId);
+		}
+
+		/// <summary>
+		/// Fetches a list of <see cref="JobStepNotificationStateDocument"/>
+		/// </summary>
+		/// <param name="filter">Filter used to query against the collection</param>
+		/// <param name="cancellationToken">Cancellation token for the operation</param>
+		/// <returns>A list of <see cref="JobStepNotificationStateDocument"/> if records were found, otherwise returns null for error checking</returns>
+		async Task<List<JobStepNotificationStateDocument>?> GetJobStepNotificationStatesAsync(FilterDefinition<JobStepNotificationStateDocument> filter, CancellationToken cancellationToken = default)
+		{
+			List<JobStepNotificationStateDocument> states = await _jobStepNotificationStates.Find(filter).SortBy(x => x.JobId).SortBy(x => x.TemplateId).SortBy(x => x.Recipient).SortBy(x => x.Group).ToListAsync(cancellationToken);
+			if (states.Count == 0)
+			{
+				return null;
+			}
+
+			return states;
+		}
+
+		/// <summary>
+		/// Deletes all <see cref="JobStepNotificationStateDocument"/> from the collection
+		/// </summary>
+		/// <param name="jobId">Identifier of the job</param>
+		/// <param name="cancellationToken">Cancellation token for the operation</param>
+		/// <returns>True if records were deleted, otherwise false</returns>
+		async Task<bool> DeleteJobStepNotificationStatesAsync(JobId jobId, CancellationToken cancellationToken = default)
+		{
+			FilterDefinition<JobStepNotificationStateDocument> filter = Builders<JobStepNotificationStateDocument>.Filter.Eq(x => x.JobId, jobId);
+			DeleteResult result = await _jobStepNotificationStates.DeleteManyAsync(filter, cancellationToken);
+			return result.DeletedCount > 0;
+		}
+
+		#endregion MongoDB
+
+		#region Slack
+
+		/// <summary>
+		/// Create the job's summary message
+		/// </summary>
+		/// <param name="job">Information about the job</param>
+		/// <param name="branch">Name of the branch from the stream</param>
+		/// <param name="hordeJobURL">URL of the job in Horde</param>
+		/// <param name="templateId">Identifier of the stream's template</param>
+		/// <param name="recipient">Identifier of the channel or user in Slack or user's email address</param>
+		/// <param name="abortedByUserName">Optional username provided when cancelling a job</param>
+		/// <param name="cancellationToken">Cancellation token for the operation</param>
+		/// <returns>Generated <see cref="SlackMessage"/> if successful, otherwise returns null</returns>
+		async Task<SlackMessage?> GenerateSlackSummaryMessageAsync(IJob job, string branch, string hordeJobURL, TemplateId templateId, string recipient, string? abortedByUserName = null, CancellationToken cancellationToken = default)
+		{
+			// Fetch all of our job steps needed to send out the summary
+			FilterDefinition<JobStepNotificationStateDocument> filter = Builders<JobStepNotificationStateDocument>.Filter.Eq(x => x.JobId, job.Id) &
+																		Builders<JobStepNotificationStateDocument>.Filter.Eq(x => x.TemplateId, templateId) &
+																		Builders<JobStepNotificationStateDocument>.Filter.Eq(x => x.Recipient, recipient);
+			List<JobStepNotificationStateDocument>? jobStepNotificationStateDocuments = await GetJobStepNotificationStatesAsync(filter, cancellationToken);
+			if (jobStepNotificationStateDocuments is null)
+			{
+				return null;
+			}
+
+			return CreateSlackSummaryNotification(job, branch, hordeJobURL, jobStepNotificationStateDocuments, abortedByUserName);
+		}
+
+		async Task SendSetupFailedSlackMessageAsync(IJob job, IJobStep step, string branch, List<LogData> jobStepEventData, List<JobNotificationConfig> jobNotificationConfigs, CancellationToken cancellationToken = default)
+		{
+			DividerBlock divider = new DividerBlock();
+			SlackMessage message = new SlackMessage();
+
+			string jobUrlLink = new Uri($"{_serverInfo.DashboardUrl}job/{job.Id}").ToString();
+			string jobStepLink = new Uri($"{jobUrlLink}?step={step.Id}").ToString();
+
+			// Create our notification text
+			string notificationText = GenerateSlackMessageTitle(job, branch);
+			message.Text = $"{notificationText} Details...";
+
+			// Create a Block for the Slack Title
+			string messageTitle = GenerateSlackMessageTitle(job, branch, jobUrlLink);
+			ContextBlock jobContext = new ContextBlock();
+			jobContext.Elements.Add(new TextObject(messageTitle));
+			message.Blocks.Add(jobContext);
+
+			ContextBlock infoBlock = new ContextBlock();
+			infoBlock.Elements.Add(new TextObject($"<{jobStepLink}|*{step.Name}*> encountered errors"));
+			message.Blocks.Add(infoBlock);
+
+			AddLogDataContext(message, jobStepEventData);
+
+			foreach (JobNotificationConfig config in jobNotificationConfigs)
+			{
+				foreach (string recipient in config.Channels)
+				{
+					// This is our first time sending the summary grid, add both the job and the job steps to the collection
+					string channel = recipient;
+					SlackUser? userInfo = await _slackClient.FindUserByEmailAsync(recipient, cancellationToken);
+					if (userInfo is not null && userInfo.Id is not null)
+					{
+						channel = userInfo.Id;
+					}
+
+					_ = await _slackClient.PostMessageAsync(recipient, message, cancellationToken);
+				}
+			}
+
+			_logger.LogInformation("Finished sending failing setup notification for job '{JobId}'", job.Id);
+		}
+
+		/// <summary>
+		/// Sends a Slack notification for the specified job
+		/// </summary>
+		/// <param name="job">Information about the job</param>
+		/// <param name="branch">Name of the branch from the stream</param>
+		/// <param name="hordeJobURL">Optional URL of the job in Horde</param>
+		/// <param name="templateId">Identifier of the stream's template</param>
+		/// <param name="recipient">Identifier of the channel or user in Slack or user's email address</param>
+		/// <param name="abortedByUserName">Optional username provided when cancelling a job</param>
+		/// <param name="cancellationToken">Cancellation token for the operation</param>
+		/// <returns>True if a notification needed to be sent, otherwise returns false</returns>
+		async Task SendSlackSummaryMessageAsync(IJob job, string branch, string hordeJobURL, TemplateId templateId, string recipient, string? abortedByUserName = null, CancellationToken cancellationToken = default)
+		{
+			SlackMessage? message = await GenerateSlackSummaryMessageAsync(job, branch, hordeJobURL, templateId, recipient, abortedByUserName, cancellationToken);
+			if (message is null)
+			{
+				_logger.LogError("Unable to generate Job summary notification for job '{JobId}'", job.Id);
+				return;
+			}
+
+			JobNotificationStateDocument? jobNotificationState = await GetJobNotificationStateAsync(job.Id, templateId, recipient, cancellationToken);
+			if (jobNotificationState is not null)
+			{
+				_ = _slackClient.UpdateMessageAsync(jobNotificationState.MessageId, message, cancellationToken);
+			}
+			else
+			{
+				// This is our first time sending the summary grid, add both the job and the job steps to the collection
+				string channel = recipient;
+				SlackUser? userInfo = await _slackClient.FindUserByEmailAsync(recipient, cancellationToken);
+				if (userInfo is not null && userInfo.Id is not null)
+				{
+					channel = userInfo.Id;
+				}
+
+				SlackMessageId messageId = await _slackClient.PostMessageAsync(recipient, message, cancellationToken);
+				_ = await AddOrUpdateJobNotificationStateAsync(job.Id, templateId, channel, messageId, cancellationToken);
+			}
+		}
+
+		/// <summary>
+		/// Sends a Slack notification to a new or existing thread
+		/// </summary>
+		/// <param name="jobMessageId"><see cref="SlackMessageId"/> of the summary message</param>
+		/// <param name="jobStepNotificationState">Information about the job step's notification state</param>
+		/// <param name="threadMessage"><see cref="SlackMessage"/> to be sent</param>
+		/// <param name="cancellationToken">Cancellation token for the operation</param>
+		/// <returns>The <see cref="SlackMessageId"/> of either the newly sent notification or from the <see cref="JobStepNotificationStateDocument"/></returns>
+		async Task<SlackMessageId> SendSlackThreadMessageAsync(SlackMessageId jobMessageId, JobStepNotificationStateDocument jobStepNotificationState, SlackMessage threadMessage, CancellationToken cancellationToken = default)
+		{
+			// Check to see if a message has already been sent out to avoid sending duplicate notifications
+			if (jobStepNotificationState.MessageId.ThreadTs is not null)
+			{
+				return jobStepNotificationState.MessageId;
+			}
+
+			SlackMessageId threadMessageId = await _slackClient.PostMessageToThreadAsync(jobMessageId, threadMessage, cancellationToken);
+
+			_logger.LogInformation("Finished sending notification for job step '{JobStepId}' with Job Id '{JobId}' for template '{Template}'", jobStepNotificationState.JobStepId, jobStepNotificationState.JobId, jobStepNotificationState.TemplateId);
+
+			return threadMessageId;
+		}
+
+		/// <summary>
+		/// Sends a Slack notification for the specified job step
+		/// </summary>
+		/// <param name="jobId">Identifier of the job</param>
+		/// <param name="step">Information about the step</param>
+		/// <param name="jobStepEventData">List of the step's logged events</param>
+		/// <param name="underlyingStep">Optional information about the step's dependent step</param>
+		/// <param name="cancellationToken">Cancellation token for the operation</param>
+		/// <returns>True if a job step notification needed to be updated, otherwise returns false</returns>
+		async Task<bool> UpdateAndSendSlackNotificationsForStepAsync(JobId jobId, IJobStep step, List<LogData> jobStepEventData, IJobStep? underlyingStep = null, CancellationToken cancellationToken = default)
+		{
+			// Update our job step if it's being monitored and send an updated summary and post to the thread if needed
+			FilterDefinition<JobStepNotificationStateDocument> filter = Builders<JobStepNotificationStateDocument>.Filter.Eq(x => x.JobId, jobId) &
+																		Builders<JobStepNotificationStateDocument>.Filter.Eq(x => x.JobStepId, step.Id);
+			List<JobStepNotificationStateDocument>? jobStepNotificationStates = await GetJobStepNotificationStatesAsync(filter, cancellationToken);
+			if (jobStepNotificationStates is null || jobStepNotificationStates.Count == 0)
+			{
+				_logger.LogInformation("No notification to be sent for step '{JobStepId}' in job '{JobId}'", step.Id, jobId);
+				return false;
+			}
+
+			string badge = GetBadgeForStep(step);
+			SlackMessage? threadMessage = null;
+
+			// Notify on any problems of the Job Step
+			if (step.Outcome != JobStepOutcome.Success)
+			{
+				threadMessage = GenerateSlackThreadMessage(step, badge, jobStepEventData, underlyingStep);
+			}
+
+			foreach (JobStepNotificationStateDocument jobStepNotificationState in jobStepNotificationStates)
+			{
+				SlackMessageId? threadMessageId = null;
+				if (threadMessage is not null)
+				{
+					// Fetch our initial job notification to get the Message Id for threads
+					JobNotificationStateDocument? jobNotificationState = await GetJobNotificationStateAsync(jobId, jobStepNotificationState.TemplateId, jobStepNotificationState.Recipient, cancellationToken);
+					if (jobNotificationState is null)
+					{
+						_logger.LogError("No message id found for job '{JobId}' with template '{TemplateId}' for thread to be posted.", jobId, jobStepNotificationState.TemplateId);
+						continue;
+					}
+					threadMessageId = await SendSlackThreadMessageAsync(jobNotificationState.MessageId, jobStepNotificationState, threadMessage, cancellationToken);
+				}
+				_ = await AddOrUpdateJobStepNotificationStateAsync(jobStepNotificationState.JobId, jobStepNotificationState.TemplateId, jobStepNotificationState.Recipient, jobStepNotificationState.JobStepId, jobStepNotificationState.Group, jobStepNotificationState.TargetPlatform, badge, threadMessageId, cancellationToken);
+			}
+
+			return true;
+		}
+
+		/// <summary>
+		/// Adds a context to the Slack message with log details
+		/// </summary>
+		/// <param name="slackBlockContainer">Slack block container</param>
+		/// <param name="logData">List of our log information</param>
+		void AddLogDataContext(ISlackBlockContainer slackBlockContainer, List<LogData> logData)
+		{
+			// Don't clutter the message with all of the events and do a preliminary check to see if the user needs to be made aware that events have been truncated
+			if (logData.Count > MaxEventMessages)
+			{
+				ContextBlock limitContext = new ContextBlock();
+				limitContext.Elements.Add(new TextObject($":warning: *Step has multiple events reported. Will only show first 5.* :warning:"));
+				slackBlockContainer.Blocks.Add(limitContext);
+			}
+
+			foreach (LogData log in logData.Take(MaxEventMessages))
+			{
+				ContextBlock logContext = new ContextBlock();
+				logContext.Elements.Add(new TextObject($"```{log.Message}```"));
+				slackBlockContainer.Blocks.Add(logContext);
+
+				if (log.LogId is not null)
+				{
+					string lineIndexLink = (log.LineIndex is not null) ? $"?lineIndex={log.LineIndex}" : String.Empty;
+					string logUrl = new Uri(_serverInfo.DashboardUrl, $"log/{log.LogId}{lineIndexLink}").ToString();
+
+					ContextBlock detailsContext = new ContextBlock();
+					detailsContext.Elements.Add(new TextObject($"<{logUrl}|View Job Step Log>"));
+					slackBlockContainer.Blocks.Add(detailsContext);
+				}
+			}
+		}
+
+		/// <summary>
+		/// 
+		/// </summary>
+		/// <param name="step">Information about the step</param>
+		/// <param name="badge">Slack emoji associated with the notification</param>
+		/// <param name="logData">List of our log information</param>
+		/// <param name="underlyingStep">Optional information about the step's dependent step</param>
+		/// <returns>Generated <see cref="SlackMessage"/> if successful, otherwise returns null</returns>
+		SlackMessage? GenerateSlackThreadMessage(IJobStep step, string badge, List<LogData> logData, IJobStep? underlyingStep = null)
+		{
+			if (logData.Count == 0)
+			{
+				return null;
+			}
+
+			List<LogData> data = logData.Where(x => x.Severity == LogEventSeverity.Error).ToList();
+			if (!data.Any())
+			{
+				data = logData.Where(x => x.Severity == LogEventSeverity.Warning).ToList();
+			}
+
+			string jobStepLink = new Uri($"{_serverInfo.DashboardUrl}job/{step.Job.Id}?step={step.Id}").ToString();
+
+			SlackMessage alertMessage = new SlackMessage();
+
+			ContextBlock alertStepBlock = new ContextBlock();
+			alertStepBlock.Elements.Add(new PlainTextObject(badge, true));
+			alertStepBlock.Elements.Add(new TextObject($"<{jobStepLink}|*{step.Name}*>"));
+			alertStepBlock.Elements.Add(new PlainTextObject(badge, true));
+			alertMessage.Blocks.Add(alertStepBlock);
+
+			// If an underlying step was provided then the error occurred further up the job and impacted the current step
+			if (underlyingStep is not null)
+			{
+				ContextBlock dependencyContext = new ContextBlock();
+				dependencyContext.Elements.Add(new TextObject("*Step was marked as skipped due to errors in the following dependency*"));
+				alertMessage.Blocks.Add(dependencyContext);
+
+				// Slack Attachments are marked as being legacy and is recommended to be avoided
+				// The reason why Attachments are used is because this information is secondary to describe why a test had skipped
+				// We also wanted to differentiate the dependencies from the actual monitored node by indenting the dependency information
+				SlackAttachment attachment = new SlackAttachment();
+				attachment.Color = "#666666";
+
+				string dependencyJobStepLink = new Uri($"{_serverInfo.DashboardUrl}job/{step.Job.Id}?step={underlyingStep.Id}").ToString();
+
+				attachment.AddContext(new TextObject($"<{dependencyJobStepLink}|*{underlyingStep.Name}*>"));
+				AddLogDataContext(attachment, data);
+				alertMessage.Attachments.Add(attachment);
+			}
+			else
+			{
+				AddLogDataContext(alertMessage, data);
+			}
+
+			return alertMessage;
+		}
+
+		/// <summary>
+		/// Generates the title for the summary messages
+		/// </summary>
+		/// <param name="job">Information about the job</param>
+		/// <param name="branch">Name of the branch from the stream</param>
+		/// <param name="hordeJobURL">Optional URL of the job in Horde</param>
+		/// <returns>Title message using Slack's markdown format</returns>
+		static string GenerateSlackMessageTitle(IJob job, string branch, string? hordeJobURL = null)
+		{
+			string changelistDetails = $"CL {job.CommitId}";
+			if (job.PreflightCommitId is not null)
+			{
+				changelistDetails = $"Preflight CL {job.PreflightCommitId} against {changelistDetails}";
+			}
+
+			if (!String.IsNullOrEmpty(hordeJobURL))
+			{
+				return $"<{hordeJobURL}|*{branch} - {changelistDetails} - {job.Name}*>";
+			}
+
+			return $"{branch} - {changelistDetails} - {job.Name}";
+		}
+
+		/// <summary>
+		/// Adds a context to the Slack message with cancellation details
+		/// </summary>
+		/// <param name="message"><see cref="SlackMessage"/> to add the context to</param>
+		/// <param name="text">Information to provide in the context</param>
+		/// <param name="cancellationReason">Optional message explaining the cancellation</param>
+		static void AddCancelledJobSlackContext(SlackMessage message, string text, string? cancellationReason = null)
+		{
+			ContextBlock context = new ContextBlock();
+			context.Elements.Add(new TextObject(text));
+			message.Blocks.Add(context);
+
+			if (!String.IsNullOrEmpty(cancellationReason))
+			{
+				ContextBlock cancellationHeader = new ContextBlock();
+				cancellationHeader.Elements.Add(new TextObject($"*Cancellation Reason*"));
+				message.Blocks.Add(cancellationHeader);
+
+				ContextBlock cancellationContext = new ContextBlock();
+				cancellationContext.Elements.Add(new TextObject(cancellationReason));
+				message.Blocks.Add(cancellationContext);
+			}
+		}
+
+		/// <summary>
+		/// Create the job's summary message
+		/// </summary>
+		/// <param name="job">Information about the job</param>
+		/// <param name="branch">Name of the branch from the stream</param>
+		/// <param name="hordeJobURL">URL of the job in Horde</param>
+		/// <param name="jobStepNotificationStateDocuments">List of <see cref="JobStepNotificationStateDocument"/></param>
+		/// <param name="abortedByUserName">Optional username provided when cancelling a job</param>
+		/// <returns>Generated <see cref="SlackMessage"/></returns>
+		static SlackMessage CreateSlackSummaryNotification(IJob job, string branch, string hordeJobURL, List<JobStepNotificationStateDocument> jobStepNotificationStateDocuments, string? abortedByUserName = null)
+		{
+			DividerBlock divider = new DividerBlock();
+			SlackMessage message = new SlackMessage();
+
+			// Create our notification text
+			string notificationText = GenerateSlackMessageTitle(job, branch);
+			message.Text = $"{notificationText} Details...";
+
+			// Create a Block for the Slack Title
+			string messageTitle = GenerateSlackMessageTitle(job, branch, hordeJobURL);
+			ContextBlock jobContext = new ContextBlock();
+			jobContext.Elements.Add(new TextObject(messageTitle));
+			message.Blocks.Add(jobContext);
+
+			// If the job was manually aborted, then provide a message that catches the recipients eyes that the full information may not be available
+			if (!String.IsNullOrEmpty(abortedByUserName))
+			{
+				AddCancelledJobSlackContext(message, $"*{AlertBadge} Job was canceled by {abortedByUserName} and may be missing information {AlertBadge}*", job.CancellationReason);
+				message.Blocks.Add(divider);
+			}
+
+			int itemsInGroup = 0;
+			string currentGroup = String.Empty;
+			ContextBlock? header = null;
+			ContextBlock? stepContext = null;
+			foreach (JobStepNotificationStateDocument state in jobStepNotificationStateDocuments)
+			{
+				// Nothing to report on if the step is no longer in the job
+				if (!job.TryGetStep(state.JobStepId, out IJobStep? step))
+				{
+					continue;
+				}
+
+				if (currentGroup != state.Group || header is null)
+				{
+					// Check to see if we have an existing step context that needs to be added
+					if (stepContext is not null)
+					{
+						message.Blocks.Add(stepContext);
+						message.Blocks.Add(divider);
+						stepContext = null;
+					}
+
+					itemsInGroup = 1;
+					currentGroup = state.Group;
+					header = new ContextBlock();
+					header.Elements.Add(new TextObject($"*{currentGroup}*"));
+					message.Blocks.Add(header);
+				}
+
+				string jobStepLink = $"{hordeJobURL}?step={state.JobStepId}";
+
+				// Slack can only allow for 10 elements per block (https://api.slack.com/reference/block-kit/blocks#context) and we want to track when to create new ContextBlocks
+				bool bNeedsNewContext = (itemsInGroup % 10 == 0 || stepContext is null);
+				if (bNeedsNewContext)
+				{
+					if (stepContext is not null)
+					{
+						message.Blocks.Add(stepContext);
+					}
+					stepContext = new ContextBlock();
+				}
+				stepContext?.Elements.Add(new TextObject($"<{jobStepLink}|{state.Badge} {state.TargetPlatform}>"));
+			}
+
+			// Check to see if we have an existing step context that needs to be added
+			if (stepContext is not null)
+			{
+				message.Blocks.Add(stepContext);
+				message.Blocks.Add(divider);
+			}
+
+			return message;
+		}
+
+		#endregion Slack
 	}
 }
diff --git a/Engine/Source/Programs/Horde/Plugins/Experimental/HordeServer.Experimental/Slack/SlackConfig.cs b/Engine/Source/Programs/Horde/Plugins/Experimental/HordeServer.Experimental/Slack/SlackConfig.cs
new file mode 100644
--- /dev/null
+++ b/Engine/Source/Programs/Horde/Plugins/Experimental/HordeServer.Experimental/Slack/SlackConfig.cs
@@ -0,0 +1,95 @@
+// Copyright Epic Games, Inc. All Rights Reserved.
+
+using System.Text.RegularExpressions;
+using EpicGames.Horde.Jobs.Templates;
+using EpicGames.Horde.Streams;
+using HordeServer.Plugins;
+
+namespace HordeServer
+{
+	/// <summary>
+	/// Configuration for the experimental Slack notifications
+	/// </summary>
+	public class SlackConfig : IPluginConfig
+	{
+		public List<JobNotificationConfig> JobNotifications { get; set; } = new List<JobNotificationConfig>();
+	
+		private Lookup<string, JobNotificationConfig>? _notificationLookup = null;
+	
+		/// <inheritdoc/>
+		public void PostLoad(PluginConfigOptions configOptions)
+		{
+			_notificationLookup = (Lookup<string, JobNotificationConfig>)JobNotifications.ToLookup(item => $"{item.Template.ToString()}");
+		}
+	
+		public List<JobNotificationConfig>? GetNotificationConfigs(string jobName, TemplateId templateId, StreamId streamId)
+		{
+			if (_notificationLookup is null)
+			{
+				// Add error logging
+				return null;
+			}
+
+			IEnumerable<JobNotificationConfig> foundGroups = _notificationLookup[templateId.ToString()];
+			if (!foundGroups.Any())
+			{
+				// Add info logging
+				return null;
+			}
+
+			List<JobNotificationConfig> configs = new List<JobNotificationConfig>();
+			foreach (JobNotificationConfig group in foundGroups)
+			{
+				if (group.Enabled && group.Streams.Contains(streamId))
+				{
+					configs.Add(group);
+				}
+			}
+
+			return configs;
+		}
+	}
+	
+	public class JobNotificationConfig
+	{
+		public bool Enabled { get; set; } = false;
+		public TemplateId Template { get; set; }
+		public string PlatformPattern { get; set; } = String.Empty;
+		public List<NotificationFormatConfig> NotificationFormats { get; set; } = new List<NotificationFormatConfig>();
+		public List<string> Channels { get; set; } = new List<string>();
+		public List<StreamId> Streams { get; set; } = new List<StreamId>();
+
+		public bool GetStepDetails(string stepName, out string? group, out string? platform)
+		{
+			group = null;
+			platform = null;
+
+			foreach (NotificationFormatConfig formatConfig in NotificationFormats)
+			{
+				Match monitoredPattern = Regex.Match(stepName, formatConfig.StepPattern);
+				if (monitoredPattern.Success)
+				{
+					// Fetch our platform from the step name
+					string platformPattern = !String.IsNullOrEmpty(formatConfig.AlternativePlatformPattern) ? formatConfig.AlternativePlatformPattern : PlatformPattern;
+					Match targetPlatformPattern = Regex.Match(stepName, platformPattern);
+					if (targetPlatformPattern.Success)
+					{
+						group = formatConfig.Group;
+						platform = targetPlatformPattern.Groups[0].Value;
+
+						return true;
+					}
+				}
+			}
+
+			return false;
+		}
+	}
+
+	public class NotificationFormatConfig
+	{
+		public string Group { get; set; } = String.Empty;
+		public string StepPattern { get; set; } = String.Empty;
+		public string? AlternativePlatformPattern { get; set; }
+	}
+}
