Horde: Updating schema documentation

#rnx

diff --git a/Engine/Source/Programs/Horde/Docs/Config/Schema/AclActions.md b/Engine/Source/Programs/Horde/Docs/Config/Schema/AclActions.md
--- a/Engine/Source/Programs/Horde/Docs/Config/Schema/AclActions.md
+++ b/Engine/Source/Programs/Horde/Docs/Config/Schema/AclActions.md
@@ -605,6 +605,12 @@
 | `UpdateReplicator` |  |
 | `ViewReplicator` |  |
 
+## Secrets
+
+| Name | Description |
+| ---- | ----------- |
+| `ViewSecret` | View a credential |
+
 ## ServiceAccounts
 
 | Name | Description |
@@ -777,7 +783,7 @@
 
 | Name | Description |
 | ---- | ----------- |
-| `ViewSecret` | View a credential |
+| `ViewSecret` |  |
 
 ## ServiceAccounts
 
diff --git a/Engine/Source/Programs/Horde/Docs/Config/Schema/Globals.md b/Engine/Source/Programs/Horde/Docs/Config/Schema/Globals.md
--- a/Engine/Source/Programs/Horde/Docs/Config/Schema/Globals.md
+++ b/Engine/Source/Programs/Horde/Docs/Config/Schema/Globals.md
@@ -60,8 +60,8 @@ Name | Description
 `build` | [BuildConfig](#buildconfig)<br>Configuration for the build plugin
 `compute` | [ComputeConfig](#computeconfig)<br>Configuration for the compute plugin
 `ddc` | [EmptyPluginConfig](#emptypluginconfig)<br>Configuration for the ddc plugin
-`experimental` | [ExperimentalConfig](#experimentalconfig)<br>Configuration for the experimental plugin
 `secrets` | [SecretsConfig](#secretsconfig)<br>Configuration for the secrets plugin
+`experimental` | [ExperimentalConfig](#experimentalconfig)<br>Configuration for the experimental plugin
 `storage` | [StorageConfig](#storageconfig)<br>Configuration for the storage plugin
 `symbols` | [SymbolsConfig](#symbolsconfig)<br>Configuration for the symbols plugin
 `tools` | [ToolsConfig](#toolsconfig)<br>Configuration for the tools plugin
@@ -421,11 +421,6 @@ Name | Description
 Empty implementation of
 
 
-## ExperimentalConfig
-
-Configuration for the experimental plugin
-
-
 ## SecretsConfig
 
 Configuration for the secrets system
@@ -465,6 +460,41 @@ Name | Description
 `Text` | Secret is a plain text value which will be stored using the external secret key
 `Json` | Secret is a JSON formatted string containing key/value pairs
 
+## ExperimentalConfig
+
+Configuration for the experimental plugin
+
+Name | Description
+---- | -----------
+`slack` | [SlackConfig](#slackconfig)<br>
+
+## SlackConfig
+
+Configuration for the experimental Slack notifications
+
+Name | Description
+---- | -----------
+`jobNotifications` | [JobNotificationConfig](#jobnotificationconfig)`[]`<br>
+
+## JobNotificationConfig
+
+Name | Description
+---- | -----------
+`enabled` | `boolean`<br>
+`template` | `string`<br>
+`platformPattern` | `string`<br>
+`notificationFormats` | [NotificationFormatConfig](#notificationformatconfig)`[]`<br>
+`channels` | `string[]`<br>
+`streams` | `string[]`<br>
+
+## NotificationFormatConfig
+
+Name | Description
+---- | -----------
+`group` | `string`<br>
+`stepPattern` | `string`<br>
+`alternativePlatformPattern` | `string`<br>
+
 ## StorageConfig
 
 Configuration for storage
