Addressing Horde analyzer warnings
- Removed unused functions and variables added during development
- Added documentation of public methods and properties
- Fixed unused variable warning that would could cause an issue if an email address was used for a user
- Adjusted logging

#rnx

diff --git a/Engine/Source/Programs/Horde/Plugins/Build/HordeServer.Build/Notifications/NotificationService.cs b/Engine/Source/Programs/Horde/Plugins/Build/HordeServer.Build/Notifications/NotificationService.cs
--- a/Engine/Source/Programs/Horde/Plugins/Build/HordeServer.Build/Notifications/NotificationService.cs
+++ b/Engine/Source/Programs/Horde/Plugins/Build/HordeServer.Build/Notifications/NotificationService.cs
@@ -755,30 +755,6 @@ namespace HordeServer.Notifications
 			_logger.LogDebug("Finished sending notifications for step {JobId}:{BatchId}:{StepId}", job.Id, batchId, stepId);
 		}
 
-		private async Task<List<ILogEventData>?> FetchLogEventsAsync(IJobStep step, CancellationToken cancellationToken)
-		{
-			if (step.LogId is null)
-			{
-				return null;
-			}
-
-			ILog? log = await _logCollection.GetAsync(step.LogId.Value, cancellationToken);
-			if (log is null)
-			{
-				return null;
-			}
-
-			List<ILogAnchor> jobStepEvents = await log.GetAnchorsAsync(cancellationToken: cancellationToken);
-			List<ILogEventData> jobStepEventData = new List<ILogEventData>();
-			foreach (ILogAnchor logEvent in jobStepEvents)
-			{
-				ILogEventData eventData = await logEvent.GetDataAsync(cancellationToken);
-				jobStepEventData.Add(eventData);
-			}
-
-			return jobStepEventData;
-		}
-
 		private async Task SendAllLabelNotificationsAsync(IJob job, IReadOnlyList<(LabelState State, LabelOutcome Outcome)> oldLabelStates, IReadOnlyList<(LabelState, LabelOutcome)> newLabelStates, CancellationToken cancellationToken)
 		{
 			IGraph? graph = await _graphCollection.GetAsync(job.GraphHash, cancellationToken);
diff --git a/Engine/Source/Programs/Horde/Plugins/Experimental/HordeServer.Experimental/ExperimentalConfig.cs b/Engine/Source/Programs/Horde/Plugins/Experimental/HordeServer.Experimental/ExperimentalConfig.cs
--- a/Engine/Source/Programs/Horde/Plugins/Experimental/HordeServer.Experimental/ExperimentalConfig.cs
+++ b/Engine/Source/Programs/Horde/Plugins/Experimental/HordeServer.Experimental/ExperimentalConfig.cs
@@ -9,6 +9,9 @@ namespace HordeServer
 	/// </summary>
 	public class ExperimentalConfig : IPluginConfig
 	{
+		/// <summary>
+		/// Configuration for sending experimental Job Summary Slack notifications
+		/// </summary>
 		public SlackConfig Slack { get; set; } = new SlackConfig();
 
 		/// <inheritdoc/>
diff --git a/Engine/Source/Programs/Horde/Plugins/Experimental/HordeServer.Experimental/Notifications/Sinks/ExperimentalSlackNotificationSink.cs b/Engine/Source/Programs/Horde/Plugins/Experimental/HordeServer.Experimental/Notifications/Sinks/ExperimentalSlackNotificationSink.cs
--- a/Engine/Source/Programs/Horde/Plugins/Experimental/HordeServer.Experimental/Notifications/Sinks/ExperimentalSlackNotificationSink.cs
+++ b/Engine/Source/Programs/Horde/Plugins/Experimental/HordeServer.Experimental/Notifications/Sinks/ExperimentalSlackNotificationSink.cs
@@ -219,7 +219,7 @@ namespace HordeServer.Notifications.Sinks
 			List<JobNotificationStateDocument>? jobNotificationStates = await GetJobNotificationStatesAsync(job.Id, cancellationToken);
 			if (jobNotificationStates is null)
 			{
-				_logger.LogError("No job notifications found for job id '{JobId}'.", job.Id);
+				_logger.LogInformation("No job notifications found for job id '{JobId}'.", job.Id);
 				return;
 			}
 
@@ -327,7 +327,7 @@ namespace HordeServer.Notifications.Sinks
 		{
 			SlackConfig currentConfig = _experimentalConfig.CurrentValue.Slack;
 
-			List<JobNotificationConfig>? configs = currentConfig.GetNotificationConfigs(job.Name, job.TemplateId, job.StreamId);
+			List<JobNotificationConfig>? configs = currentConfig.GetNotificationConfigs(job.TemplateId, job.StreamId);
 			if (configs is null || !configs.Any())
 			{
 				// Log that there is nothing to do here
@@ -394,7 +394,7 @@ namespace HordeServer.Notifications.Sinks
 			List<JobNotificationStateDocument>? jobNotificationStates = await GetJobNotificationStatesAsync(job.Id, cancellationToken);
 			if (jobNotificationStates is null)
 			{
-				_logger.LogError("No job notifications found for job id '{JobId}'.", job.Id);
+				_logger.LogWarning("No job notifications found for job id '{JobId}' even though job step '{StepId}' was updated.", job.Id, step.Id);
 				return;
 			}
 
@@ -659,7 +659,7 @@ namespace HordeServer.Notifications.Sinks
 			if (states.Count != 1)
 			{
 				// Log error that multiple were found and could cause issues
-				_logger.LogError("Found {Count} records when fetching for {JobId}:{TemplateId}:{Recipient}. Using first record which could cause issues.", states.Count, jobId, templateId, recipient);
+				_logger.LogWarning("Found {Count} records when fetching for {JobId}:{TemplateId}:{Recipient}. Using first record which could cause issues.", states.Count, jobId, templateId, recipient);
 			}
 
 			return states[0];
@@ -715,7 +715,7 @@ namespace HordeServer.Notifications.Sinks
 			}
 			else
 			{
-				_logger.LogInformation("Updatedmessage {StateId} (jobId: {JobId}, templateId: {TemplateId}, recipient: {Recipient}, jobStepId: {JobStepId}, group: {Group}, platform: {Platform}, badge: {Badge}, messageId: {MessageId})", state.Id, state.JobId, state.TemplateId, state.Recipient, state.JobStepId, state.Group, state.TargetPlatform, state.Badge, state.MessageId);
+				_logger.LogInformation("Updated message {StateId} (jobId: {JobId}, templateId: {TemplateId}, recipient: {Recipient}, jobStepId: {JobStepId}, group: {Group}, platform: {Platform}, badge: {Badge}, messageId: {MessageId})", state.Id, state.JobId, state.TemplateId, state.Recipient, state.JobStepId, state.Group, state.TargetPlatform, state.Badge, state.MessageId);
 			}
 
 			return (state, state.Id == newId);
@@ -783,7 +783,6 @@ namespace HordeServer.Notifications.Sinks
 
 		async Task SendSetupFailedSlackMessageAsync(IJob job, IJobStep step, string branch, List<LogData> jobStepEventData, List<JobNotificationConfig> jobNotificationConfigs, CancellationToken cancellationToken = default)
 		{
-			DividerBlock divider = new DividerBlock();
 			SlackMessage message = new SlackMessage();
 
 			string jobUrlLink = new Uri($"{_serverInfo.DashboardUrl}job/{job.Id}").ToString();
@@ -817,7 +816,7 @@ namespace HordeServer.Notifications.Sinks
 						channel = userInfo.Id;
 					}
 
-					_ = await _slackClient.PostMessageAsync(recipient, message, cancellationToken);
+					_ = await _slackClient.PostMessageAsync(channel, message, cancellationToken);
 				}
 			}
 
@@ -840,7 +839,7 @@ namespace HordeServer.Notifications.Sinks
 			SlackMessage? message = await GenerateSlackSummaryMessageAsync(job, branch, hordeJobURL, templateId, recipient, abortedByUserName, cancellationToken);
 			if (message is null)
 			{
-				_logger.LogError("Unable to generate Job summary notification for job '{JobId}'", job.Id);
+				_logger.LogWarning("Unable to generate Job summary notification for job '{JobId}' for the template '{TemplateId}'", job.Id, templateId);
 				return;
 			}
 
@@ -926,7 +925,7 @@ namespace HordeServer.Notifications.Sinks
 					JobNotificationStateDocument? jobNotificationState = await GetJobNotificationStateAsync(jobId, jobStepNotificationState.TemplateId, jobStepNotificationState.Recipient, cancellationToken);
 					if (jobNotificationState is null)
 					{
-						_logger.LogError("No message id found for job '{JobId}' with template '{TemplateId}' for thread to be posted.", jobId, jobStepNotificationState.TemplateId);
+						_logger.LogInformation("No message id found for job '{JobId}' with template '{TemplateId}' for thread to be posted.", jobId, jobStepNotificationState.TemplateId);
 						continue;
 					}
 					threadMessageId = await SendSlackThreadMessageAsync(jobNotificationState.MessageId, jobStepNotificationState, threadMessage, cancellationToken);
diff --git a/Engine/Source/Programs/Horde/Plugins/Experimental/HordeServer.Experimental/Slack/SlackConfig.cs b/Engine/Source/Programs/Horde/Plugins/Experimental/HordeServer.Experimental/Slack/SlackConfig.cs
--- a/Engine/Source/Programs/Horde/Plugins/Experimental/HordeServer.Experimental/Slack/SlackConfig.cs
+++ b/Engine/Source/Programs/Horde/Plugins/Experimental/HordeServer.Experimental/Slack/SlackConfig.cs
@@ -12,6 +12,9 @@ namespace HordeServer
 	/// </summary>
 	public class SlackConfig : IPluginConfig
 	{
+		/// <summary>
+		/// List of configurations for sending out job summary notifications
+		/// </summary>
 		public List<JobNotificationConfig> JobNotifications { get; set; } = new List<JobNotificationConfig>();
 	
 		private Lookup<string, JobNotificationConfig>? _notificationLookup = null;
@@ -22,7 +25,13 @@ namespace HordeServer
 			_notificationLookup = (Lookup<string, JobNotificationConfig>)JobNotifications.ToLookup(item => $"{item.Template.ToString()}");
 		}
 	
-		public List<JobNotificationConfig>? GetNotificationConfigs(string jobName, TemplateId templateId, StreamId streamId)
+		/// <summary>
+		/// Fetches available job summary configurations from the criteria specified
+		/// </summary>
+		/// <param name="templateId">Identifier of the template for the specified job</param>
+		/// <param name="streamId">Identifier of the stream for the specified job</param>
+		/// <returns>List of <see cref="JobNotificationConfig"/> or null if no configurations for the criteria have been found</returns>
+		public List<JobNotificationConfig>? GetNotificationConfigs(TemplateId templateId, StreamId streamId)
 		{
 			if (_notificationLookup is null)
 			{
@@ -50,15 +59,48 @@ namespace HordeServer
 		}
 	}
 	
+	/// <summary>
+	/// Configuration for the job notification
+	/// </summary>
 	public class JobNotificationConfig
 	{
+		/// <summary>
+		/// Whether or not the current configuration is enabled
+		/// </summary>
 		public bool Enabled { get; set; } = false;
+
+		/// <summary>
+		/// Identifier of the template this notification applies to
+		/// </summary>
 		public TemplateId Template { get; set; }
+
+		/// <summary>
+		/// Regex pattern of the platforms this notification will include
+		/// </summary>
 		public string PlatformPattern { get; set; } = String.Empty;
+
+		/// <summary>
+		/// List of formats to be applied for this notification
+		/// </summary>
 		public List<NotificationFormatConfig> NotificationFormats { get; set; } = new List<NotificationFormatConfig>();
+
+		/// <summary>
+		/// List of Slack channel or user identifiers the notification will be sent to
+		/// </summary>
 		public List<string> Channels { get; set; } = new List<string>();
+
+		/// <summary>
+		/// List of streams this notification will be used for
+		/// </summary>
 		public List<StreamId> Streams { get; set; } = new List<StreamId>();
 
+		/// <summary>
+		/// Gets the group and platform information of the corresponding step within the list of notification formats
+		/// </summary>
+		/// <param name="stepName">Name of the job step to test against</param>
+		/// <param name="group">Name of the notification group the job step belongs to or null if not found</param>
+		/// <param name="platform">Name of the platform associated with the job step or null if not found</param>
+		/// <returns>True if the job step is to be included in the notification or false otherwise</returns>
 		public bool GetStepDetails(string stepName, out string? group, out string? platform)
 		{
 			group = null;
@@ -86,10 +128,24 @@ namespace HordeServer
 		}
 	}
 
+	/// <summary>
+	/// Configuration for the format of the notification group
+	/// </summary>
 	public class NotificationFormatConfig
 	{
+		/// <summary>
+		/// Name of the group
+		/// </summary>
 		public string Group { get; set; } = String.Empty;
+
+		/// <summary>
+		/// Regex pattern to determine if a job step belongs within the group
+		/// </summary>
 		public string StepPattern { get; set; } = String.Empty;
+
+		/// <summary>
+		/// Optional regex pattern to override the <see cref="JobNotificationConfig.PlatformPattern"/>
+		/// </summary>
 		public string? AlternativePlatformPattern { get; set; }
 	}
 }
