Horde: Add a bit more information when can't resolve preflight stream
#submitool bypass
#skipci

diff --git a/Engine/Source/Programs/Horde/HordeDashboard/src/components/Preflight.tsx b/Engine/Source/Programs/Horde/HordeDashboard/src/components/Preflight.tsx
--- a/Engine/Source/Programs/Horde/HordeDashboard/src/components/Preflight.tsx
+++ b/Engine/Source/Programs/Horde/HordeDashboard/src/components/Preflight.tsx
@@ -79,7 +79,7 @@ export const PreflightRedirector: React.FC = () => {
    }
 
    if (!stream) {
-      setError(`Unable to resolve stream with name ${streamName}`);
+      setError(`Unable to resolve stream with name ${streamName}, please verify that this isn't a virtual stream and that you have access.`);
       return null;
    }
 
