Make Preferences, user icon, and Tools middle-clickable.
#tests Tested locally using horde-preview

diff --git a/Engine/Source/Programs/Horde/HordeDashboard/src/components/TopNav.tsx b/Engine/Source/Programs/Horde/HordeDashboard/src/components/TopNav.tsx
--- a/Engine/Source/Programs/Horde/HordeDashboard/src/components/TopNav.tsx
+++ b/Engine/Source/Programs/Horde/HordeDashboard/src/components/TopNav.tsx
@@ -343,7 +343,7 @@ const MenuButton: React.FC<{ props: IButtonProps, link?: string }> = ({ props, l
 };
 
 const ToolsButton: React.FunctionComponent<IButtonProps> = (props) => {
-   return <MenuButton props={props} />
+   return <MenuButton props={props} link="/tools"/>
 }
 
 
@@ -726,7 +726,9 @@ export const TopNav: React.FC<{ suppressServer?: boolean }> = observer(({ suppre
                {
                   key: 'settingsItem',
                   text: 'Preferences',
-                  onClick: () => {
+                  href: `/dashboard`,
+                  onClick: (ev: any) => {
+                     ev.preventDefault();
                      navigate("/dashboard");
                   }
                }],
@@ -827,8 +829,10 @@ export const TopNav: React.FC<{ suppressServer?: boolean }> = observer(({ suppre
                   <Stack style={{ width: "32px" }} onMouseEnter={() => setShowMenu(true)}
                      onMouseLeave={() => { setShowMenu(false) }} >
                      <div ref={divRef}>
-                        <Persona styles={{ root: { selectors: { ".ms-Persona-initials": { fontWeight: "unset", fontFamily: "Horde Open Sans SemiBold", cursor: "pointer" } } } }} imageShouldFadeIn={false} imageInitials={initials} imageUrl={dashboard.userImage32} size={PersonaSize.size32}
-                           onClick={() => { navigate("/dashboard"); }} />
+                        <Link to={"/dashboard"}>
+                           <Persona styles={{ root: { selectors: { ".ms-Persona-initials": { fontWeight: "unset", fontFamily: "Horde Open Sans SemiBold", cursor: "pointer" } } } }} 
+                              imageShouldFadeIn={false} imageInitials={initials} imageUrl={dashboard.userImage32} size={PersonaSize.size32} />
+                        </Link>
                      </div>
                      <ContextualMenu
                         items={menuItems}
