Horde: Improve agent telemetry display
#skipci

diff --git a/Engine/Source/Programs/Horde/HordeDashboard/src/components/HistoryModal.tsx b/Engine/Source/Programs/Horde/HordeDashboard/src/components/HistoryModal.tsx
--- a/Engine/Source/Programs/Horde/HordeDashboard/src/components/HistoryModal.tsx
+++ b/Engine/Source/Programs/Horde/HordeDashboard/src/components/HistoryModal.tsx
@@ -9,7 +9,7 @@ import { Sparklines, SparklinesLine, SparklinesReferenceLine } from "react-spark
 import backend from "../backend";
 import { agentStore } from "../backend/AgentStore";
 import { AgentData, GetAgentLeaseResponse, GetAgentSessionResponse, GetAgentTelemetrySampleResponse, GetJobResponse, JobStepBatchError, JobStepOutcome, LeaseData, SessionData, UpdateAgentRequest } from "../backend/Api";
-import dashboard from "../backend/Dashboard";
+import dashboard, { StatusColor } from "../backend/Dashboard";
 import { getNiceTime, getShortNiceTime } from "../base/utilities/timeUtils";
 import { getHordeStyling } from "../styles/Styles";
 import { getHordeTheme } from "../styles/theme";
@@ -543,7 +543,7 @@ const cpuSelections: IComboBoxOption[] = [
 
 const TelemetryPanel: React.FC<{}> = observer(({ }) => {
 
-   const [timeKey, setTimeKey] = useState("time_2_weeks");
+   const [timeKey, setTimeKey] = useState("time_2_days");
    const [cpuKey, setCPUKey] = useState("combined");
 
    useEffect(() => {
@@ -570,6 +570,7 @@ const TelemetryPanel: React.FC<{}> = observer(({ }) => {
       </Stack>;
    }
 
+   const statusColors = dashboard.getStatusColors();
 
    const sparkWidth = 1100;
    const data = (currentData[0] as GetAgentTelemetrySampleResponse[]);
@@ -592,14 +593,21 @@ const TelemetryPanel: React.FC<{}> = observer(({ }) => {
    const minTime = getShortNiceTime(min, true);
    const maxTime = getShortNiceTime(max, true);
 
-
-
    const ramData = data.map(d => {
       return d.usedRam / d.totalRam;
    })
 
+   const diskData = data.map(d => {
+      if (typeof (d.totalDisk) === "number" && d.totalDisk) {
+         return 1.0 - (d.freeDisk / d.totalDisk)
+      }
+
+      return 0;
+   })
+
    let cpuText = "";
    let ramText = "";
+   let diskText = "";
 
    let cpuLabel = "";
    if (cpuKey == "user") cpuLabel = "CPU - User";
@@ -614,11 +622,32 @@ const TelemetryPanel: React.FC<{}> = observer(({ }) => {
          if (v.startsWith("CPU=")) {
             cpuText = `${v.replace("CPU=", "")}`;
          }
-   
-         if (v.startsWith("RAM=")) {
+
+         if (v.startsWith("RAM=")) {            
             ramText = `${v.replace("RAM=", "")} GB`;
          }
-      })   
+
+         if (v.startsWith("DiskTotalSize=")) {
+            const elements = v.split("=");
+            if (elements.length === 2) {
+
+               function formatBytes(bytes: number, decimals = 0) {
+                  if (!+bytes) return '0 Bytes'
+
+                  const k = 1024
+                  const dm = decimals < 0 ? 0 : decimals
+                  const sizes = ['Bytes', 'KB', 'MB', 'GB', 'TB', 'PB', 'EB', 'ZB', 'YB']
+
+                  const i = Math.floor(Math.log(bytes) / Math.log(k))
+
+                  return `${parseFloat((bytes / Math.pow(k, i)).toFixed(dm))} ${sizes[i]}`
+               }
+
+               diskText = formatBytes(parseInt(elements[1]))
+            }
+         }
+
+      })
    }
 
    function relativeCoords(event: any, id: string) {
@@ -693,10 +722,8 @@ const TelemetryPanel: React.FC<{}> = observer(({ }) => {
          </Stack>
       </Stack>
       <Stack verticalAlign="center" verticalFill>
-         <Stack horizontal>
-            <Label>{cpuLabel}</Label>
-            <Stack grow></Stack>
-            {!!cpuText && <Label>{cpuText}</Label>}
+         <Stack>
+            <Label>{cpuLabel} Utilization {cpuText ? `- ${cpuText}` : ""}</Label>            
          </Stack>
          <div onMouseLeave={() => state.setLeaseTooltip(undefined)} onMouseMove={(ev) => { relativeCoords(ev, "cpu") }} onClick={() => {
 
@@ -712,15 +739,17 @@ const TelemetryPanel: React.FC<{}> = observer(({ }) => {
 
             <Sparklines width={sparkWidth} height={128} data={cpuData} max={100} style={{ backgroundColor: dashboard.darktheme ? "#060709" : "#F3F2F1", padding: 8, border: "solid 1px #181A1B" }}>
                <SparklinesLine color={dashboard.darktheme ? "lightblue" : "#1E90FF"} />
-               <SparklinesReferenceLine type="avg" />
+               <SparklinesReferenceLine type="custom" value={0} style={{ stroke: statusColors.get(StatusColor.Failure) }} />
+               <SparklinesReferenceLine type="custom" value={32} style={{ stroke: statusColors.get(StatusColor.Warnings) }} />
+               <SparklinesReferenceLine type="custom" value={64} style={{ stroke: statusColors.get(StatusColor.Success) }} />
+               <SparklinesReferenceLine type="custom" value={96} style={{ stroke: statusColors.get(StatusColor.Success) }} />
+
             </Sparklines>
          </div>
       </Stack>
       <Stack verticalAlign="center" verticalFill>
-         <Stack horizontal>
-            <Label>RAM</Label>
-            <Stack grow></Stack>
-            {!!ramText && <Label>{ramText}</Label>}
+         <Stack>
+            <Label>RAM Utilization {ramText ? `- ${ramText}` : ""}</Label>
          </Stack>
          <div onMouseLeave={() => state.setLeaseTooltip(undefined)} onMouseMove={(ev) => { relativeCoords(ev, "ram") }} onClick={() => {
 
@@ -735,11 +764,40 @@ const TelemetryPanel: React.FC<{}> = observer(({ }) => {
             </div>
             <Sparklines width={sparkWidth} height={128} data={ramData} max={1} style={{ backgroundColor: dashboard.darktheme ? "#060709" : "#F3F2F1", padding: 8, border: "solid 1px #181A1B" }}>
                <SparklinesLine color={dashboard.darktheme ? "lightblue" : "#1E90FF"} />
-               <SparklinesReferenceLine type="avg" />
+               <SparklinesReferenceLine type="custom" value={0} style={{ stroke: statusColors.get(StatusColor.Failure) }} />
+               <SparklinesReferenceLine type="custom" value={32} style={{ stroke: statusColors.get(StatusColor.Warnings) }} />
+               <SparklinesReferenceLine type="custom" value={64} style={{ stroke: statusColors.get(StatusColor.Success) }} />
+               <SparklinesReferenceLine type="custom" value={96} style={{ stroke: statusColors.get(StatusColor.Success) }} />
+
             </Sparklines>
          </div>
       </Stack>
 
+      {!!diskText && <Stack verticalAlign="center" verticalFill>
+         <Stack>
+            <Label>Storage Utilization - {diskText}</Label>
+         </Stack>
+         <div onMouseLeave={() => state.setLeaseTooltip(undefined)} onMouseMove={(ev) => { relativeCoords(ev, "disk") }} onClick={() => {
+
+            if (!state.currentLeaseTooltip?.frozen) {
+               if (state.currentLeaseTooltip) {
+                  state.setLeaseTooltip({ ...state.currentLeaseTooltip, frozen: true });
+               }
+            }
+         }}>
+            <div style={{ position: "relative" }}>
+               <TelemetryTooltip id="disk" />
+            </div>
+            <Sparklines width={sparkWidth} height={128} data={diskData} max={1} style={{ backgroundColor: dashboard.darktheme ? "#060709" : "#F3F2F1", padding: 8, border: "solid 1px #181A1B" }}>
+               <SparklinesLine color={dashboard.darktheme ? "lightblue" : "#1E90FF"} />
+               <SparklinesReferenceLine type="custom" value={0} style={{ stroke: statusColors.get(StatusColor.Failure) }} />
+               <SparklinesReferenceLine type="custom" value={32} style={{ stroke: statusColors.get(StatusColor.Warnings) }} />
+               <SparklinesReferenceLine type="custom" value={64} style={{ stroke: statusColors.get(StatusColor.Success) }} />
+               <SparklinesReferenceLine type="custom" value={96} style={{ stroke: statusColors.get(StatusColor.Success) }} />
+            </Sparklines>
+         </div>
+      </Stack>}
+
    </Stack>
 })
 
