Horde: Improvements to analytics top/bottom N capture

diff --git a/Engine/Source/Programs/Horde/Plugins/Analytics/HordeServer.Analytics/Telemetry/Metrics/MetricCollection.cs b/Engine/Source/Programs/Horde/Plugins/Analytics/HordeServer.Analytics/Telemetry/Metrics/MetricCollection.cs
--- a/Engine/Source/Programs/Horde/Plugins/Analytics/HordeServer.Analytics/Telemetry/Metrics/MetricCollection.cs
+++ b/Engine/Source/Programs/Horde/Plugins/Analytics/HordeServer.Analytics/Telemetry/Metrics/MetricCollection.cs
@@ -352,7 +352,7 @@ namespace HordeServer.Telemetry.Metrics
 
 			List<MetricDocument> results = await _metrics.Find(filter).SortByDescending(x => x.Time).Limit(maxResults).ToListAsync(cancellationToken);
 
-			HashSet<ObjectId> filterIds = new HashSet<ObjectId>();
+			HashSet<string> filterGroups = new HashSet<string>();
 			if (_analyticsConfig.CurrentValue.TryGetTelemetryStore(telemetryStoreId, out TelemetryStoreConfig? telemetryStoreConfig))
 			{
 				foreach (MetricId metricId in metricIds)
@@ -362,27 +362,26 @@ namespace HordeServer.Telemetry.Metrics
 					{
 						if (metricConfig.TopN > 0 || metricConfig.BottomN > 0)
 						{
-							foreach (IGrouping<DateTime, MetricDocument> metrics in results.Where(x =>  x.MetricId == metricId).GroupBy(x => x.Time, x => x))
+
+							List<string> ordered = results.Where(x => x.MetricId == metricId).GroupBy(x => x.Group, x => x).Select(x => (x.Key, x.Sum(x => x.Value))).OrderByDescending(x => x.Item2).Select(x => x.Item1).ToList();
+
+							if (ordered.Count <= metricConfig.TopN && ordered.Count <= metricConfig.BottomN)
 							{
-								int count = metrics.Count();
-								if (count <= metricConfig.TopN && count <= metricConfig.BottomN)
-								{
-									continue;
-								}
+								continue;
+							}
 
-								if (metricConfig.TopN > 0)
+							if (metricConfig.TopN > 0)
+							{
+								foreach (string groupName in ordered.TakeLast(ordered.Count - metricConfig.TopN))
 								{
-									foreach (MetricDocument m in metrics.OrderBy(x => x.Value).Take(count - metricConfig.TopN))
-									{
-										filterIds.Add(m.Id);
-									}
+									filterGroups.Add(groupName);
 								}
-								else
+							}
+							else
+							{
+								foreach (string groupName in ordered.Take(ordered.Count - metricConfig.BottomN))
 								{
-									foreach (MetricDocument m in metrics.OrderBy(x => x.Value).TakeLast(count - metricConfig.BottomN))
-									{
-										filterIds.Add(m.Id);
-									}
+									filterGroups.Add(groupName);
 								}
 							}
 						}
@@ -390,7 +389,7 @@ namespace HordeServer.Telemetry.Metrics
 				}				 
 			}
 
-			return results.Where(x => !filterIds.Contains(x.Id)).ToList().ConvertAll<MetricDocument, IMetric>(x => x);
+			return results.Where(x => !filterGroups.Contains(x.Group)).ToList().ConvertAll<MetricDocument, IMetric>(x => x);
 		}
 	}
 }
