Horde: Add support for enumerating analytics metrics and providing multi-group filters

diff --git a/Engine/Source/Programs/Horde/Plugins/Analytics/HordeServer.Analytics/Telemetry/Metrics/IMetricCollection.cs b/Engine/Source/Programs/Horde/Plugins/Analytics/HordeServer.Analytics/Telemetry/Metrics/IMetricCollection.cs
--- a/Engine/Source/Programs/Horde/Plugins/Analytics/HordeServer.Analytics/Telemetry/Metrics/IMetricCollection.cs
+++ b/Engine/Source/Programs/Horde/Plugins/Analytics/HordeServer.Analytics/Telemetry/Metrics/IMetricCollection.cs
@@ -25,10 +25,11 @@ namespace HordeServer.Telemetry.Metrics
 		/// <param name="metricIds">Metrics to search for</param>
 		/// <param name="minTime">Start of the time period to query</param>
 		/// <param name="maxTime">End of the time period to query</param>
-		/// <param name="group">Grouping key to filter results</param>
+		/// <param name="groups">Grouping keys to filter results</param>
+		/// <param name="enumerate">Whether to enumerate results, avoiding additional filtering</param> 
 		/// <param name="maxResults">Maximum number of results to return</param>
 		/// <param name="cancellationToken">Cancellation token for the operation</param>
-		Task<List<IMetric>> FindAsync(TelemetryStoreId storeId, MetricId[] metricIds, DateTime? minTime = null, DateTime? maxTime = null, string? group = null, int maxResults = 50, CancellationToken cancellationToken = default);
+		Task<List<IMetric>> FindAsync(TelemetryStoreId storeId, MetricId[] metricIds, DateTime? minTime = null, DateTime? maxTime = null, string[]? groups = null, bool enumerate = false, int maxResults = 50, CancellationToken cancellationToken = default);
 
 		/// <summary>
 		/// Flush all the pending events to storage
diff --git a/Engine/Source/Programs/Horde/Plugins/Analytics/HordeServer.Analytics/Telemetry/Metrics/MetricCollection.cs b/Engine/Source/Programs/Horde/Plugins/Analytics/HordeServer.Analytics/Telemetry/Metrics/MetricCollection.cs
--- a/Engine/Source/Programs/Horde/Plugins/Analytics/HordeServer.Analytics/Telemetry/Metrics/MetricCollection.cs
+++ b/Engine/Source/Programs/Horde/Plugins/Analytics/HordeServer.Analytics/Telemetry/Metrics/MetricCollection.cs
@@ -329,7 +329,7 @@ namespace HordeServer.Telemetry.Metrics
 		}
 
 		/// <inheritdoc/>
-		public async Task<List<IMetric>> FindAsync(TelemetryStoreId telemetryStoreId, MetricId[] metricIds, DateTime? minTime = null, DateTime? maxTime = null, string? group = null, int maxResults = 50, CancellationToken cancellationToken = default)
+		public async Task<List<IMetric>> FindAsync(TelemetryStoreId telemetryStoreId, MetricId[] metricIds, DateTime? minTime = null, DateTime? maxTime = null, string[]? groups = null, bool enumerate = false, int maxResults = 50, CancellationToken cancellationToken = default)
 		{
 			FilterDefinition<MetricDocument> filter = FilterDefinition<MetricDocument>.Empty;
 			filter &= Builders<MetricDocument>.Filter.Eq(x => x.TelemetryStoreId, telemetryStoreId);
@@ -345,48 +345,51 @@ namespace HordeServer.Telemetry.Metrics
 
 			filter &= Builders<MetricDocument>.Filter.In(x => x.MetricId, metricIds);
 
-			if (group != null)
+			if (groups != null)
 			{
-				filter &= Builders<MetricDocument>.Filter.Eq(x => x.Group, group);
+				filter &= Builders<MetricDocument>.Filter.In(x => x.Group, groups);
 			}
 
 			List<MetricDocument> results = await _metrics.Find(filter).SortByDescending(x => x.Time).Limit(maxResults).ToListAsync(cancellationToken);
 
 			HashSet<string> filterGroups = new HashSet<string>();
-			if (_analyticsConfig.CurrentValue.TryGetTelemetryStore(telemetryStoreId, out TelemetryStoreConfig? telemetryStoreConfig))
+			if (!enumerate)
 			{
-				foreach (MetricId metricId in metricIds)
+				if (_analyticsConfig.CurrentValue.TryGetTelemetryStore(telemetryStoreId, out TelemetryStoreConfig? telemetryStoreConfig))
 				{
-					MetricConfig? metricConfig;
-					if (telemetryStoreConfig.TryGetMetric(metricId, out metricConfig))
+					foreach (MetricId metricId in metricIds)
 					{
-						if (metricConfig.TopN > 0 || metricConfig.BottomN > 0)
+						MetricConfig? metricConfig;
+						if (telemetryStoreConfig.TryGetMetric(metricId, out metricConfig))
 						{
+							if (metricConfig.TopN > 0 || metricConfig.BottomN > 0)
+							{
 
-							List<string> ordered = results.Where(x => x.MetricId == metricId).GroupBy(x => x.Group, x => x).Select(x => (x.Key, x.Sum(x => x.Value))).OrderByDescending(x => x.Item2).Select(x => x.Item1).ToList();
+								List<string> ordered = results.Where(x => x.MetricId == metricId).GroupBy(x => x.Group, x => x).Select(x => (x.Key, x.Sum(x => x.Value))).OrderByDescending(x => x.Item2).Select(x => x.Item1).ToList();
 
-							if (ordered.Count <= metricConfig.TopN && ordered.Count <= metricConfig.BottomN)
-							{
-								continue;
-							}
+								if (ordered.Count <= metricConfig.TopN && ordered.Count <= metricConfig.BottomN)
+								{
+									continue;
+								}
 
-							if (metricConfig.TopN > 0)
-							{
-								foreach (string groupName in ordered.TakeLast(ordered.Count - metricConfig.TopN))
+								if (metricConfig.TopN > 0)
 								{
-									filterGroups.Add(groupName);
+									foreach (string groupName in ordered.TakeLast(ordered.Count - metricConfig.TopN))
+									{
+										filterGroups.Add(groupName);
+									}
 								}
-							}
-							else
-							{
-								foreach (string groupName in ordered.Take(ordered.Count - metricConfig.BottomN))
+								else
 								{
-									filterGroups.Add(groupName);
+									foreach (string groupName in ordered.Take(ordered.Count - metricConfig.BottomN))
+									{
+										filterGroups.Add(groupName);
+									}
 								}
 							}
 						}
 					}
-				}				 
+				}
 			}
 
 			return results.Where(x => !filterGroups.Contains(x.Group)).ToList().ConvertAll<MetricDocument, IMetric>(x => x);
diff --git a/Engine/Source/Programs/Horde/Plugins/Analytics/HordeServer.Analytics/Telemetry/TelemetryController.cs b/Engine/Source/Programs/Horde/Plugins/Analytics/HordeServer.Analytics/Telemetry/TelemetryController.cs
--- a/Engine/Source/Programs/Horde/Plugins/Analytics/HordeServer.Analytics/Telemetry/TelemetryController.cs
+++ b/Engine/Source/Programs/Horde/Plugins/Analytics/HordeServer.Analytics/Telemetry/TelemetryController.cs
@@ -48,7 +48,7 @@ namespace HordeServer.Telemetry
 		[Obsolete("Pass a telemetry store id in the route instead")]
 		public Task<ActionResult<List<GetTelemetryMetricsResponse>>> GetMetricsAsync([FromQuery] MetricId[] id, [FromQuery] DateTime? minTime = null, [FromQuery] DateTime? maxTime = null, [FromQuery] string? group = null, [FromQuery] int results = 50, CancellationToken cancellationToken = default)
 		{
-			return GetMetricsAsync(TelemetryStoreId.Default, id, minTime, maxTime, group, results, cancellationToken);
+			return GetMetricsAsync(TelemetryStoreId.Default, id, minTime, maxTime, group == null ? null : new string[] { group }, true, results, cancellationToken);
 		}
 
 		/// <summary>
@@ -58,13 +58,14 @@ namespace HordeServer.Telemetry
 		/// <param name="id">The metrics to query</param>
 		/// <param name="minTime">Minimum time interval to query</param>
 		/// <param name="maxTime">Maximum time interval to query</param>
-		/// <param name="group">Grouping key</param>
+		/// <param name="groups">Grouping keys</param>
+		/// <param name="enumerate">Whether to enumerate keys</param>
 		/// <param name="results">Number of results to return</param>
 		/// <param name="cancellationToken">Cancellation token for the operation</param>
 		[HttpGet]
 		[Authorize]
 		[Route("/api/v1/telemetry/{storeId}/metrics")]
-		public async Task<ActionResult<List<GetTelemetryMetricsResponse>>> GetMetricsAsync(TelemetryStoreId storeId, [FromQuery] MetricId[] id, [FromQuery] DateTime? minTime = null, [FromQuery] DateTime? maxTime = null, [FromQuery] string? group = null, [FromQuery] int results = 50, CancellationToken cancellationToken = default)
+		public async Task<ActionResult<List<GetTelemetryMetricsResponse>>> GetMetricsAsync(TelemetryStoreId storeId, [FromQuery] MetricId[] id, [FromQuery] DateTime? minTime = null, [FromQuery] DateTime? maxTime = null, [FromQuery(Name = "group")] string[]? groups = null, [FromQuery] bool enumerate = false, [FromQuery] int results = 50, CancellationToken cancellationToken = default)
 		{
 			if (!_analyticsConfig.Value.Authorize(TelemetryAclAction.QueryMetrics, User))
 			{
@@ -79,7 +80,7 @@ namespace HordeServer.Telemetry
 
 			List<GetTelemetryMetricsResponse> result = new List<GetTelemetryMetricsResponse>();
 
-			List<IMetric> metrics = await _metricCollection.FindAsync(storeId, id, minTime, maxTime, group, results, cancellationToken);
+			List<IMetric> metrics = await _metricCollection.FindAsync(storeId, id, minTime, maxTime, groups, enumerate, results, cancellationToken);
 
 			HashSet<MetricId> unique = new HashSet<MetricId>(metrics.Select(m => m.MetricId));
 
@@ -103,7 +104,7 @@ namespace HordeServer.Telemetry
 					rmetric.Time = metric.Time;
 					rmetric.Value = metric.Value;
 
-					if (group == null)
+					if (groups == null || groups.Length > 1)
 					{
 						rmetric.Group = metric.Group;
 					}
