Temporary change to Horde metrics to collect Content events at 5 minute intervals

diff --git a/Engine/Source/Programs/Horde/HordeServer/Defaults/default-metrics.telemetry.json b/Engine/Source/Programs/Horde/HordeServer/Defaults/default-metrics.telemetry.json
--- a/Engine/Source/Programs/Horde/HordeServer/Defaults/default-metrics.telemetry.json
+++ b/Engine/Source/Programs/Horde/HordeServer/Defaults/default-metrics.telemetry.json
@@ -28,7 +28,7 @@
 			"property": "$.Payload.LoadingSeconds",
 			"groupBy": "$.Payload.Build_BranchName,$.Payload.ProjectName,$.Payload.Horde_StepName",
 			"function": "Average",
-			"interval": "00:15:00"
+			"interval": "03:00:00"
 		},
 		{
 			"id": "core-loading-build-timetopie-kpi",
@@ -36,7 +36,7 @@
 			"property": "$.Payload.LoadingSeconds",
 			"groupBy": "$.Payload.Build_BranchName,$.Payload.ProjectName,$.Payload.Horde_StepName",
 			"function": "Average",
-			"interval": "00:15:00"
+			"interval": "03:00:00"
 		},
 		{
 			"id": "core-loading-build-timetoedit-percentile",
@@ -45,7 +45,7 @@
 			"property": "$.Payload.LoadingSeconds",
 			"function": "Percentile",
 			"percentile": 75,
-			"interval": "00:15:00"
+			"interval": "03:00:00"
 		},
 		{
 			"id": "core-loading-build-timetopie-percentile",
@@ -54,7 +54,7 @@
 			"property": "$.Payload.LoadingSeconds",
 			"function": "Percentile",
 			"percentile": 75,
-			"interval": "00:15:00"
+			"interval": "03:00:00"
 		},
 		{
 			"id": "core-loading-build-timetoedit",
@@ -62,7 +62,7 @@
 			"property": "$.Payload.LoadingSeconds",
 			"groupBy": "$.Payload.Build_BranchName,$.Payload.ProjectName,$.Payload.Horde_StepName",
 			"function": "Average",
-			"interval": "00:15:00"
+			"interval": "03:00:00"
 		},
 		{
 			"id": "core-loading-build-timetopie",
@@ -70,7 +70,7 @@
 			"property": "$.Payload.LoadingSeconds",
 			"groupBy": "$.Payload.Build_BranchName,$.Payload.ProjectName,$.Payload.Horde_StepName",
 			"function": "Average",
-			"interval": "00:15:00"
+			"interval": "03:00:00"
 		},
 		{
 			"id": "core-loading-user-loadingtime",
@@ -536,8 +536,7 @@
 			"property": "$.Payload.TotalFileSize",
 			"groupBy": "$.Payload.Build_BranchName",
 			"function": "Sum",
-			"interval": "00:15:00",
-			"topN": 25
+			"interval": "00:5:00"
 		},
 		{
 			"id": "core-content-totalsize-mp",
@@ -545,17 +544,15 @@
 			"property": "$.Payload.TotalFileSize",
 			"groupBy": "$.Payload.Build_BranchName,$.Payload.Name",
 			"function": "Sum",
-			"interval": "00:15:00",
-			"topN": 25
+			"interval": "00:5:00"
 		},
 		{
 			"id": "core-content-virtualization",
 			"filter": "$[?(@.Payload.EventName == 'Core.VirtualAssets.MountPoint')]",
 			"property": "$.Payload.VirtualizedPercent",
 			"groupBy": "$.Payload.Build_BranchName",
-			"function": "Average",
-			"interval": "00:15:00",
-			"bottomN": 25
+			"function": "Sum",
+			"interval": "00:5:00"
 		},
 		{
 			"id": "core-content-virtualization-mp",
@@ -563,8 +560,7 @@
 			"property": "$.Payload.VirtualizedPercent",
 			"groupBy": "$.Payload.Build_BranchName,$.Payload.Name",
 			"function": "Average",
-			"interval": "00:15:00",
-			"bottomN": 25
+			"interval": "00:5:00"
 		},
 		{
 			"id": "horde-incremental-build-jobs",
