Fixing problems encountered during testing:
- Removing ability to link directly to the appropriate line index as this is not available
- Fixing duplicate posting of failing build steps
- Fixing lack of notification due to incorrect slack identifier used

#rnx

diff --git a/Engine/Source/Programs/Horde/Plugins/Experimental/HordeServer.Experimental/Notifications/Sinks/ExperimentalSlackNotificationSink.cs b/Engine/Source/Programs/Horde/Plugins/Experimental/HordeServer.Experimental/Notifications/Sinks/ExperimentalSlackNotificationSink.cs
--- a/Engine/Source/Programs/Horde/Plugins/Experimental/HordeServer.Experimental/Notifications/Sinks/ExperimentalSlackNotificationSink.cs
+++ b/Engine/Source/Programs/Horde/Plugins/Experimental/HordeServer.Experimental/Notifications/Sinks/ExperimentalSlackNotificationSink.cs
@@ -18,8 +18,8 @@ using HordeServer.Streams;
 using HordeServer.Users;
 using Microsoft.Extensions.Logging;
 using Microsoft.Extensions.Options;
-using MongoDB.Bson.Serialization.Attributes;
 using MongoDB.Bson;
+using MongoDB.Bson.Serialization.Attributes;
 using MongoDB.Driver;
 
 namespace HordeServer.Notifications.Sinks
@@ -107,13 +107,11 @@ namespace HordeServer.Notifications.Sinks
 			public LogEventSeverity Severity { get; }
 			public string Message { get; }
 			public LogId? LogId { get; } = null;
-			public int? LineIndex { get; } = null;
-			public LogData(LogEventSeverity severity, string message, LogId? logId = null, int? lineIndex = null)
+			public LogData(LogEventSeverity severity, string message, LogId? logId = null)
 			{
 				Severity = severity;
 				Message = message;
 				LogId = logId;
-				LineIndex = lineIndex;
 			}
 		}
 
@@ -267,6 +265,12 @@ namespace HordeServer.Notifications.Sinks
 		/// <inheritdoc/>
 		public async Task NotifyJobStepCompleteAsync(IEnumerable<IUser>? usersToNotify, IJob job, IJobStepBatch batch, IJobStep step, INode node, List<ILogEventData> jobStepEventData, CancellationToken cancellationToken)
 		{
+			// Currently we only want to notify from the plugin's configuration which would not have this set
+			if (usersToNotify is not null)
+			{
+				return;
+			}
+
 			List<LogData> logEvents = ExtractLogInformation(step, jobStepEventData);
 			await ProcessJobStepCompleteNotificationAsync(job, step, logEvents, cancellationToken);
 		}
@@ -568,17 +572,7 @@ namespace HordeServer.Notifications.Sinks
 			List<LogData> logData = new List<LogData>();
 			foreach (ILogEventData eventData in jobStepEventData)
 			{
-				// Extract the log information from the first log line
-				if (eventData.Lines.Count > 0)
-				{
-					LogEvent logEvent = LogEvent.Read(eventData.Lines[0].Data.Span);
-					logData.Add(new LogData(eventData.Severity, eventData.Message, step.LogId, logEvent.LineIndex));
-				}
-				else if (!String.IsNullOrEmpty(eventData.Message))
-				{
-					// Only store the relevant message for the event if there is no actual log information
-					logData.Add(new LogData(eventData.Severity, eventData.Message, step.LogId));
-				}
+				logData.Add(new LogData(eventData.Severity, eventData.Message, step.LogId));
 			}
 			return logData;
 		}
@@ -858,8 +852,8 @@ namespace HordeServer.Notifications.Sinks
 					channel = userInfo.Id;
 				}
 
-				SlackMessageId messageId = await _slackClient.PostMessageAsync(recipient, message, cancellationToken);
-				_ = await AddOrUpdateJobNotificationStateAsync(job.Id, templateId, channel, messageId, cancellationToken);
+				SlackMessageId messageId = await _slackClient.PostMessageAsync(channel, message, cancellationToken);
+				_ = await AddOrUpdateJobNotificationStateAsync(job.Id, templateId, recipient, messageId, cancellationToken);
 			}
 		}
 
@@ -959,8 +953,7 @@ namespace HordeServer.Notifications.Sinks
 
 				if (log.LogId is not null)
 				{
-					string lineIndexLink = (log.LineIndex is not null) ? $"?lineIndex={log.LineIndex}" : String.Empty;
-					string logUrl = new Uri(_serverInfo.DashboardUrl, $"log/{log.LogId}{lineIndexLink}").ToString();
+					string logUrl = new Uri(_serverInfo.DashboardUrl, $"log/{log.LogId}").ToString();
 
 					ContextBlock detailsContext = new ContextBlock();
 					detailsContext.Elements.Add(new TextObject($"<{logUrl}|View Job Step Log>"));
