Horde: Add queryable Horde analytics meta data for better performance and to limit data set for top/bottom metrics

diff --git a/Engine/Source/Programs/Horde/Plugins/Analytics/HordeServer.Analytics.Tests/Telemetry/TelemetryTests.cs b/Engine/Source/Programs/Horde/Plugins/Analytics/HordeServer.Analytics.Tests/Telemetry/TelemetryTests.cs
--- a/Engine/Source/Programs/Horde/Plugins/Analytics/HordeServer.Analytics.Tests/Telemetry/TelemetryTests.cs
+++ b/Engine/Source/Programs/Horde/Plugins/Analytics/HordeServer.Analytics.Tests/Telemetry/TelemetryTests.cs
@@ -488,7 +488,11 @@ namespace HordeServer.Analytics.Tests.Telemetry
 				await sink.FlushAsync(CancellationToken.None);
 				await collection.FlushAsync(CancellationToken.None);
 
-				List<IMetric> metrics = await collection.FindAsync(telemetryStoreConfig.Id, [metricConfig.Id]);
+				List<IMetricMeta> metricMeta = await collection.FindMetaAsync(telemetryStoreConfig.Id, [metricConfig.Id]);
+
+				Assert.AreEqual(1, metricMeta.Count);
+
+				List<IMetric> metrics = await collection.FindAsync(telemetryStoreConfig.Id, [metricConfig.Id], null, null, metricMeta[0].Groups.ToArray());
 				metrics = metrics.OrderBy(x => x.Group).ToList();
 				Assert.AreEqual(3, metrics.Count);
 
@@ -550,7 +554,11 @@ namespace HordeServer.Analytics.Tests.Telemetry
 				await sink.FlushAsync(CancellationToken.None);
 				await collection.FlushAsync(CancellationToken.None);
 
-				List<IMetric> metrics = await collection.FindAsync(telemetryStoreConfig.Id, [metricConfig.Id], new DateTime(2024, 1, 1, 1, 1, 1, DateTimeKind.Utc));
+				List<IMetricMeta> metricMeta = await collection.FindMetaAsync(telemetryStoreConfig.Id, [metricConfig.Id], new DateTime(2024, 1, 1, 1, 1, 1, DateTimeKind.Utc));
+
+				Assert.AreEqual(1, metricMeta.Count);
+
+				List<IMetric> metrics = await collection.FindAsync(telemetryStoreConfig.Id, [metricConfig.Id], new DateTime(2024, 1, 1, 1, 1, 1, DateTimeKind.Utc), null, metricMeta[0].Groups.ToArray());
 				metrics = metrics.OrderBy(x => x.Group).ToList();
 				Assert.AreEqual(3, metrics.Count);
 
diff --git a/Engine/Source/Programs/Horde/Plugins/Analytics/HordeServer.Analytics/Telemetry/Metrics/IMetric.cs b/Engine/Source/Programs/Horde/Plugins/Analytics/HordeServer.Analytics/Telemetry/Metrics/IMetric.cs
--- a/Engine/Source/Programs/Horde/Plugins/Analytics/HordeServer.Analytics/Telemetry/Metrics/IMetric.cs
+++ b/Engine/Source/Programs/Horde/Plugins/Analytics/HordeServer.Analytics/Telemetry/Metrics/IMetric.cs
@@ -34,4 +34,35 @@ namespace HordeServer.Telemetry.Metrics
 		/// </summary>
 		int Count { get; }
 	}
+
+	/// <summary>
+	/// Interface for metric meta data
+	/// </summary>
+	public interface IMetricMeta
+	{
+		/// <summary>
+		/// Unique identifier for this event
+		/// </summary>
+		MetricId MetricId { get; }
+
+		/// <summary>
+		/// GroupBy for the metric
+		/// </summary>
+		string GroupBy { get; }
+
+		/// <summary>
+		/// Unique groups
+		/// </summary>
+		List<string> Groups { get; }
+
+		/// <summary>
+		/// TopN aggregation
+		/// </summary>
+		public int TopN { get; }
+
+		/// <summary>
+		/// BottomN aggregation
+		/// </summary>
+		public int BottomN { get; }
+	}
 }
diff --git a/Engine/Source/Programs/Horde/Plugins/Analytics/HordeServer.Analytics/Telemetry/Metrics/IMetricCollection.cs b/Engine/Source/Programs/Horde/Plugins/Analytics/HordeServer.Analytics/Telemetry/Metrics/IMetricCollection.cs
--- a/Engine/Source/Programs/Horde/Plugins/Analytics/HordeServer.Analytics/Telemetry/Metrics/IMetricCollection.cs
+++ b/Engine/Source/Programs/Horde/Plugins/Analytics/HordeServer.Analytics/Telemetry/Metrics/IMetricCollection.cs
@@ -18,6 +18,17 @@ namespace HordeServer.Telemetry.Metrics
 		/// <param name="node">Arbitrary event node</param>
 		void AddEvent(TelemetryStoreId storeId, JsonNode node);
 
+		/// <summary>
+		/// Finds metrics over a given time period
+		/// </summary>
+		/// <param name="storeId">Identifier for a telemetry store</param>
+		/// <param name="metricIds">Metrics to search for</param>
+		/// <param name="minTime">Start of the time period to query</param>
+		/// <param name="maxTime">End of the time period to query</param>
+		/// <param name="maxResults">Maximum number of results to return</param>
+		/// <param name="cancellationToken">Cancellation token for the operation</param>
+		Task<List<IMetricMeta>> FindMetaAsync(TelemetryStoreId storeId, MetricId[] metricIds, DateTime? minTime = null, DateTime? maxTime = null, int maxResults = 50, CancellationToken cancellationToken = default);
+
 		/// <summary>
 		/// Finds metrics over a given time period
 		/// </summary>
@@ -29,7 +40,7 @@ namespace HordeServer.Telemetry.Metrics
 		/// <param name="enumerate">Whether to enumerate results, avoiding additional filtering</param> 
 		/// <param name="maxResults">Maximum number of results to return</param>
 		/// <param name="cancellationToken">Cancellation token for the operation</param>
-		Task<List<IMetric>> FindAsync(TelemetryStoreId storeId, MetricId[] metricIds, DateTime? minTime = null, DateTime? maxTime = null, string[]? groups = null, bool enumerate = false, int maxResults = 50, CancellationToken cancellationToken = default);
+		Task<List<IMetric>> FindAsync(TelemetryStoreId storeId, MetricId[] metricIds, DateTime? minTime = null, DateTime? maxTime = null, string[]? groups = null, int maxResults = 50, CancellationToken cancellationToken = default);
 
 		/// <summary>
 		/// Flush all the pending events to storage
diff --git a/Engine/Source/Programs/Horde/Plugins/Analytics/HordeServer.Analytics/Telemetry/Metrics/MetricCollection.cs b/Engine/Source/Programs/Horde/Plugins/Analytics/HordeServer.Analytics/Telemetry/Metrics/MetricCollection.cs
--- a/Engine/Source/Programs/Horde/Plugins/Analytics/HordeServer.Analytics/Telemetry/Metrics/MetricCollection.cs
+++ b/Engine/Source/Programs/Horde/Plugins/Analytics/HordeServer.Analytics/Telemetry/Metrics/MetricCollection.cs
@@ -17,6 +17,33 @@ using TDigestNet;
 
 namespace HordeServer.Telemetry.Metrics
 {
+	class MetricMeta : IMetricMeta
+	{
+		/// <summary>
+		/// Unique identifier for this event
+		/// </summary>
+		public MetricId MetricId { get; set; }
+
+		/// <summary>
+		/// Metric GroupBy
+		/// </summary>
+		public string GroupBy { get; set; } = String.Empty;
+
+		/// <summary>
+		/// Unique groups
+		/// </summary>
+		public List<string> Groups { get; set; } = new List<string>();
+
+		/// <summary>
+		/// TopN aggregation
+		/// </summary>
+		public int TopN { get; set; } = 0;
+
+		/// <summary>
+		/// BottomN aggregation
+		/// </summary>
+		public int BottomN { get; set; } = 0;
+	}
 	class MetricCollection : IMetricCollection, IHostedService, IAsyncDisposable
 	{
 		class MetricDocument : IMetric
@@ -328,8 +355,72 @@ namespace HordeServer.Telemetry.Metrics
 			}
 		}
 
+		class MetricProjection
+		{
+			public string _group = String.Empty;
+			public double _total;
+		}
+
+		/// <inheritdoc/>
+		public async Task<List<IMetricMeta>> FindMetaAsync(TelemetryStoreId telemetryStoreId, MetricId[] metricIds, DateTime? minTime = null, DateTime? maxTime = null, int maxResults = 50, CancellationToken cancellationToken = default)
+		{
+			List<IMetricMeta> result = new List<IMetricMeta>();
+
+			if (!_analyticsConfig.CurrentValue.TryGetTelemetryStore(telemetryStoreId, out TelemetryStoreConfig? telemetryStoreConfig))
+			{
+				return result;
+			}
+
+			FilterDefinition<MetricDocument> filter = FilterDefinition<MetricDocument>.Empty;
+			filter &= Builders<MetricDocument>.Filter.Eq(x => x.TelemetryStoreId, telemetryStoreId);
+
+			if (minTime != null)
+			{
+				filter &= Builders<MetricDocument>.Filter.Gte(x => x.Time, minTime.Value);
+			}
+			if (maxTime != null)
+			{
+				filter &= Builders<MetricDocument>.Filter.Lte(x => x.Time, maxTime.Value);
+			}
+
+			for (int i = 0; i < metricIds.Length; i++)
+			{
+				MetricId metricId = metricIds[i];
+				MetricConfig? metricConfig = null;
+
+				if (!telemetryStoreConfig.TryGetMetric(metricId, out metricConfig))
+				{
+					continue;
+				}
+
+				FilterDefinition<MetricDocument> metricFilter = filter & Builders<MetricDocument>.Filter.Eq(x => x.MetricId, metricIds[i]);
+
+				List<string> groups = new List<string>();
+
+				if (metricConfig.TopN > 0)
+				{ 
+					List<MetricProjection> values = await _metrics.Aggregate().Match(metricFilter).Group(x => x.Group, group => new MetricProjection { _group = group.Key, _total = group.Sum(y => y.Value) }).SortByDescending(x => x._total).Limit(metricConfig.TopN).ToListAsync(cancellationToken);
+					groups = values.Select(x => x._group).ToList();
+				}
+				else if (metricConfig.BottomN > 0)
+				{
+
+					List<MetricProjection> values = await _metrics.Aggregate().Match(metricFilter).Group(x => x.Group, group => new MetricProjection { _group = group.Key, _total = group.Sum(y => y.Value) }).SortBy(x => x._total).Limit(metricConfig.BottomN).ToListAsync(cancellationToken);
+					groups = values.Select(x => x._group).ToList();
+				}
+				else
+				{
+					groups = await (await _metrics.DistinctAsync(x => x.Group, metricFilter, cancellationToken: cancellationToken)).ToListAsync(cancellationToken);
+				}
+
+				result.Add(new MetricMeta() { MetricId = metricId, GroupBy = metricConfig!.GroupBy, Groups = groups, TopN = metricConfig.TopN, BottomN = metricConfig.BottomN });
+			}
+
+			return result;
+		}
+
 		/// <inheritdoc/>
-		public async Task<List<IMetric>> FindAsync(TelemetryStoreId telemetryStoreId, MetricId[] metricIds, DateTime? minTime = null, DateTime? maxTime = null, string[]? groups = null, bool enumerate = false, int maxResults = 50, CancellationToken cancellationToken = default)
+		public async Task<List<IMetric>> FindAsync(TelemetryStoreId telemetryStoreId, MetricId[] metricIds, DateTime? minTime = null, DateTime? maxTime = null, string[]? groups = null, int maxResults = 50, CancellationToken cancellationToken = default)
 		{
 			FilterDefinition<MetricDocument> filter = FilterDefinition<MetricDocument>.Empty;
 			filter &= Builders<MetricDocument>.Filter.Eq(x => x.TelemetryStoreId, telemetryStoreId);
@@ -352,47 +443,7 @@ namespace HordeServer.Telemetry.Metrics
 
 			List<MetricDocument> results = await _metrics.Find(filter).SortByDescending(x => x.Time).Limit(maxResults).ToListAsync(cancellationToken);
 
-			HashSet<string> filterGroups = new HashSet<string>();
-			if (!enumerate)
-			{
-				if (_analyticsConfig.CurrentValue.TryGetTelemetryStore(telemetryStoreId, out TelemetryStoreConfig? telemetryStoreConfig))
-				{
-					foreach (MetricId metricId in metricIds)
-					{
-						MetricConfig? metricConfig;
-						if (telemetryStoreConfig.TryGetMetric(metricId, out metricConfig))
-						{
-							if (metricConfig.TopN > 0 || metricConfig.BottomN > 0)
-							{
-
-								List<string> ordered = results.Where(x => x.MetricId == metricId).GroupBy(x => x.Group, x => x).Select(x => (x.Key, x.Sum(x => x.Value))).OrderByDescending(x => x.Item2).Select(x => x.Item1).ToList();
-
-								if (ordered.Count <= metricConfig.TopN && ordered.Count <= metricConfig.BottomN)
-								{
-									continue;
-								}
-
-								if (metricConfig.TopN > 0)
-								{
-									foreach (string groupName in ordered.TakeLast(ordered.Count - metricConfig.TopN))
-									{
-										filterGroups.Add(groupName);
-									}
-								}
-								else
-								{
-									foreach (string groupName in ordered.Take(ordered.Count - metricConfig.BottomN))
-									{
-										filterGroups.Add(groupName);
-									}
-								}
-							}
-						}
-					}
-				}
-			}
-
-			return results.Where(x => !filterGroups.Contains(x.Group)).ToList().ConvertAll<MetricDocument, IMetric>(x => x);
+			return results.ConvertAll<MetricDocument, IMetric>(x => x);
 		}
 	}
 }
diff --git a/Engine/Source/Programs/Horde/Plugins/Analytics/HordeServer.Analytics/Telemetry/Metrics/MetricConfig.cs b/Engine/Source/Programs/Horde/Plugins/Analytics/HordeServer.Analytics/Telemetry/Metrics/MetricConfig.cs
--- a/Engine/Source/Programs/Horde/Plugins/Analytics/HordeServer.Analytics/Telemetry/Metrics/MetricConfig.cs
+++ b/Engine/Source/Programs/Horde/Plugins/Analytics/HordeServer.Analytics/Telemetry/Metrics/MetricConfig.cs
@@ -115,12 +115,12 @@ namespace HordeServer.Telemetry.Metrics
 		public int Percentile { get; set; } = 95;
 
 		/// <summary>
-		/// 
+		/// TopN aggregation
 		/// </summary>
 		public int TopN { get; set; } = 0;
 
 		/// <summary>
-		/// 
+		/// BottomN aggregation
 		/// </summary>
 		public int BottomN { get; set; } = 0;
 
diff --git a/Engine/Source/Programs/Horde/Plugins/Analytics/HordeServer.Analytics/Telemetry/TelemetryController.cs b/Engine/Source/Programs/Horde/Plugins/Analytics/HordeServer.Analytics/Telemetry/TelemetryController.cs
--- a/Engine/Source/Programs/Horde/Plugins/Analytics/HordeServer.Analytics/Telemetry/TelemetryController.cs
+++ b/Engine/Source/Programs/Horde/Plugins/Analytics/HordeServer.Analytics/Telemetry/TelemetryController.cs
@@ -48,7 +48,44 @@ namespace HordeServer.Telemetry
 		[Obsolete("Pass a telemetry store id in the route instead")]
 		public Task<ActionResult<List<GetTelemetryMetricsResponse>>> GetMetricsAsync([FromQuery] MetricId[] id, [FromQuery] DateTime? minTime = null, [FromQuery] DateTime? maxTime = null, [FromQuery] string? group = null, [FromQuery] int results = 50, CancellationToken cancellationToken = default)
 		{
-			return GetMetricsAsync(TelemetryStoreId.Default, id, minTime, maxTime, group == null ? null : new string[] { group }, true, results, cancellationToken);
+			return GetMetricsAsync(TelemetryStoreId.Default, id, minTime, maxTime, group == null ? null : new string[] { group }, results, cancellationToken);
+		}
+
+		/// <summary>
+		/// Queries aggregated metrics from the telemetry system
+		/// </summary>
+		/// <param name="storeId">The telemetry store id</param>
+		/// <param name="id">The metrics to query</param>
+		/// <param name="minTime">Minimum time interval to query</param>
+		/// <param name="maxTime">Maximum time interval to query</param>
+		/// <param name="results">Number of results to return</param>
+		/// <param name="cancellationToken">Cancellation token for the operation</param>
+		[HttpGet]
+		[Authorize]
+		[Route("/api/v1/telemetry/{storeId}/metrics/meta")]
+		public async Task<ActionResult<List<GetTelemetryMetricsMetaResponse>>> GetMetricsMetaAsync(TelemetryStoreId storeId, [FromQuery] MetricId[] id, [FromQuery] DateTime? minTime = null, [FromQuery] DateTime? maxTime = null, [FromQuery] int results = 50, CancellationToken cancellationToken = default)
+		{
+			if (!_analyticsConfig.Value.Authorize(TelemetryAclAction.QueryMetrics, User))
+			{
+				return Forbid(TelemetryAclAction.QueryMetrics);
+			}
+
+			TelemetryStoreConfig? telemetryStoreConfig;
+			if (!_analyticsConfig.Value.TryGetTelemetryStore(storeId, out telemetryStoreConfig))
+			{
+				return NotFound(storeId);
+			}
+
+			List<GetTelemetryMetricsMetaResponse> result = new List<GetTelemetryMetricsMetaResponse>();
+
+			List<IMetricMeta> metricMetas = await _metricCollection.FindMetaAsync(telemetryStoreConfig.Id, id, minTime, maxTime, results, cancellationToken);
+
+			foreach (IMetricMeta metricMeta in metricMetas)
+			{
+				result.Add(new GetTelemetryMetricsMetaResponse() { MetricId = metricMeta.MetricId.ToString(), GroupBy = metricMeta.GroupBy, Groups = metricMeta.Groups, TopN = metricMeta.TopN, BottomN = metricMeta.BottomN });
+			}
+
+			return result;
 		}
 
 		/// <summary>
@@ -65,7 +102,7 @@ namespace HordeServer.Telemetry
 		[HttpGet]
 		[Authorize]
 		[Route("/api/v1/telemetry/{storeId}/metrics")]
-		public async Task<ActionResult<List<GetTelemetryMetricsResponse>>> GetMetricsAsync(TelemetryStoreId storeId, [FromQuery] MetricId[] id, [FromQuery] DateTime? minTime = null, [FromQuery] DateTime? maxTime = null, [FromQuery(Name = "group")] string[]? groups = null, [FromQuery] bool enumerate = false, [FromQuery] int results = 50, CancellationToken cancellationToken = default)
+		public async Task<ActionResult<List<GetTelemetryMetricsResponse>>> GetMetricsAsync(TelemetryStoreId storeId, [FromQuery] MetricId[] id, [FromQuery] DateTime? minTime = null, [FromQuery] DateTime? maxTime = null, [FromQuery(Name = "group")] string[]? groups = null, [FromQuery] int results = 50, CancellationToken cancellationToken = default)
 		{
 			if (!_analyticsConfig.Value.Authorize(TelemetryAclAction.QueryMetrics, User))
 			{
@@ -80,7 +117,7 @@ namespace HordeServer.Telemetry
 
 			List<GetTelemetryMetricsResponse> result = new List<GetTelemetryMetricsResponse>();
 
-			List<IMetric> metrics = await _metricCollection.FindAsync(storeId, id, minTime, maxTime, groups, enumerate, results, cancellationToken);
+			List<IMetric> metrics = await _metricCollection.FindAsync(storeId, id, minTime, maxTime, groups, results, cancellationToken);
 
 			HashSet<MetricId> unique = new HashSet<MetricId>(metrics.Select(m => m.MetricId));
 
diff --git a/Engine/Source/Programs/Shared/EpicGames.Horde/Telemetry/TelemetryMessages.cs b/Engine/Source/Programs/Shared/EpicGames.Horde/Telemetry/TelemetryMessages.cs
--- a/Engine/Source/Programs/Shared/EpicGames.Horde/Telemetry/TelemetryMessages.cs
+++ b/Engine/Source/Programs/Shared/EpicGames.Horde/Telemetry/TelemetryMessages.cs
@@ -34,6 +34,37 @@ namespace EpicGames.Horde.Telemetry
 		EtEventStream
 	}
 
+	/// <summary>
+	/// Metrics meta data for query
+	/// </summary>
+	public class GetTelemetryMetricsMetaResponse
+	{
+		/// <summary>
+		/// The corresponding metric id
+		/// </summary>
+		public string MetricId { get; set; } = String.Empty;
+
+		/// <summary>
+		/// Metric grouping information
+		/// </summary>
+		public string GroupBy { get; set; } = String.Empty;
+
+		/// <summary>
+		/// Metrics groups matching the search terms
+		/// </summary>
+		public List<string> Groups { get; set; } = new List<string>();
+
+		/// <summary>
+		/// TopN aggregation
+		/// </summary>
+		public int TopN { get; set; } = 0;
+
+		/// <summary>
+		/// BottomN aggregation
+		/// </summary>
+		public int BottomN { get; set; } = 0;
+	}
+
 	/// <summary>
 	/// Metrics matching a particular query
 	/// </summary>
