Horde: Fix server warnings
#submitool bypass
#rnx

diff --git a/Engine/Source/Programs/Horde/Plugins/Analytics/HordeServer.Analytics/Telemetry/Metrics/IMetricCollection.cs b/Engine/Source/Programs/Horde/Plugins/Analytics/HordeServer.Analytics/Telemetry/Metrics/IMetricCollection.cs
--- a/Engine/Source/Programs/Horde/Plugins/Analytics/HordeServer.Analytics/Telemetry/Metrics/IMetricCollection.cs
+++ b/Engine/Source/Programs/Horde/Plugins/Analytics/HordeServer.Analytics/Telemetry/Metrics/IMetricCollection.cs
@@ -36,8 +36,7 @@ namespace HordeServer.Telemetry.Metrics
 		/// <param name="metricIds">Metrics to search for</param>
 		/// <param name="minTime">Start of the time period to query</param>
 		/// <param name="maxTime">End of the time period to query</param>
-		/// <param name="groups">Grouping keys to filter results</param>
-		/// <param name="enumerate">Whether to enumerate results, avoiding additional filtering</param> 
+		/// <param name="groups">Grouping keys to filter results</param>		
 		/// <param name="maxResults">Maximum number of results to return</param>
 		/// <param name="cancellationToken">Cancellation token for the operation</param>
 		Task<List<IMetric>> FindAsync(TelemetryStoreId storeId, MetricId[] metricIds, DateTime? minTime = null, DateTime? maxTime = null, string[]? groups = null, int maxResults = 50, CancellationToken cancellationToken = default);
diff --git a/Engine/Source/Programs/Horde/Plugins/Analytics/HordeServer.Analytics/Telemetry/TelemetryController.cs b/Engine/Source/Programs/Horde/Plugins/Analytics/HordeServer.Analytics/Telemetry/TelemetryController.cs
--- a/Engine/Source/Programs/Horde/Plugins/Analytics/HordeServer.Analytics/Telemetry/TelemetryController.cs
+++ b/Engine/Source/Programs/Horde/Plugins/Analytics/HordeServer.Analytics/Telemetry/TelemetryController.cs
@@ -96,7 +96,6 @@ namespace HordeServer.Telemetry
 		/// <param name="minTime">Minimum time interval to query</param>
 		/// <param name="maxTime">Maximum time interval to query</param>
 		/// <param name="groups">Grouping keys</param>
-		/// <param name="enumerate">Whether to enumerate keys</param>
 		/// <param name="results">Number of results to return</param>
 		/// <param name="cancellationToken">Cancellation token for the operation</param>
 		[HttpGet]
