Moved Clear Selection and All Templates to a single checkbox, and added a Select All/None for each category with at least 3 options.
Minor code tweaks for clarity and to remove unused variables.
#tests Tested locally in browser

diff --git a/Engine/Source/Programs/Horde/HordeDashboard/src/components/JobSearchSimple.tsx b/Engine/Source/Programs/Horde/HordeDashboard/src/components/JobSearchSimple.tsx
--- a/Engine/Source/Programs/Horde/HordeDashboard/src/components/JobSearchSimple.tsx
+++ b/Engine/Source/Programs/Horde/HordeDashboard/src/components/JobSearchSimple.tsx
@@ -78,6 +78,7 @@ export const JobSearchSimpleModal: React.FC<{ streamIdIn: string, changeListIn?:
 
    const sorted = new Map<string, GetTemplateRefResponse[]>();
 
+   // filter and sort templates into categories
    templates.forEach(t => {
 
       stream.tabs.forEach(tab => {
@@ -103,12 +104,12 @@ export const JobSearchSimpleModal: React.FC<{ streamIdIn: string, changeListIn?:
 
    Array.from(sorted.keys()).sort((a, b) => a < b ? -1 : 1).forEach(cat => {
 
-      const templates = sorted.get(cat);
-      if (!templates?.length) {
+      const catTemplates = sorted.get(cat);
+      if (!catTemplates?.length) {
          return;
       }
 
-      const subItems: IContextualMenuItem[] = templates.sort((a, b) => a.name < b.name ? -1 : 1).map(t => {
+      const subItems: IContextualMenuItem[] = catTemplates.sort((a, b) => a.name < b.name ? -1 : 1).map(t => {
 
          allTemplates.add(t.id);
          return {
@@ -116,7 +117,7 @@ export const JobSearchSimpleModal: React.FC<{ streamIdIn: string, changeListIn?:
             key: t.id,
             text: t.name,
             data: t,
-            onRender: (props: IContextualMenuItemProps) => {
+            onRender: () => {
                return <Stack horizontal verticalFill verticalAlign="center" style={{ padding: "4px 12px 4px 12px", cursor: "pointer" }} onClick={(ev) => {
                   ev?.preventDefault();
                   ev?.stopPropagation();
@@ -134,27 +135,66 @@ export const JobSearchSimpleModal: React.FC<{ streamIdIn: string, changeListIn?:
          };
       })
 
+      const isSelected = (ct: GetTemplateRefResponse) => searchState.templates.has(ct.id);
+      let allCatTemplatesSelected: boolean = catTemplates.every(isSelected);
+
+      if (catTemplates?.length > 2) {
+         subItems.unshift({
+            key: `select_all_${cat}`,
+            onRender: () => {
+               return <Stack horizontal verticalFill verticalAlign="center" style={{ padding: "4px 12px 4px 12px", cursor: "pointer" }} onClick={(ev) => { 
+                  ev?.stopPropagation();
+                  ev?.preventDefault();
+      
+                  allCatTemplatesSelected = !allCatTemplatesSelected;
+                  
+                  if (allCatTemplatesSelected) {
+                     catTemplates.forEach( ct => 
+                        searchState.templates.add(ct.id)
+                     )
+                     setSearchState({ ...searchState, containsParameter: undefined })
+                  }
+                  else {
+                     catTemplates.forEach( ct => 
+                        searchState.templates.delete(ct.id)
+                     )
+                     setSearchState({ ...searchState, containsParameter: undefined })
+                  }
+                  return false;
+               }} >
+               <Checkbox checked={allCatTemplatesSelected} />
+               <Text style={{ color: dashboard.darktheme ? "#ABABAB" : "#575757" }}>Select All</Text></Stack>
+            }
+         });
+      }
+
       templateOptions.push({ key: `${cat}_category`, text: cat, subMenuProps: { items: subItems } });
 
    })
 
-   templateOptions.push({
-      key: `show_all_templates`, text: "Clear Selection", onClick: (ev) => {
-         ev?.stopPropagation();
-         ev?.preventDefault();
-         searchState.templates = new Set();
-         setSearchState({ ...searchState, containsParameter: undefined })
-         return false;
-      }
-   });
+   let allTemplatesSelected: boolean =  searchState.templates.size === allTemplates.size;
 
-   templateOptions.push({
-      key: `show_all_templates`, text: "All Templates", onClick: (ev) => {
-         ev?.stopPropagation();
-         ev?.preventDefault();
-         searchState.templates = allTemplates;
-         setSearchState({ ...searchState, containsParameter: undefined })
-         return false;
+   templateOptions.unshift({
+      key: `select_all_templates`,
+      onRender: () => {
+         return <Stack horizontal verticalFill verticalAlign="center" style={{ padding: "4px 12px 4px 7px", cursor: "pointer" }} onClick={(ev) => { 
+            ev?.stopPropagation();
+            ev?.preventDefault();
+
+            allTemplatesSelected = !allTemplatesSelected;
+            
+            if (allTemplatesSelected) {
+               searchState.templates = allTemplates;
+               setSearchState({ ...searchState, containsParameter: undefined })
+            }
+            else {
+               searchState.templates = new Set();
+               setSearchState({ ...searchState, containsParameter: undefined })
+            }
+            return false;
+         }} >
+         <Checkbox checked={allTemplatesSelected} />
+         <Text style={{ padding: "6px 0px"}}>Select All</Text></Stack>
       }
    });
 
