Horde: Yet more work on telemetry meta data and queries
#rnx

diff --git a/Engine/Source/Programs/Horde/Plugins/Analytics/HordeServer.Analytics/Telemetry/Metrics/MetricCollection.cs b/Engine/Source/Programs/Horde/Plugins/Analytics/HordeServer.Analytics/Telemetry/Metrics/MetricCollection.cs
--- a/Engine/Source/Programs/Horde/Plugins/Analytics/HordeServer.Analytics/Telemetry/Metrics/MetricCollection.cs
+++ b/Engine/Source/Programs/Horde/Plugins/Analytics/HordeServer.Analytics/Telemetry/Metrics/MetricCollection.cs
@@ -358,6 +358,7 @@ namespace HordeServer.Telemetry.Metrics
 		class MetricProjection
 		{
 			public string _group = String.Empty;
+			public IEnumerable<ObjectId>? _ids;
 			public double _total;
 		}
 
@@ -395,23 +396,7 @@ namespace HordeServer.Telemetry.Metrics
 
 				FilterDefinition<MetricDocument> metricFilter = filter & Builders<MetricDocument>.Filter.Eq(x => x.MetricId, metricIds[i]);
 
-				List<string> groups = new List<string>();
-
-				if (metricConfig.TopN > 0)
-				{ 
-					List<MetricProjection> values = await _metrics.Aggregate().Match(metricFilter).Group(x => x.Group, group => new MetricProjection { _group = group.Key, _total = group.Sum(y => y.Value) }).SortByDescending(x => x._total).Limit(metricConfig.TopN).ToListAsync(cancellationToken);
-					groups = values.Select(x => x._group).ToList();
-				}
-				else if (metricConfig.BottomN > 0)
-				{
-
-					List<MetricProjection> values = await _metrics.Aggregate().Match(metricFilter).Group(x => x.Group, group => new MetricProjection { _group = group.Key, _total = group.Sum(y => y.Value) }).SortBy(x => x._total).Limit(metricConfig.BottomN).ToListAsync(cancellationToken);
-					groups = values.Select(x => x._group).ToList();
-				}
-				else
-				{
-					groups = await (await _metrics.DistinctAsync(x => x.Group, metricFilter, cancellationToken: cancellationToken)).ToListAsync(cancellationToken);
-				}
+				List<string> groups = await (await _metrics.DistinctAsync(x => x.Group, metricFilter, cancellationToken: cancellationToken)).ToListAsync(cancellationToken);
 
 				result.Add(new MetricMeta() { MetricId = metricId, GroupBy = metricConfig!.GroupBy, Groups = groups, TopN = metricConfig.TopN, BottomN = metricConfig.BottomN });
 			}
@@ -422,6 +407,25 @@ namespace HordeServer.Telemetry.Metrics
 		/// <inheritdoc/>
 		public async Task<List<IMetric>> FindAsync(TelemetryStoreId telemetryStoreId, MetricId[] metricIds, DateTime? minTime = null, DateTime? maxTime = null, string[]? groups = null, int maxResults = 50, CancellationToken cancellationToken = default)
 		{
+			if (!_analyticsConfig.CurrentValue.TryGetTelemetryStore(telemetryStoreId, out TelemetryStoreConfig? telemetryStoreConfig))
+			{
+				return new List<IMetric>();
+			}
+
+			List<MetricConfig> metricConfigs = new List<MetricConfig>();
+			for (int i = 0; i < metricIds.Length; i++)
+			{
+				MetricId metricId = metricIds[i];
+				MetricConfig? metricConfig = null;
+
+				if (!telemetryStoreConfig.TryGetMetric(metricId, out metricConfig))
+				{
+					continue;
+				}
+
+				metricConfigs.Add(metricConfig);
+			}
+
 			FilterDefinition<MetricDocument> filter = FilterDefinition<MetricDocument>.Empty;
 			filter &= Builders<MetricDocument>.Filter.Eq(x => x.TelemetryStoreId, telemetryStoreId);
 
@@ -433,15 +437,43 @@ namespace HordeServer.Telemetry.Metrics
 			{
 				filter &= Builders<MetricDocument>.Filter.Lte(x => x.Time, maxTime.Value);
 			}
+			if (groups != null && groups.Length > 0)
+			{
+				filter &= Builders<MetricDocument>.Filter.In(x => x.Group, groups);
+			}
 
-			filter &= Builders<MetricDocument>.Filter.In(x => x.MetricId, metricIds);
+			List<MetricDocument> results = new List<MetricDocument>();
+			List<MetricConfig> configs = metricConfigs.Where(x => x.TopN == 0 && x.BottomN == 0).ToList();
 
-			if (groups != null)
+			// batch into a single query where possible
+			if (configs.Count > 0)
 			{
-				filter &= Builders<MetricDocument>.Filter.In(x => x.Group, groups);
+				FilterDefinition<MetricDocument> metricFilter = filter & Builders<MetricDocument>.Filter.In(x => x.MetricId, configs.Select(x => x.Id));
+				results = await _metrics.Find(metricFilter).SortByDescending(x => x.Time).Limit(maxResults).ToListAsync(cancellationToken);				
 			}
 
-			List<MetricDocument> results = await _metrics.Find(filter).SortByDescending(x => x.Time).Limit(maxResults).ToListAsync(cancellationToken);
+			configs = metricConfigs.Where(x => x.TopN > 0 || x.BottomN > 0).ToList();
+			foreach (MetricConfig config in configs)
+			{
+				FilterDefinition<MetricDocument> metricFilter = filter & Builders<MetricDocument>.Filter.Eq(x => x.MetricId, config.Id);
+
+				List<MetricProjection> values;
+				if (config.TopN > 0)
+				{
+					values = await _metrics.Aggregate().Match(metricFilter).Group(x => x.Group, group => new MetricProjection { _group = group.Key, _ids = group.Select(x => x.Id), _total = group.Sum(y => y.Value) }).SortByDescending(x => x._total).Limit(config.TopN).ToListAsync(cancellationToken);
+				}
+				else
+				{
+					values = await _metrics.Aggregate().Match(metricFilter).Group(x => x.Group, group => new MetricProjection { _group = group.Key, _ids = group.Select(x => x.Id), _total = group.Sum(y => y.Value) }).SortBy(x => x._total).Limit(config.BottomN).ToListAsync(cancellationToken);
+				}
+
+				if (values.Count > 0)
+				{
+					FilterDefinition<MetricDocument> tfilter = FilterDefinition<MetricDocument>.Empty;
+					tfilter &= Builders<MetricDocument>.Filter.In(x => x.Id, values.SelectMany(x => x._ids!));
+					results.AddRange(await _metrics.Find(tfilter).SortByDescending(x => x.Time).Limit(maxResults).ToListAsync(cancellationToken));
+				}
+			}
 
 			return results.ConvertAll<MetricDocument, IMetric>(x => x);
 		}
