Horde: Add body to metrics query so can support filterig by groups

diff --git a/Engine/Source/Programs/Horde/Plugins/Analytics/HordeServer.Analytics/Telemetry/TelemetryController.cs b/Engine/Source/Programs/Horde/Plugins/Analytics/HordeServer.Analytics/Telemetry/TelemetryController.cs
--- a/Engine/Source/Programs/Horde/Plugins/Analytics/HordeServer.Analytics/Telemetry/TelemetryController.cs
+++ b/Engine/Source/Programs/Horde/Plugins/Analytics/HordeServer.Analytics/Telemetry/TelemetryController.cs
@@ -48,7 +48,7 @@ namespace HordeServer.Telemetry
 		[Obsolete("Pass a telemetry store id in the route instead")]
 		public Task<ActionResult<List<GetTelemetryMetricsResponse>>> GetMetricsAsync([FromQuery] MetricId[] id, [FromQuery] DateTime? minTime = null, [FromQuery] DateTime? maxTime = null, [FromQuery] string? group = null, [FromQuery] int results = 50, CancellationToken cancellationToken = default)
 		{
-			return GetMetricsAsync(TelemetryStoreId.Default, id, minTime, maxTime, group == null ? null : new string[] { group }, results, cancellationToken);
+			return GetMetricsAsync(TelemetryStoreId.Default, id, minTime, maxTime,results, new GetTelemetryMetricsRequest { Groups = group == null ? null : new string[] { group } }, cancellationToken);
 		}
 
 		/// <summary>
@@ -92,6 +92,7 @@ namespace HordeServer.Telemetry
 		/// Queries aggregated metrics from the telemetry system
 		/// </summary>
 		/// <param name="storeId">The telemetry store id</param>
+		/// <param name="request">The metrics request</param>
 		/// <param name="id">The metrics to query</param>
 		/// <param name="minTime">Minimum time interval to query</param>
 		/// <param name="maxTime">Maximum time interval to query</param>
@@ -101,7 +102,7 @@ namespace HordeServer.Telemetry
 		[HttpGet]
 		[Authorize]
 		[Route("/api/v1/telemetry/{storeId}/metrics")]
-		public async Task<ActionResult<List<GetTelemetryMetricsResponse>>> GetMetricsAsync(TelemetryStoreId storeId, [FromQuery] MetricId[] id, [FromQuery] DateTime? minTime = null, [FromQuery] DateTime? maxTime = null, [FromQuery(Name = "group")] string[]? groups = null, [FromQuery] int results = 50, CancellationToken cancellationToken = default)
+		public async Task<ActionResult<List<GetTelemetryMetricsResponse>>> GetMetricsAsync(TelemetryStoreId storeId, [FromQuery] MetricId[] id, [FromQuery] DateTime? minTime = null, [FromQuery] DateTime? maxTime = null, [FromQuery] int results = 50, [FromBody] GetTelemetryMetricsRequest? request = null, CancellationToken cancellationToken = default)
 		{
 			if (!_analyticsConfig.Value.Authorize(TelemetryAclAction.QueryMetrics, User))
 			{
@@ -116,7 +117,7 @@ namespace HordeServer.Telemetry
 
 			List<GetTelemetryMetricsResponse> result = new List<GetTelemetryMetricsResponse>();
 
-			List<IMetric> metrics = await _metricCollection.FindAsync(storeId, id, minTime, maxTime, groups, results, cancellationToken);
+			List<IMetric> metrics = await _metricCollection.FindAsync(storeId, id, minTime, maxTime, request?.Groups, results, cancellationToken);
 
 			HashSet<MetricId> unique = new HashSet<MetricId>(metrics.Select(m => m.MetricId));
 
@@ -140,7 +141,7 @@ namespace HordeServer.Telemetry
 					rmetric.Time = metric.Time;
 					rmetric.Value = metric.Value;
 
-					if (groups == null || groups.Length > 1)
+					if (request?.Groups == null || request.Groups.Length > 1)
 					{
 						rmetric.Group = metric.Group;
 					}
diff --git a/Engine/Source/Programs/Shared/EpicGames.Horde/Telemetry/TelemetryMessages.cs b/Engine/Source/Programs/Shared/EpicGames.Horde/Telemetry/TelemetryMessages.cs
--- a/Engine/Source/Programs/Shared/EpicGames.Horde/Telemetry/TelemetryMessages.cs
+++ b/Engine/Source/Programs/Shared/EpicGames.Horde/Telemetry/TelemetryMessages.cs
@@ -34,6 +34,17 @@ namespace EpicGames.Horde.Telemetry
 		EtEventStream
 	}
 
+	/// <summary>
+	/// 
+	/// </summary>
+	public class GetTelemetryMetricsRequest
+	{
+		/// <summary>
+		/// List of telemetry groups
+		/// </summary>
+		public string[]? Groups { get; set; }
+	}
+
 	/// <summary>
 	/// Metrics meta data for query
 	/// </summary>
