Horde: Fix server warning
#rnx

diff --git a/Engine/Source/Programs/Horde/Plugins/Analytics/HordeServer.Analytics/Telemetry/TelemetryController.cs b/Engine/Source/Programs/Horde/Plugins/Analytics/HordeServer.Analytics/Telemetry/TelemetryController.cs
--- a/Engine/Source/Programs/Horde/Plugins/Analytics/HordeServer.Analytics/Telemetry/TelemetryController.cs
+++ b/Engine/Source/Programs/Horde/Plugins/Analytics/HordeServer.Analytics/Telemetry/TelemetryController.cs
@@ -95,8 +95,7 @@ namespace HordeServer.Telemetry
 		/// <param name="request">The metrics request</param>
 		/// <param name="id">The metrics to query</param>
 		/// <param name="minTime">Minimum time interval to query</param>
-		/// <param name="maxTime">Maximum time interval to query</param>
-		/// <param name="groups">Grouping keys</param>
+		/// <param name="maxTime">Maximum time interval to query</param>		
 		/// <param name="results">Number of results to return</param>
 		/// <param name="cancellationToken">Cancellation token for the operation</param>
 		[HttpGet]
