Adjusting experimental notification to take into account potential race condition between notification of job and job step completion

#rnx

diff --git a/Engine/Source/Programs/Horde/Plugins/Experimental/HordeServer.Experimental/Notifications/Sinks/ExperimentalSlackNotificationSink.cs b/Engine/Source/Programs/Horde/Plugins/Experimental/HordeServer.Experimental/Notifications/Sinks/ExperimentalSlackNotificationSink.cs
--- a/Engine/Source/Programs/Horde/Plugins/Experimental/HordeServer.Experimental/Notifications/Sinks/ExperimentalSlackNotificationSink.cs
+++ b/Engine/Source/Programs/Horde/Plugins/Experimental/HordeServer.Experimental/Notifications/Sinks/ExperimentalSlackNotificationSink.cs
@@ -203,14 +203,18 @@ namespace HordeServer.Notifications.Sinks
 				branch = config.Name;
 			}
 
-			string? abortedByUserName = null;
-			if (job.AbortedByUserId is not null)
+			// Only handle the job complete notification if the job was canceled
+			// Otherwise we will leave it to the steps to handle when all notifications have been sent out to avoid conditions where the job complete notification gets received before the job step complete notification
+			if (job.AbortedByUserId is null)
 			{
-				IUser? user = await _userCollection.GetUserAsync(job.AbortedByUserId.Value, cancellationToken);
-				if (user is not null)
-				{
-					abortedByUserName = user.Name;
-				}
+				return;
+			}
+
+			string abortedByUserName = "Unknown User";
+			IUser? user = await _userCollection.GetUserAsync(job.AbortedByUserId.Value, cancellationToken);
+			if (user is not null)
+			{
+				abortedByUserName = user.Name;
 			}
 
 			// Send the updated summary after the steps has been updated
@@ -224,11 +228,11 @@ namespace HordeServer.Notifications.Sinks
 			foreach (JobNotificationStateDocument jobNotificationState in jobNotificationStates)
 			{
 				await SendSlackSummaryMessageAsync(job, branch, hordeJobURL, jobNotificationState.TemplateId, jobNotificationState.Recipient, abortedByUserName, cancellationToken);
-			}
 
-			// Delete all documents related to this job as no other notifications will be sent
-			_ = await DeleteJobNotificationStatesAsync(job.Id, cancellationToken);
-			_ = await DeleteJobStepNotificationStatesAsync(job.Id, cancellationToken);
+				// No further notifications will be sent out and can be cleaned up
+				_ = await DeleteJobStepNotificationStatesAsync(jobNotificationState.JobId, jobNotificationState.TemplateId, jobNotificationState.Recipient, cancellationToken);
+				_ = await DeleteJobNotificationStatesAsync(jobNotificationState.JobId, jobNotificationState.TemplateId, jobNotificationState.Recipient, cancellationToken);
+			}
 		}
 
 		/// <inheritdoc/>
@@ -405,6 +409,9 @@ namespace HordeServer.Notifications.Sinks
 			foreach (JobNotificationStateDocument jobNotificationState in jobNotificationStates)
 			{
 				await SendSlackSummaryMessageAsync(job, branch, hordeJobURL, jobNotificationState.TemplateId, jobNotificationState.Recipient, null, cancellationToken);
+
+				// Check to see if all notifications for this job were sent to then delete all documents related to this job as no other notifications will be sent
+				await TryNotificationCleanupAsync(jobNotificationState, cancellationToken);
 			}
 		}
 
@@ -500,6 +507,28 @@ namespace HordeServer.Notifications.Sinks
 			return dependencyJobStepNotifications;
 		}
 
+		/// <summary>
+		/// Clean up notification states once all notifications have been sent
+		/// </summary>
+		/// <param name="jobNotificationState">Job Notification State to check against</param>
+		/// <param name="cancellationToken">Cancellation token for the operation</param>
+		/// <returns></returns>
+		async Task TryNotificationCleanupAsync(JobNotificationStateDocument jobNotificationState, CancellationToken cancellationToken = default)
+		{
+			// Fetch all of our job steps that are still pending to determine if we still have notifications to send
+			FilterDefinition<JobStepNotificationStateDocument> filter = Builders<JobStepNotificationStateDocument>.Filter.Eq(x => x.JobId, jobNotificationState.JobId) &
+																			Builders<JobStepNotificationStateDocument>.Filter.Eq(x => x.TemplateId, jobNotificationState.TemplateId) &
+																			Builders<JobStepNotificationStateDocument>.Filter.Eq(x => x.Recipient, jobNotificationState.Recipient) &
+																			Builders<JobStepNotificationStateDocument>.Filter.Eq(x => x.Badge, PendingBadge);
+
+			List<JobStepNotificationStateDocument>? jobStepNotificationStates = await GetJobStepNotificationStatesAsync(filter, cancellationToken);
+			if (jobStepNotificationStates is null)
+			{
+				_ = await DeleteJobStepNotificationStatesAsync(jobNotificationState.JobId, jobNotificationState.TemplateId, jobNotificationState.Recipient, cancellationToken);
+				_ = await DeleteJobNotificationStatesAsync(jobNotificationState.JobId, jobNotificationState.TemplateId, jobNotificationState.Recipient, cancellationToken);
+			}
+		}
+
 		/// <summary>
 		/// Determines the badge based on the job step's state and outcome
 		/// </summary>
@@ -660,14 +689,18 @@ namespace HordeServer.Notifications.Sinks
 		}
 
 		/// <summary>
-		/// Deletes all <see cref="JobNotificationStateDocument"/> from the collection
+		/// Deletes <see cref="JobNotificationStateDocument"/> from the collection
 		/// </summary>
 		/// <param name="jobId">Identifier of the job</param>
+		/// <param name="templateId">Identifier of the stream's template</param>
+		/// <param name="recipient">Identifier of the channel or user in Slack or user's email address</param>
 		/// <param name="cancellationToken">Cancellation token for the operation</param>
 		/// <returns>True if records were deleted, otherwise false</returns>
-		async Task<bool> DeleteJobNotificationStatesAsync(JobId jobId, CancellationToken cancellationToken = default)
+		async Task<bool> DeleteJobNotificationStatesAsync(JobId jobId, TemplateId templateId, string recipient, CancellationToken cancellationToken = default)
 		{
-			FilterDefinition<JobNotificationStateDocument> filter = Builders<JobNotificationStateDocument>.Filter.Eq(x => x.JobId, jobId);
+			FilterDefinition<JobNotificationStateDocument> filter = Builders<JobNotificationStateDocument>.Filter.Eq(x => x.JobId, jobId) &
+																		Builders<JobNotificationStateDocument>.Filter.Eq(x => x.TemplateId, templateId) &
+																		Builders<JobNotificationStateDocument>.Filter.Eq(x => x.Recipient, recipient);
 			DeleteResult result = await _jobNotificationStates.DeleteManyAsync(filter, cancellationToken);
 			return result.DeletedCount > 0;
 		}
@@ -697,9 +730,9 @@ namespace HordeServer.Notifications.Sinks
 																		Builders<JobStepNotificationStateDocument>.Filter.Eq(x => x.TargetPlatform, platform);
 			UpdateDefinition<JobStepNotificationStateDocument> update = Builders<JobStepNotificationStateDocument>.Update.SetOnInsert(x => x.Id, newId).Set(x => x.Badge, badge);
 
-			if (messageId != null)
+			if (messageId is not null)
 			{
-				update = update.Set(x => x.Channel, messageId.Channel).Set(x => x.Ts, messageId.Ts);
+				update = update.Set(x => x.Channel, messageId.Channel).Set(x => x.ThreadTs, messageId.ThreadTs).Set(x => x.Ts, messageId.Ts);
 			}
 
 			JobStepNotificationStateDocument state = await _jobStepNotificationStates.FindOneAndUpdateAsync(filter, update, new FindOneAndUpdateOptions<JobStepNotificationStateDocument> { IsUpsert = true, ReturnDocument = ReturnDocument.After }, cancellationToken);
@@ -733,14 +766,18 @@ namespace HordeServer.Notifications.Sinks
 		}
 
 		/// <summary>
-		/// Deletes all <see cref="JobStepNotificationStateDocument"/> from the collection
+		/// Deletes <see cref="JobStepNotificationStateDocument"/> from the collection
 		/// </summary>
 		/// <param name="jobId">Identifier of the job</param>
+		/// <param name="templateId">Identifier of the stream's template</param>
+		/// <param name="recipient">Identifier of the channel or user in Slack or user's email address</param>
 		/// <param name="cancellationToken">Cancellation token for the operation</param>
 		/// <returns>True if records were deleted, otherwise false</returns>
-		async Task<bool> DeleteJobStepNotificationStatesAsync(JobId jobId, CancellationToken cancellationToken = default)
+		async Task<bool> DeleteJobStepNotificationStatesAsync(JobId jobId, TemplateId templateId, string recipient, CancellationToken cancellationToken = default)
 		{
-			FilterDefinition<JobStepNotificationStateDocument> filter = Builders<JobStepNotificationStateDocument>.Filter.Eq(x => x.JobId, jobId);
+			FilterDefinition<JobStepNotificationStateDocument> filter = Builders<JobStepNotificationStateDocument>.Filter.Eq(x => x.JobId, jobId) &
+																		Builders<JobStepNotificationStateDocument>.Filter.Eq(x => x.TemplateId, templateId) &
+																		Builders<JobStepNotificationStateDocument>.Filter.Eq(x => x.Recipient, recipient);
 			DeleteResult result = await _jobStepNotificationStates.DeleteManyAsync(filter, cancellationToken);
 			return result.DeletedCount > 0;
 		}
@@ -775,6 +812,16 @@ namespace HordeServer.Notifications.Sinks
 			return CreateSlackSummaryNotification(job, branch, hordeJobURL, jobStepNotificationStateDocuments, abortedByUserName);
 		}
 
+		/// <summary>
+		/// Send a notification on a job setup failure
+		/// </summary>
+		/// <param name="job">Information about the job</param>
+		/// <param name="step">Information about the step</param>
+		/// <param name="branch">Name of the branch from the stream</param>
+		/// <param name="jobStepEventData">List of the step's logged events</param>
+		/// <param name="jobNotificationConfigs">Information about the configuration tied to his job</param>
+		/// <param name="cancellationToken">Cancellation token for the operation</param>
+		/// <returns></returns>
 		async Task SendSetupFailedSlackMessageAsync(IJob job, IJobStep step, string branch, List<LogData> jobStepEventData, List<JobNotificationConfig> jobNotificationConfigs, CancellationToken cancellationToken = default)
 		{
 			SlackMessage message = new SlackMessage();
@@ -868,7 +915,7 @@ namespace HordeServer.Notifications.Sinks
 		async Task<SlackMessageId> SendSlackThreadMessageAsync(SlackMessageId jobMessageId, JobStepNotificationStateDocument jobStepNotificationState, SlackMessage threadMessage, CancellationToken cancellationToken = default)
 		{
 			// Check to see if a message has already been sent out to avoid sending duplicate notifications
-			if (jobStepNotificationState.MessageId.ThreadTs is not null)
+			if (!String.IsNullOrEmpty(jobStepNotificationState.MessageId.ThreadTs))
 			{
 				return jobStepNotificationState.MessageId;
 			}
