Horde: Fix warning with parameter in telemetry controller
#rnx

diff --git a/Engine/Source/Programs/Horde/Plugins/Analytics/HordeServer.Analytics/Telemetry/TelemetryController.cs b/Engine/Source/Programs/Horde/Plugins/Analytics/HordeServer.Analytics/Telemetry/TelemetryController.cs
--- a/Engine/Source/Programs/Horde/Plugins/Analytics/HordeServer.Analytics/Telemetry/TelemetryController.cs
+++ b/Engine/Source/Programs/Horde/Plugins/Analytics/HordeServer.Analytics/Telemetry/TelemetryController.cs
@@ -91,8 +91,7 @@ namespace HordeServer.Telemetry
 		/// <summary>
 		/// Queries aggregated metrics from the telemetry system
 		/// </summary>
-		/// <param name="storeId">The telemetry store id</param>
-		/// <param name="request">The metrics request</param>
+		/// <param name="storeId">The telemetry store id</param>		
 		/// <param name="id">The metrics to query</param>
 		/// <param name="minTime">Minimum time interval to query</param>
 		/// <param name="maxTime">Maximum time interval to query</param>		
