Horde: Do not filter to top/bottom for prod deployment
#rnx

diff --git a/Engine/Source/Programs/Horde/Plugins/Analytics/HordeServer.Analytics.Tests/Telemetry/TelemetryTests.cs b/Engine/Source/Programs/Horde/Plugins/Analytics/HordeServer.Analytics.Tests/Telemetry/TelemetryTests.cs
--- a/Engine/Source/Programs/Horde/Plugins/Analytics/HordeServer.Analytics.Tests/Telemetry/TelemetryTests.cs
+++ b/Engine/Source/Programs/Horde/Plugins/Analytics/HordeServer.Analytics.Tests/Telemetry/TelemetryTests.cs
@@ -445,7 +445,7 @@ namespace HordeServer.Analytics.Tests.Telemetry
 
 		[TestMethod]
 		public async Task TopAndBottomAsync()
-		{
+		{   /*
 			// TopN
 			{
 				Clock.UtcNow = new DateTime(2023, 6, 8, 4, 30, 0, DateTimeKind.Utc);
@@ -574,6 +574,7 @@ namespace HordeServer.Analytics.Tests.Telemetry
 				Assert.AreEqual(1, metrics[2].Count);
 				Assert.AreEqual(5, metrics[2].Value);
 			}
+			*/
 		}		
 	}
 }
diff --git a/Engine/Source/Programs/Horde/Plugins/Analytics/HordeServer.Analytics/Telemetry/Metrics/MetricCollection.cs b/Engine/Source/Programs/Horde/Plugins/Analytics/HordeServer.Analytics/Telemetry/Metrics/MetricCollection.cs
--- a/Engine/Source/Programs/Horde/Plugins/Analytics/HordeServer.Analytics/Telemetry/Metrics/MetricCollection.cs
+++ b/Engine/Source/Programs/Horde/Plugins/Analytics/HordeServer.Analytics/Telemetry/Metrics/MetricCollection.cs
@@ -442,6 +442,7 @@ namespace HordeServer.Telemetry.Metrics
 				filter &= Builders<MetricDocument>.Filter.In(x => x.Group, groups);
 			}
 
+			/*
 			List<MetricDocument> results = new List<MetricDocument>();
 			List<MetricConfig> configs = metricConfigs.Where(x => x.TopN == 0 && x.BottomN == 0).ToList();
 
@@ -474,7 +475,11 @@ namespace HordeServer.Telemetry.Metrics
 					results.AddRange(await _metrics.Find(tfilter).SortByDescending(x => x.Time).Limit(maxResults).ToListAsync(cancellationToken));
 				}
 			}
+			*/
 
+			List<MetricDocument> results = new List<MetricDocument>();
+			FilterDefinition<MetricDocument> metricFilter = filter & Builders<MetricDocument>.Filter.In(x => x.MetricId, metricIds);
+			results = await _metrics.Find(metricFilter).SortByDescending(x => x.Time).Limit(maxResults).ToListAsync(cancellationToken);
 			return results.ConvertAll<MetricDocument, IMetric>(x => x);
 		}
 	}
