Horde: Address yet more warnings
#rnx

diff --git a/Engine/Source/Programs/Horde/Plugins/Analytics/HordeServer.Analytics.Tests/Telemetry/TelemetryTests.cs b/Engine/Source/Programs/Horde/Plugins/Analytics/HordeServer.Analytics.Tests/Telemetry/TelemetryTests.cs
--- a/Engine/Source/Programs/Horde/Plugins/Analytics/HordeServer.Analytics.Tests/Telemetry/TelemetryTests.cs
+++ b/Engine/Source/Programs/Horde/Plugins/Analytics/HordeServer.Analytics.Tests/Telemetry/TelemetryTests.cs
@@ -443,9 +443,10 @@ namespace HordeServer.Analytics.Tests.Telemetry
 			}
 		}
 
+		/*
 		[TestMethod]
 		public async Task TopAndBottomAsync()
-		{   /*
+		{   
 			// TopN
 			{
 				Clock.UtcNow = new DateTime(2023, 6, 8, 4, 30, 0, DateTimeKind.Utc);
@@ -573,8 +574,8 @@ namespace HordeServer.Analytics.Tests.Telemetry
 				Assert.AreEqual("third", metrics[2].Group);
 				Assert.AreEqual(1, metrics[2].Count);
 				Assert.AreEqual(5, metrics[2].Value);
-			}
-			*/
-		}		
+			}			
+		}
+		*/
 	}
 }
diff --git a/Engine/Source/Programs/Horde/Plugins/Analytics/HordeServer.Analytics/Telemetry/Metrics/MetricCollection.cs b/Engine/Source/Programs/Horde/Plugins/Analytics/HordeServer.Analytics/Telemetry/Metrics/MetricCollection.cs
--- a/Engine/Source/Programs/Horde/Plugins/Analytics/HordeServer.Analytics/Telemetry/Metrics/MetricCollection.cs
+++ b/Engine/Source/Programs/Horde/Plugins/Analytics/HordeServer.Analytics/Telemetry/Metrics/MetricCollection.cs
@@ -355,12 +355,14 @@ namespace HordeServer.Telemetry.Metrics
 			}
 		}
 
+		/*
 		class MetricProjection
 		{
 			public string _group = String.Empty;
 			public IEnumerable<ObjectId>? _ids;
 			public double _total;
 		}
+		*/
 
 		/// <inheritdoc/>
 		public async Task<List<IMetricMeta>> FindMetaAsync(TelemetryStoreId telemetryStoreId, MetricId[] metricIds, DateTime? minTime = null, DateTime? maxTime = null, int maxResults = 50, CancellationToken cancellationToken = default)
