Adjusting Slack notification format and improving cleanup logic when a job completes

#rnx

diff --git a/Engine/Source/Programs/Horde/Plugins/Experimental/HordeServer.Experimental/Notifications/Sinks/ExperimentalSlackNotificationSink.cs b/Engine/Source/Programs/Horde/Plugins/Experimental/HordeServer.Experimental/Notifications/Sinks/ExperimentalSlackNotificationSink.cs
--- a/Engine/Source/Programs/Horde/Plugins/Experimental/HordeServer.Experimental/Notifications/Sinks/ExperimentalSlackNotificationSink.cs
+++ b/Engine/Source/Programs/Horde/Plugins/Experimental/HordeServer.Experimental/Notifications/Sinks/ExperimentalSlackNotificationSink.cs
@@ -195,14 +195,6 @@ namespace HordeServer.Notifications.Sinks
 		/// <inheritdoc/>
 		public async Task NotifyJobCompleteAsync(IJob job, IGraph graph, LabelOutcome outcome, CancellationToken cancellationToken)
 		{
-			string hordeJobURL = new Uri(_serverInfo.DashboardUrl, $"job/{job.Id}").ToString();
-
-			string branch = "Unknown Branch";
-			if (_buildConfig.CurrentValue.TryGetStream(job.StreamId, out StreamConfig? config))
-			{
-				branch = config.Name;
-			}
-
 			// Only handle the job complete notification if the job was canceled
 			// Otherwise we will leave it to the steps to handle when all notifications have been sent out to avoid conditions where the job complete notification gets received before the job step complete notification
 			if (job.AbortedByUserId is null)
@@ -210,13 +202,6 @@ namespace HordeServer.Notifications.Sinks
 				return;
 			}
 
-			string abortedByUserName = "Unknown User";
-			IUser? user = await _userCollection.GetUserAsync(job.AbortedByUserId.Value, cancellationToken);
-			if (user is not null)
-			{
-				abortedByUserName = user.Name;
-			}
-
 			// Send the updated summary after the steps has been updated
 			List<JobNotificationStateDocument>? jobNotificationStates = await GetJobNotificationStatesAsync(job.Id, cancellationToken);
 			if (jobNotificationStates is null)
@@ -225,8 +210,26 @@ namespace HordeServer.Notifications.Sinks
 				return;
 			}
 
+			string hordeJobURL = new Uri(_serverInfo.DashboardUrl, $"job/{job.Id}").ToString();
+
+			string branch = "Unknown Branch";
+			if (_buildConfig.CurrentValue.TryGetStream(job.StreamId, out StreamConfig? config))
+			{
+				branch = config.Name;
+			}
+
+			string abortedByUserName = "Unknown User";
+			IUser? user = await _userCollection.GetUserAsync(job.AbortedByUserId.Value, cancellationToken);
+			if (user is not null)
+			{
+				abortedByUserName = user.Name;
+			}
+
 			foreach (JobNotificationStateDocument jobNotificationState in jobNotificationStates)
 			{
+				// Update our steps to reflect the canceled job
+				_ = await UpdatePendingJobStepNotificationsAsync(job, jobNotificationState, cancellationToken);
+
 				await SendSlackSummaryMessageAsync(job, branch, hordeJobURL, jobNotificationState.TemplateId, jobNotificationState.Recipient, abortedByUserName, cancellationToken);
 
 				// No further notifications will be sent out and can be cleaned up
@@ -392,6 +395,16 @@ namespace HordeServer.Notifications.Sinks
 				}
 			}
 
+			// Check if the job is marked as being completed and cleanup any steps that we fell out of sync with
+			List<JobNotificationStateDocument>? jobNotificationStates = await GetJobNotificationStatesAsync(job.Id, cancellationToken);
+			if (jobNotificationStates is not null && job.GetState() == JobState.Complete)
+			{
+				foreach (JobNotificationStateDocument jobNotificationState in jobNotificationStates)
+				{
+					needsSummaryUpdate |= await UpdatePendingJobStepNotificationsAsync(job, jobNotificationState, cancellationToken);
+				}
+			}
+
 			// Check to see if any of our steps were updated which requires the Slack summary to be updated
 			if (!needsSummaryUpdate)
 			{
@@ -399,7 +412,6 @@ namespace HordeServer.Notifications.Sinks
 			}
 
 			// Send the updated summary after the steps has been updated
-			List<JobNotificationStateDocument>? jobNotificationStates = await GetJobNotificationStatesAsync(job.Id, cancellationToken);
 			if (jobNotificationStates is null)
 			{
 				_logger.LogWarning("No job notifications found for job id '{JobId}' even though job step '{StepId}' was updated.", job.Id, step.Id);
@@ -507,6 +519,48 @@ namespace HordeServer.Notifications.Sinks
 			return dependencyJobStepNotifications;
 		}
 
+		/// <summary>
+		/// Syncs up the job information with our current recorded steps
+		/// </summary>
+		/// <param name="job">Information about the job</param>
+		/// <param name="jobNotificationState">Job Notification State to check against</param>
+		/// <param name="cancellationToken">Cancellation token for the operation</param>
+		/// <returns></returns>
+		async Task<bool> UpdatePendingJobStepNotificationsAsync(IJob job, JobNotificationStateDocument jobNotificationState, CancellationToken cancellationToken = default)
+		{
+			bool wasUpdated = false;
+			// Fetch all of our job steps that are still pending to determine if we still have notifications to send
+			FilterDefinition<JobStepNotificationStateDocument> filter = Builders<JobStepNotificationStateDocument>.Filter.Eq(x => x.JobId, jobNotificationState.JobId) &
+																			Builders<JobStepNotificationStateDocument>.Filter.Eq(x => x.TemplateId, jobNotificationState.TemplateId) &
+																			Builders<JobStepNotificationStateDocument>.Filter.Eq(x => x.Recipient, jobNotificationState.Recipient) &
+																			Builders<JobStepNotificationStateDocument>.Filter.Eq(x => x.Badge, PendingBadge);
+
+			List<JobStepNotificationStateDocument>? jobStepNotificationStates = await GetJobStepNotificationStatesAsync(filter, cancellationToken);
+			if (jobStepNotificationStates is null)
+			{
+				return false;
+			}
+
+			foreach (JobStepNotificationStateDocument jobStepNotificationState in jobStepNotificationStates)
+			{
+				IJobStep? step = null;
+				if (!job.TryGetStep(jobStepNotificationState.JobStepId, out step))
+				{
+					continue;
+				}
+
+				string badge = GetBadgeForStep(step);
+				if (badge != jobStepNotificationState.Badge)
+				{
+					_ = await AddOrUpdateJobStepNotificationStateAsync(jobStepNotificationState.JobId, jobStepNotificationState.TemplateId, jobStepNotificationState.Recipient, jobStepNotificationState.JobStepId,
+																		jobStepNotificationState.Group, jobStepNotificationState.TargetPlatform, badge, null, cancellationToken);
+					wasUpdated = true;
+				}
+			}
+
+			return wasUpdated;
+		}
+
 		/// <summary>
 		/// Clean up notification states once all notifications have been sent
 		/// </summary>
@@ -843,7 +897,7 @@ namespace HordeServer.Notifications.Sinks
 			infoBlock.Elements.Add(new TextObject($"<{jobStepLink}|*{step.Name}*> encountered errors"));
 			message.Blocks.Add(infoBlock);
 
-			AddLogDataContext(message, jobStepEventData);
+			AddLogDataContext(message, jobStepEventData, step.LogId);
 
 			foreach (JobNotificationConfig config in jobNotificationConfigs)
 			{
@@ -982,7 +1036,8 @@ namespace HordeServer.Notifications.Sinks
 		/// </summary>
 		/// <param name="slackBlockContainer">Slack block container</param>
 		/// <param name="logData">List of our log information</param>
-		void AddLogDataContext(ISlackBlockContainer slackBlockContainer, List<LogData> logData)
+		/// <param name="logId">Optional identifier of the log</param>
+		void AddLogDataContext(ISlackBlockContainer slackBlockContainer, List<LogData> logData, LogId? logId = null)
 		{
 			// Don't clutter the message with all of the events and do a preliminary check to see if the user needs to be made aware that events have been truncated
 			if (logData.Count > MaxEventMessages)
@@ -997,15 +1052,15 @@ namespace HordeServer.Notifications.Sinks
 				ContextBlock logContext = new ContextBlock();
 				logContext.Elements.Add(new TextObject($"```{log.Message}```"));
 				slackBlockContainer.Blocks.Add(logContext);
+			}
 
-				if (log.LogId is not null)
-				{
-					string logUrl = new Uri(_serverInfo.DashboardUrl, $"log/{log.LogId}").ToString();
+			if (logId is not null)
+			{
+				string logUrl = new Uri(_serverInfo.DashboardUrl, $"log/{logId}").ToString();
 
-					ContextBlock detailsContext = new ContextBlock();
-					detailsContext.Elements.Add(new TextObject($"<{logUrl}|View Job Step Log>"));
-					slackBlockContainer.Blocks.Add(detailsContext);
-				}
+				ContextBlock detailsContext = new ContextBlock();
+				detailsContext.Elements.Add(new TextObject($"<{logUrl}|View Job Step Log>"));
+				slackBlockContainer.Blocks.Add(detailsContext);
 			}
 		}
 
@@ -1056,12 +1111,12 @@ namespace HordeServer.Notifications.Sinks
 				string dependencyJobStepLink = new Uri($"{_serverInfo.DashboardUrl}job/{step.Job.Id}?step={underlyingStep.Id}").ToString();
 
 				attachment.AddContext(new TextObject($"<{dependencyJobStepLink}|*{underlyingStep.Name}*>"));
-				AddLogDataContext(attachment, data);
+				AddLogDataContext(attachment, data, underlyingStep.LogId);
 				alertMessage.Attachments.Add(attachment);
 			}
 			else
 			{
-				AddLogDataContext(alertMessage, data);
+				AddLogDataContext(alertMessage, data, step.LogId);
 			}
 
 			return alertMessage;
