Added Debug Probe metric to Horde Analytics
Captures sum of SchemaVersion where SchemaVersion>=0 grouped by EventName and SessionLabel name

diff --git a/Engine/Source/Programs/Horde/HordeServer/Defaults/default-metrics.telemetry.json b/Engine/Source/Programs/Horde/HordeServer/Defaults/default-metrics.telemetry.json
--- a/Engine/Source/Programs/Horde/HordeServer/Defaults/default-metrics.telemetry.json
+++ b/Engine/Source/Programs/Horde/HordeServer/Defaults/default-metrics.telemetry.json
@@ -801,6 +801,14 @@
 			"property": "$.Payload.Duration",
 			"function": "Average",
 			"interval": "00:30:00"
+		},
+		{
+			"id": "debug-probe",
+			"filter": "$[?(@.Payload.SchemaVersion >= 0 )]",
+			"groupBy": "$.Payload.Build_BranchName,$.Payload.Session_Label,$.Payload.EventName",
+			"property": "$.Payload.SchemaVersion",
+			"function": "Sum",
+			"interval": "00:05:00"
 		}
 	]
 }
\ No newline at end of file
diff --git a/Engine/Source/Programs/Horde/HordeServer/Defaults/default-views.telemetry.json b/Engine/Source/Programs/Horde/HordeServer/Defaults/default-views.telemetry.json
--- a/Engine/Source/Programs/Horde/HordeServer/Defaults/default-views.telemetry.json
+++ b/Engine/Source/Programs/Horde/HordeServer/Defaults/default-views.telemetry.json
@@ -60,6 +60,14 @@
 					"name": "StepName",
 					"group": "StepName"
 				},
+				{
+					"name": "EventName",
+					"group": "Event"
+				},
+				{
+					"name": "Session_Label",
+					"group": "Label"
+				},
 				{
 					"name": "Target",
 					"group": "Target",
@@ -566,7 +574,16 @@
 									"id": "core-content-virtualization-mp"
 								}
 							]
-						}
+						},
+						{
+							"name": "Debug",
+							"display": "Value",
+							"metrics": [
+								{
+									"id": "debug-probe"
+								}
+							]
+						}		
 					]
 				},
 				{
