Horde Analytics Dashboard.
Fix for the variables so they show the correct values

diff --git a/Engine/Source/Programs/Horde/HordeServer/Defaults/default-views.telemetry.json b/Engine/Source/Programs/Horde/HordeServer/Defaults/default-views.telemetry.json
--- a/Engine/Source/Programs/Horde/HordeServer/Defaults/default-views.telemetry.json
+++ b/Engine/Source/Programs/Horde/HordeServer/Defaults/default-views.telemetry.json
@@ -61,12 +61,12 @@
 					"group": "StepName"
 				},
 				{
-					"name": "EventName",
-					"group": "Event"
+					"name": "Event",
+					"group": "EventName"
 				},
 				{
-					"name": "Session_Label",
-					"group": "Label"
+					"name": "Label",
+					"group": "Session_Label"
 				},
 				{
 					"name": "Target",
