Horde Analytics
Moved debug probe metrics charting into new Debug Category

diff --git a/Engine/Source/Programs/Horde/HordeServer/Defaults/default-metrics.telemetry.json b/Engine/Source/Programs/Horde/HordeServer/Defaults/default-metrics.telemetry.json
--- a/Engine/Source/Programs/Horde/HordeServer/Defaults/default-metrics.telemetry.json
+++ b/Engine/Source/Programs/Horde/HordeServer/Defaults/default-metrics.telemetry.json
@@ -807,7 +807,7 @@
 			"filter": "$[?(@.Payload.SchemaVersion >= 0 )]",
 			"groupBy": "$.Payload.Build_BranchName,$.Payload.Session_Label,$.Payload.EventName",
 			"property": "$.Payload.SchemaVersion",
-			"function": "Sum",
+			"function": "Count",
 			"interval": "00:05:00"
 		}
 	]
diff --git a/Engine/Source/Programs/Horde/HordeServer/Defaults/default-views.telemetry.json b/Engine/Source/Programs/Horde/HordeServer/Defaults/default-views.telemetry.json
--- a/Engine/Source/Programs/Horde/HordeServer/Defaults/default-views.telemetry.json
+++ b/Engine/Source/Programs/Horde/HordeServer/Defaults/default-views.telemetry.json
@@ -574,16 +574,7 @@
 									"id": "core-content-virtualization-mp"
 								}
 							]
-						},
-						{
-							"name": "Debug",
-							"display": "Value",
-							"metrics": [
-								{
-									"id": "debug-probe"
-								}
-							]
-						}		
+						}
 					]
 				},
 				{
@@ -750,6 +741,20 @@
 							]
 						}
 					]
+				},
+				{
+					"name": "Debug",
+					"charts": [
+						{
+							"name": "Debug",
+							"display": "Value",
+							"metrics": [
+								{
+									"id": "debug-probe"
+								}
+							]
+						}
+					]
 				}
 			]
 		}
