HordeAnalytics:
Group debug metircs on ProjectName

diff --git a/Engine/Source/Programs/Horde/HordeServer/Defaults/default-metrics.telemetry.json b/Engine/Source/Programs/Horde/HordeServer/Defaults/default-metrics.telemetry.json
--- a/Engine/Source/Programs/Horde/HordeServer/Defaults/default-metrics.telemetry.json
+++ b/Engine/Source/Programs/Horde/HordeServer/Defaults/default-metrics.telemetry.json
@@ -805,7 +805,7 @@
 		{
 			"id": "debug-probe",
 			"filter": "$[?(@.Payload.SchemaVersion >= 0 )]",
-			"groupBy": "$.Payload.Build_BranchName,$.Payload.Session_Label,$.Payload.EventName",
+			"groupBy": "$.Payload.Build_BranchName,$.Payload.Session_Label,$.Payload.ProjectName,$.Payload.EventName",
 			"property": "$.Payload.SchemaVersion",
 			"function": "Count",
 			"interval": "00:05:00"
