Horde: Disable low cardinality indexes for Storage.Blobs collection

diff --git a/Engine/Source/Programs/Horde/Plugins/Storage/HordeServer.Storage/Storage/StorageService.cs b/Engine/Source/Programs/Horde/Plugins/Storage/HordeServer.Storage/Storage/StorageService.cs
--- a/Engine/Source/Programs/Horde/Plugins/Storage/HordeServer.Storage/Storage/StorageService.cs
+++ b/Engine/Source/Programs/Horde/Plugins/Storage/HordeServer.Storage/Storage/StorageService.cs
@@ -591,9 +591,17 @@ namespace HordeServer.Storage
 			_logger = logger;
 
 			List<MongoIndex<BlobInfo>> blobIndexes = new List<MongoIndex<BlobInfo>>();
-			blobIndexes.Add(keys => keys.Ascending(x => x.Imports));
+			
+			// Disabled as these have low cardinality leading to issues with GC inside AWS DocumentDB v4
+			// It cannot keep up leading to unbounded growth of index size (because vacuum/GC fails)
+			// Apparently v5 fixes this and normal MongoDB instances likely does not suffer from this limitation.
+			// With indexes removed, the GC operation inside Horde for storage blobs is effectively disabled as queries for imports
+			// cannot keep up instead.
+			
+			// blobIndexes.Add(keys => keys.Ascending(x => x.Imports));
+			// blobIndexes.Add(keys => keys.Ascending(x => x.NamespaceId).Ascending(s_blobAliasField));
+			
 			blobIndexes.Add(keys => keys.Ascending(x => x.NamespaceId).Ascending(x => x.Path), unique: true);
-			blobIndexes.Add(keys => keys.Ascending(x => x.NamespaceId).Ascending(s_blobAliasField));
 			_blobCollection = mongoService.GetCollection<BlobInfo>("Storage.Blobs", blobIndexes);
 
 			List<MongoIndex<RefInfo>> refIndexes = new List<MongoIndex<RefInfo>>();
