Horde Analytics : Fixed Content metrics

diff --git a/Engine/Source/Programs/Horde/HordeServer/Defaults/default-metrics.telemetry.json b/Engine/Source/Programs/Horde/HordeServer/Defaults/default-metrics.telemetry.json
--- a/Engine/Source/Programs/Horde/HordeServer/Defaults/default-metrics.telemetry.json
+++ b/Engine/Source/Programs/Horde/HordeServer/Defaults/default-metrics.telemetry.json
@@ -536,7 +536,7 @@
 			"property": "$.Payload.TotalFileSize",
 			"groupBy": "$.Payload.Build_BranchName",
 			"function": "Sum",
-			"interval": "00:5:00"
+			"interval": "00:05:00"
 		},
 		{
 			"id": "core-content-totalsize-mp",
@@ -544,7 +544,7 @@
 			"property": "$.Payload.TotalFileSize",
 			"groupBy": "$.Payload.Build_BranchName,$.Payload.Name",
 			"function": "Sum",
-			"interval": "00:5:00"
+			"interval": "00:05:00"
 		},
 		{
 			"id": "core-content-virtualization",
@@ -552,7 +552,7 @@
 			"property": "$.Payload.VirtualizedPercent",
 			"groupBy": "$.Payload.Build_BranchName",
 			"function": "Sum",
-			"interval": "00:5:00"
+			"interval": "00:05:00"
 		},
 		{
 			"id": "core-content-virtualization-mp",
@@ -560,7 +560,7 @@
 			"property": "$.Payload.VirtualizedPercent",
 			"groupBy": "$.Payload.Build_BranchName,$.Payload.Name",
 			"function": "Average",
-			"interval": "00:5:00"
+			"interval": "00:05:00"
 		},
 		{
 			"id": "horde-incremental-build-jobs",
