Horde: Update release notes
#rnx

diff --git a/Engine/Source/Programs/Horde/Docs/ReleaseNotes.md b/Engine/Source/Programs/Horde/Docs/ReleaseNotes.md
--- a/Engine/Source/Programs/Horde/Docs/ReleaseNotes.md
+++ b/Engine/Source/Programs/Horde/Docs/ReleaseNotes.md
@@ -1,5 +1,57 @@
 # Release Notes
 
+## 2024-12-09
+
+* Disable low cardinality indexes for Storage.Blobs collection (38662017)
+* Moved Clear Selection and All Templates to a single checkbox, and added a Select All/None for each category with at least 3 options. (38577053)
+* Add queryable Horde analytics meta data for better performance and to limit data set for top/bottom metrics (38571863)
+* Cloud DDC - Fixed build warning from change 38567473 (38567875)
+* Cloud DDC - Fixed an issue serializing a CbObject that was a part of a poco type. The collection of CbFields was picked up the serializer meaning the object was written as an array with fields (which is not valid json). (38567473)
+* Add support for enumerating analytics metrics and providing multi-group filters (38551348)
+* Improvements to analytics top/bottom N capture (38547327)
+* Improve agent telemetry display (38545837)
+* Make Preferences, user icon, and Tools middle-clickable. (38542886)
+* Add a bit more information when can't resolve preflight stream (38541374)
+* Cloud DDC - Fixed an issue were the json converted value of a object id was a base64 blob instead of a hex encoded string as we normally do. (38540397)
+* Cloud DDC - Fixed some issues when converting from a json object to a compact binary object in the builds api (38534145)
+* Make OIDC scopes and mappings overridable
+  - Due to .NET config only merging array entries when combining config sources, hardcoded OIDC defaults (scopes and claim mappings) cannot be overridden.
+  - Added OidcAddDefaultScopesAndMappings flag to control whether defaults are applied (38440554)
+* Added Content Metrics to Horde Analytics
+  Added Content vizualization to Horde Analytics
+  Records the total sizes and virtualiszation level of all source content (38440278)
+* EpicGames.Slack: Fix compiler warnings (38439811)
+* CpuCount and CpuMultiplier in AgentSettings used in LogicalCores property
+  - Limit the number of logical processors to CpuCount or number of logical processors reported by the OS, whichever is the smallest
+  - Scale the number of logical processors by CpuMultiplier
+  - Change CpuCount to a nullable value type to know when it is to be used
+  - Find the number of logical processors on Windows
+  - Set UE_HORDE_CPU_COUNT to CpuCount when CpuCount is set (38438211)
+* Fix crash in UHT by adding a default include path in case none were found (38374613)
+* Fix regex in mongo log/config check to be more inclusive (38335134)
+* Cloud DDC - Fixed build warnings with the newly added CbObjectId type (38332782)
+* Cloud DDC - Added initial version of a build api, this is intended to be used by cooked output snapshots and unsync builds. Its practically very similar to how snapshots already work but with a added index of which blocks exists so that it can be queried from the server rather then rely on input of a previous build which should help make better choices that increases dedup. 
+  This api also provides simple searching which is better then what we had for the DDC api (which just had enumeration under special circumstances) which is needed to identify which snapshot to use.
+  Lastly I added support for ObjectId in C# as this isn't something we have needed before but is used within objects for this API (is a 12 byte identifier that is valid within a object but not globally unique, it is similar to a uuid type 4).
+  The build store currently only supports Scylla and Memory implementations (e.g. no base cassandra or mongo) as we have concerns about those dbs ability to perform for this workload. (38329536)
+* Fixed Localisation warnings to correctly add a filepath which can be used for identifying suspected CL's which caused the issue. (38329203)
+* Fix issue with failed batches not being able to retry steps (38299359)
+* Merging Horde changes up to CL 38271598. (38274254)
+* Merging NuGet package upgrades to 5.5.
+  Source changes:
+  - 36793185 (Update Blake3)
+  - 36994680 (EpicGames.Core: Pin System.Text.Json to 8.0.5)
+  - 37007077 (Update Logging nuget packages)
+  - 37007445 (Missed a file from the logger nuget update)
+  - 37011466 (Update more nuget packages)
+  - 37033668 (Update more nuget packages (all system or ms that don't depend on System.Drawing.Common))
+  - 37097577 (Update nuget packages)
+  - 37753709 (Update nuget packages)
+* Fix missing IClock for dependency injection This missing dependency caused agent startup failure when EC2 mode was enabled. (38259940)
+* Add support for IMDSv2 for AWS EC2 in agent.  Used for detecting spot interruptions and auto-scaling changes on a running agent. (38251784)
+* PR #12533: Fix Horde configPath for globals.json with perforce:// uri not working (38221010)
+* Add additional full download size info to the current BuildDiff report (38220701)
+
 ## 2024-11-19
 
 * Rename default GetSparseClassData accessor to GetMutableSparseClassData to make mutability more obvious and to encourage use of more explicit accessor with EGetSparseClassDataMethod parameter - deprecating old version (38208446)
