HordeAnalytics
Added some debugging for VirtualAssets.MountPoint
Counts the number of VirtualAsset.MountPoint events every 5 minutes.

diff --git a/Engine/Source/Programs/Horde/HordeServer/Defaults/default-metrics.telemetry.json b/Engine/Source/Programs/Horde/HordeServer/Defaults/default-metrics.telemetry.json
--- a/Engine/Source/Programs/Horde/HordeServer/Defaults/default-metrics.telemetry.json
+++ b/Engine/Source/Programs/Horde/HordeServer/Defaults/default-metrics.telemetry.json
@@ -530,6 +530,14 @@
 			"interval": "00:15:00",
 			"topN": 25
 		},
+		{
+			"id": "core-content-events",
+			"filter": "$[?(@.Payload.EventName == 'Core.VirtualAssets.MountPoint')]",
+			"groupBy": "$.Payload.Build_BranchName,$.Payload.Session_Label,$.Payload.ProjectName",
+			"property": "$.Payload.SchemaVersion",
+			"function": "Count",
+			"interval": "00:05:00"
+		},
 		{
 			"id": "core-content-totalsize",
 			"filter": "$[?(@.Payload.EventName == 'Core.VirtualAssets.MountPoint')]",
diff --git a/Engine/Source/Programs/Horde/HordeServer/Defaults/default-views.telemetry.json b/Engine/Source/Programs/Horde/HordeServer/Defaults/default-views.telemetry.json
--- a/Engine/Source/Programs/Horde/HordeServer/Defaults/default-views.telemetry.json
+++ b/Engine/Source/Programs/Horde/HordeServer/Defaults/default-views.telemetry.json
@@ -539,6 +539,15 @@
 				{
 					"name": "Content",
 					"charts": [
+						{
+							"name": "Content Event Count",
+							"display": "Value",
+							"metrics": [
+								{
+									"id": "core-content-events"
+								}
+							]
+						},	
 						{
 							"name": "Content Size Total",
 							"display": "Value",
