HordeAnalytics
Added a chart to Debug category to track MountPoint events
Removed the duplicate chart from the Cotent categroy

diff --git a/Engine/Source/Programs/Horde/HordeServer/Defaults/default-views.telemetry.json b/Engine/Source/Programs/Horde/HordeServer/Defaults/default-views.telemetry.json
--- a/Engine/Source/Programs/Horde/HordeServer/Defaults/default-views.telemetry.json
+++ b/Engine/Source/Programs/Horde/HordeServer/Defaults/default-views.telemetry.json
@@ -762,6 +762,15 @@
 									"id": "debug-probe"
 								}
 							]
+						},
+						{
+							"name": "MountPoint Test",
+							"display": "Value",
+							"metrics": [
+								{
+									"id": "core-content-events"
+								}
+							]
 						}
 					]
 				}
