HordeAnalytics
Removed Cotent Event Count chart from Content dashbaord

diff --git a/Engine/Source/Programs/Horde/HordeServer/Defaults/default-views.telemetry.json b/Engine/Source/Programs/Horde/HordeServer/Defaults/default-views.telemetry.json
--- a/Engine/Source/Programs/Horde/HordeServer/Defaults/default-views.telemetry.json
+++ b/Engine/Source/Programs/Horde/HordeServer/Defaults/default-views.telemetry.json
@@ -539,15 +539,6 @@
 				{
 					"name": "Content",
 					"charts": [
-						{
-							"name": "Content Event Count",
-							"display": "Value",
-							"metrics": [
-								{
-									"id": "core-content-events"
-								}
-							]
-						},	
 						{
 							"name": "Content Size Total",
 							"display": "Value",
