Analytics documentation update to consider new plugins setup.

diff --git a/Engine/Source/Programs/Horde/Docs/Tutorials/Analytics.md b/Engine/Source/Programs/Horde/Docs/Tutorials/Analytics.md
--- a/Engine/Source/Programs/Horde/Docs/Tutorials/Analytics.md
+++ b/Engine/Source/Programs/Horde/Docs/Tutorials/Analytics.md
@@ -42,23 +42,19 @@ the following lines:
    [globals.json](../Config/Orientation.md) file:
 
     ```json
-    // Define the 'Engine' telemetry store and create some standard metrics within it.
-    "telemetryStores": [
-        {
-            "id": "engine",
-            "include": [
-                {
-                    "path": "$(HordeDir)/Defaults/default-metrics.telemetry.json"
-                }
-            ]
-        }
-    ],
-
-    // Configure a default dashboard to render them
-    "dashboard": {
-        "include": [
+    // Define the 'Engine' telemetry store and create some standard metrics & views with it.
+    "analytics": {
+        "stores": [
             {
-                "path": "$(HordeDir)/Defaults/default-analytics.dashboard.json"
+                "id": "engine",
+                "include": [
+                    {
+                        "path": "$(HordeDir)/Defaults/default-metrics.telemetry.json"
+                    },
+                    {
+                        "path": "$(HordeDir)/Defaults/default-views.telemetry.json"
+                    }
+                ]
             }
         ]
     },
