Horde: Adjusting timer for auto conform agents

Dropping to 30min ticks from 1hr

#rnx

diff --git a/Engine/Source/Programs/Horde/Plugins/Build/HordeServer.Build/Agents/Pools/PoolUpdateService.cs b/Engine/Source/Programs/Horde/Plugins/Build/HordeServer.Build/Agents/Pools/PoolUpdateService.cs
--- a/Engine/Source/Programs/Horde/Plugins/Build/HordeServer.Build/Agents/Pools/PoolUpdateService.cs
+++ b/Engine/Source/Programs/Horde/Plugins/Build/HordeServer.Build/Agents/Pools/PoolUpdateService.cs
@@ -42,7 +42,7 @@ namespace HordeServer.Agents.Pools
 			_logger = logger;
 			_updatePoolsTicker = clock.AddTicker($"{nameof(PoolUpdateService)}.{nameof(UpdatePoolsAsync)}", TimeSpan.FromSeconds(30.0), UpdatePoolsAsync, logger);
 			_shutdownDisabledAgentsTicker = clock.AddSharedTicker($"{nameof(PoolUpdateService)}.{nameof(ShutdownDisabledAgentsAsync)}", TimeSpan.FromHours(1), ShutdownDisabledAgentsAsync, logger);
-			_autoConformAgentsTicker = clock.AddSharedTicker($"{nameof(PoolUpdateService)}.{nameof(AutoConformAgentsAsync)}", TimeSpan.FromHours(1), AutoConformAgentsAsync, logger);
+			_autoConformAgentsTicker = clock.AddSharedTicker($"{nameof(PoolUpdateService)}.{nameof(AutoConformAgentsAsync)}", TimeSpan.FromMinutes(30), AutoConformAgentsAsync, logger);
 		}
 
 		/// <inheritdoc/>
