Horde: Add job/step information to device issue audit logging

diff --git a/Engine/Source/Programs/Horde/Plugins/Build/HordeServer.Build/Devices/DeviceCollection.cs b/Engine/Source/Programs/Horde/Plugins/Build/HordeServer.Build/Devices/DeviceCollection.cs
--- a/Engine/Source/Programs/Horde/Plugins/Build/HordeServer.Build/Devices/DeviceCollection.cs
+++ b/Engine/Source/Programs/Horde/Plugins/Build/HordeServer.Build/Devices/DeviceCollection.cs
@@ -679,11 +679,6 @@ namespace HordeServer.Devices
 					problemTime = utcNow;
 				}
 
-				if (log && problemTime != null && current!.ProblemTimeUtc != problemTime)
-				{
-					logger.LogInformation("Device reported a problem");
-				}
-
 				updates.Add(updateBuilder.Set(x => x.ProblemTimeUtc, problemTime));
 			}
 
@@ -734,9 +729,21 @@ namespace HordeServer.Devices
 				IDeviceReservation? reservation = await TryGetDeviceReservationAsync(deviceId);
 				if (reservation != null)
 				{
+					if (log)
+					{
+						 logger.LogInformation("Device reported a problem in step: {StepName} - job {JobId}, step {StepId})", reservation.StepName ?? "Unknown Step", reservation.JobId ?? "Unknown Job Id", reservation.StepId ?? "Unknown Step Id");
+					}
+
 					UpdateDefinition<DeviceTelemetryDocument> update = Builders<DeviceTelemetryDocument>.Update.Set(x => x.ProblemTimeUtc, utcNow);
 					await _deviceTelemetry.FindOneAndUpdateAsync(t => t.ReservationId == reservation.Id && t.DeviceId == deviceId, update);
 				}
+				else
+				{
+					if (log)
+					{
+						logger.LogInformation("Device reported a problem without reservation information");
+					}
+				}
 			}
 		}
 
