Horde: Make job redirector default to latest step regardless of state, with a complete query param option
#skipci

diff --git a/Engine/Source/Programs/Horde/HordeDashboard/src/components/JobRedirector.tsx b/Engine/Source/Programs/Horde/HordeDashboard/src/components/JobRedirector.tsx
--- a/Engine/Source/Programs/Horde/HordeDashboard/src/components/JobRedirector.tsx
+++ b/Engine/Source/Programs/Horde/HordeDashboard/src/components/JobRedirector.tsx
@@ -5,6 +5,7 @@ import { useNavigate } from 'react-router-dom';
 import backend from '../backend';
 import { useQuery } from './JobDetailCommon';
 import ErrorHandler from './ErrorHandler';
+import { JobStepState } from 'horde/backend/Api';
 
 function setError(message: string) {
 
@@ -32,6 +33,8 @@ export const JobRedirector: React.FC = () => {
    const templateId = !query.get("templateId") ? "" : query.get("templateId")!;
    const nodeName = !query.get("nodeName") ? "" : query.get("nodeName")!;
 
+   const complete = query.get("complete") !== null;
+
    if (!streamId || !templateId || !nodeName) {
       setError("Please specify streamId, templateId, and nodeName");
       return null;
@@ -39,18 +42,49 @@ export const JobRedirector: React.FC = () => {
 
    if (!state.backendQueried) {
 
-      backend.getJobStepHistory(streamId, nodeName, 1, templateId).then(results => {
-         if (!results?.length) {
-            setError(`No history results for ${streamId} / ${templateId} / ${nodeName}`);
-         }
+      if (!complete) {
+
+         backend.getJobs({ streamId: streamId, template: [templateId], count: 10, includePreflight: false }).then(results => {
+
+            const job = results.find(job => {
+               const batch = job.batches?.find(b => b.steps?.find(s => s.name === nodeName));
+               const step = batch?.steps?.find(s => s.name === nodeName);
+               if (!step) {
+                  return false;
+               }
+
+               return true;
+            })
+
+            if (job) {
+               const batch = job.batches?.find(b => b.steps?.find(s => s.name === nodeName))!;
+               const step = batch?.steps?.find(s => s.name === nodeName)!;
+
+               navigate(`/job/${job.id}?step=${step.id}`, { replace: true });
+
+            } else {
+               throw `Could not find job using step states`
+            }
+
+         }).catch(reason => {
+            setError(`Error getting last job for ${streamId} / ${templateId} / ${nodeName}: Please check parameters\n${reason}`);
+         })
+      }
+      else {
+         backend.getJobStepHistory(streamId, nodeName, 1, templateId).then(results => {
+            if (!results?.length) {
+               setError(`No history results for ${streamId} / ${templateId} / ${nodeName}`);
+               return;
+            }
 
-         const result = results[0];
+            const result = results[0];
 
-         navigate(`/job/${result.jobId}?step=${result.stepId}`, { replace: true });
+            navigate(`/job/${result.jobId}?step=${result.stepId}`, { replace: true });
 
-      }).catch(reason => {
-         setError(`Error getting last job for ${streamId} / ${templateId} / ${nodeName}: Please check parameters\n${reason}`);
-      });
+         }).catch(reason => {
+            setError(`Error getting last job for ${streamId} / ${templateId} / ${nodeName}: Please check parameters\n${reason}`);
+         });
+      }
 
       return null;
    }
