Horde: Support multiple templates in job redirector queries
#skipci

diff --git a/Engine/Source/Programs/Horde/HordeDashboard/src/components/JobRedirector.tsx b/Engine/Source/Programs/Horde/HordeDashboard/src/components/JobRedirector.tsx
--- a/Engine/Source/Programs/Horde/HordeDashboard/src/components/JobRedirector.tsx
+++ b/Engine/Source/Programs/Horde/HordeDashboard/src/components/JobRedirector.tsx
@@ -30,21 +30,30 @@ export const JobRedirector: React.FC = () => {
    const query = useQuery();
 
    const streamId = !query.get("streamId") ? "" : query.get("streamId")!;
-   const templateId = !query.get("templateId") ? "" : query.get("templateId")!;
+   const templateId = query.getAll("templateId")
    const nodeName = !query.get("nodeName") ? "" : query.get("nodeName")!;
 
    const complete = query.get("complete") !== null;
 
-   if (!streamId || !templateId || !nodeName) {
+   if (!streamId || !templateId.length || !nodeName) {
       setError("Please specify streamId, templateId, and nodeName");
       return null;
    }
 
+
+   if (complete && templateId.length > 1) {
+      setError("Please specify a single template id when using complete queries");
+      return null;
+   }
+
+
    if (!state.backendQueried) {
 
       if (!complete) {
 
-         backend.getJobs({ streamId: streamId, template: [templateId], count: 10, includePreflight: false }).then(results => {
+         backend.getJobs({ streamId: streamId, template: templateId, count: 10 * templateId.length, includePreflight: false, filter:"id,batches,change"}).then(results => {
+
+            results = results.sort((a, b) => (b.change ?? 0) - (a.change ?? 0))
 
             const job = results.find(job => {
                const batch = job.batches?.find(b => b.steps?.find(s => s.name === nodeName));
@@ -71,7 +80,7 @@ export const JobRedirector: React.FC = () => {
          })
       }
       else {
-         backend.getJobStepHistory(streamId, nodeName, 1, templateId).then(results => {
+         backend.getJobStepHistory(streamId, nodeName, 1, templateId[0]).then(results => {
             if (!results?.length) {
                setError(`No history results for ${streamId} / ${templateId} / ${nodeName}`);
                return;
