Horde: Collect system telemetry on Linux (CPU/mem/disk usage)

diff --git a/Engine/Source/Programs/Horde/HordeAgent/AgentApp.cs b/Engine/Source/Programs/Horde/HordeAgent/AgentApp.cs
--- a/Engine/Source/Programs/Horde/HordeAgent/AgentApp.cs
+++ b/Engine/Source/Programs/Horde/HordeAgent/AgentApp.cs
@@ -237,17 +237,21 @@ namespace HordeAgent
 
 		static ISystemMetrics CreateSystemMetrics(DirectoryReference workingDir, ILogger logger)
 		{
-			if (RuntimeInformation.IsOSPlatform(OSPlatform.Windows))
+			try
 			{
-				try
+				if (RuntimeInformation.IsOSPlatform(OSPlatform.Windows))
 				{
 					return new WindowsSystemMetrics(workingDir);
 				}
-				catch (Exception e)
+				if (RuntimeInformation.IsOSPlatform(OSPlatform.Linux))
 				{
-					logger.LogError(e, "Unable to initialize system metric collector for telemetry. Disabling. Reason: {Message}", e.Message);
+					return new LinuxSystemMetrics(workingDir);
 				}
 			}
+			catch (Exception e)
+			{
+				logger.LogError(e, "Unable to initialize system metric collector for telemetry. Disabling. Reason: {Message}", e.Message);
+			}
 			return new DefaultSystemMetrics();
 		}
 
diff --git a/Engine/Source/Programs/Horde/HordeAgent/Services/TelemetryService.cs b/Engine/Source/Programs/Horde/HordeAgent/Services/TelemetryService.cs
--- a/Engine/Source/Programs/Horde/HordeAgent/Services/TelemetryService.cs
+++ b/Engine/Source/Programs/Horde/HordeAgent/Services/TelemetryService.cs
@@ -1,6 +1,7 @@
 // Copyright Epic Games, Inc. All Rights Reserved.
 
 using System.Diagnostics;
+using System.Diagnostics.CodeAnalysis;
 using System.Runtime.InteropServices;
 using System.Runtime.Versioning;
 using System.Text;
@@ -66,7 +67,7 @@ public class MemoryMetrics
 	/// <inheritdoc />
 	public override string ToString()
 	{
-		return $"Total={Total} kB, Available={Available} kB, Used={Used} kB, Used={UsedPercentage * 100.0:F1} %";
+		return $"Total={Total} kB, Available={Available} kB, Used={Used} kB, UsedPct={UsedPercentage * 100.0:F1} %";
 	}
 }
 
@@ -84,6 +85,12 @@ public class DiskMetrics
 	/// Total size of the drive
 	/// </summary>
 	public long TotalSize { get; set; }
+	
+	/// <inheritdoc />
+	public override string ToString()
+	{
+		return $"Total={TotalSize} kB, Free={FreeSpace}";
+	}
 }
 
 /// <summary>
@@ -246,6 +253,159 @@ public sealed class WindowsSystemMetrics : ISystemMetrics, IDisposable
 
 #pragma warning restore CA1416
 
+/// <summary>
+/// Linux specific implementation for gathering system metrics
+/// </summary>
+[SupportedOSPlatform("linux")]
+public sealed class LinuxSystemMetrics : ISystemMetrics, IDisposable
+{
+	private const string ProcMeminfoPath = "/proc/meminfo";
+	private const string ProcStatPath = "/proc/stat";
+	private const int UserIndex = 1;
+	private const int SystemIndex = 3;
+	private const int IdleIndex = 4;
+	
+	private readonly DirectoryReference _workingDir;
+	
+	/// <summary>
+	/// Constructor
+	/// </summary>
+	public LinuxSystemMetrics(DirectoryReference workingDir)
+	{
+		_workingDir = workingDir;
+	}
+	
+	/// <inheritdoc/>
+	public CpuMetrics? GetCpu()
+	{
+		const int SampleDurationMs = 100;
+		(long user, long system, long idle) firstSample = ReadCpuTimes();
+		Thread.Sleep(SampleDurationMs);
+		(long user, long system, long idle) secondSample = ReadCpuTimes();
+		
+		long userDelta = secondSample.user - firstSample.user;
+		long systemDelta = secondSample.system - firstSample.system;
+		long idleDelta = secondSample.idle - firstSample.idle;
+		long totalDelta = userDelta + systemDelta + idleDelta;
+		
+		if (totalDelta == 0)
+		{
+			return new CpuMetrics { User = 0, System = 0, Idle = 100 };
+		}
+		
+		return new CpuMetrics
+		{
+			User = (userDelta * 100.0f) / totalDelta,
+			System = (systemDelta * 100.0f) / totalDelta,
+			Idle = (idleDelta * 100.0f) / totalDelta,
+		};
+	}
+	
+	/// <inheritdoc/>
+	public MemoryMetrics? GetMemory()
+	{
+		long ParseMemInfoValue(string[] parts)
+		{
+			return Convert.ToInt64(parts[1].Replace(" kB", "", StringComparison.Ordinal));
+		}
+		
+		try
+		{
+			long memTotal = 0, memAvailable = 0, memFree = 0, buffers = 0, cached = 0, sReclaimAble = 0;
+			foreach (string line in File.ReadAllLines(ProcMeminfoPath))
+			{
+				string[] parts = line.Split(':', StringSplitOptions.TrimEntries);
+				string key = parts[0];
+				switch (key)
+				{
+					case "MemTotal": memTotal = ParseMemInfoValue(parts); break;
+					case "MemAvailable": memAvailable = ParseMemInfoValue(parts); break;
+					case "MemFree": memFree = ParseMemInfoValue(parts); break;
+					case "Buffers": buffers = ParseMemInfoValue(parts); break;
+					case "Cached": cached = ParseMemInfoValue(parts); break;
+					case "SReclaimable": sReclaimAble = ParseMemInfoValue(parts); break;
+				}
+			}
+
+			long used = memTotal - memFree - buffers - cached - sReclaimAble;
+			return new MemoryMetrics
+			{
+				Total = (uint)memTotal,
+				Available = (uint)memAvailable,
+				Used = (uint)used,
+				UsedPercentage = (float)used / memTotal,
+			};
+		}
+		catch (Exception)
+		{
+			return null;
+		}
+	}
+	
+	/// <inheritdoc/>
+	public DiskMetrics? GetDisk()
+	{
+		if (Statvfs(_workingDir.FullName, out StatvfsStruct stats) != 0)
+		{
+			int errno = Marshal.GetLastPInvokeError();
+			throw new IOException($"Failed to get disk metrics for path '{_workingDir.FullName}'. Error: {Marshal.GetPInvokeErrorMessage(errno)}");
+		}
+		
+		// Calculate total size and free space
+		// We use f_bavail (free blocks available to non-privileged users) instead of f_bfree as this is more representative
+		long blockSize = (long)stats.f_frsize;
+		long totalBlocks = (long)stats.f_blocks;
+		long availableBlocks = (long)stats.f_bavail;
+		
+		return new DiskMetrics
+		{
+			TotalSize = totalBlocks * blockSize,
+			FreeSpace = availableBlocks * blockSize,
+		};
+	}
+	
+	private static (long user, long system, long idle) ReadCpuTimes()
+	{
+		string statLine = File.ReadAllLines(ProcStatPath).First(line => line.StartsWith("cpu ", StringComparison.Ordinal));
+		long[] values = statLine.Split(' ', StringSplitOptions.RemoveEmptyEntries)
+			.Skip(1) // Skip "cpu" label
+			.Select(Int64.Parse)
+			.ToArray();
+		return (user: values[UserIndex - 1], system: values[SystemIndex - 1], idle: values[IdleIndex - 1]);
+	}
+	
+	[StructLayout(LayoutKind.Sequential)]
+	[SuppressMessage("Style", "IDE1006:Naming Styles", Justification = "POSIX ABI")]
+	[SuppressMessage("Style", "IDE0044:Add readonly modifier", Justification = "POSIX ABI")]
+	private struct StatvfsStruct
+	{
+		public ulong f_bsize;    // File system block size
+		public ulong f_frsize;   // Fundamental file system block size
+		public ulong f_blocks;   // Total number of blocks on file system in units of f_frsize
+		public ulong f_bfree;    // Total number of free blocks
+		public ulong f_bavail;   // Number of free blocks available to non-privileged process
+		public ulong f_files;    // Total number of file serial numbers
+		public ulong f_ffree;    // Total number of free file serial numbers
+		public ulong f_favail;   // Number of file serial numbers available to non-privileged process
+		public ulong f_fsid;     // File system ID
+		public ulong f_flag;     // Bit mask of f_flag values
+		public ulong f_namemax;  // Maximum filename length
+		private ulong __f_spare0;
+		private ulong __f_spare1;
+		private ulong __f_spare2;
+		private ulong __f_spare3;
+		private ulong __f_spare4;
+	}
+	
+	[DllImport("libc", SetLastError = true, EntryPoint = "statvfs")]
+	private static extern int Statvfs(string path, out StatvfsStruct buf);
+	
+	/// <inheritdoc/>
+	public void Dispose()
+	{
+	}
+}
+
 /// <summary>
 /// Send telemetry events back to server at regular intervals
 /// </summary>
