Horde: Use consts for LogicalCores and PhysicalCores values

diff --git a/Engine/Source/Programs/Horde/HordeAgent/Services/CapabilitiesService.cs b/Engine/Source/Programs/Horde/HordeAgent/Services/CapabilitiesService.cs
--- a/Engine/Source/Programs/Horde/HordeAgent/Services/CapabilitiesService.cs
+++ b/Engine/Source/Programs/Horde/HordeAgent/Services/CapabilitiesService.cs
@@ -565,13 +565,13 @@ namespace HordeAgent.Services
 
 			if (numLogicalCores > 0)
 			{
-				capabilities.Resources.Add("LogicalCores", numLogicalCores);
-				capabilities.Properties.Add($"LogicalCores={numLogicalCores}");
+				capabilities.Resources.Add(KnownResourceNames.LogicalCores, numLogicalCores);
+				capabilities.Properties.Add($"{KnownResourceNames.LogicalCores}={numLogicalCores}");
 			}
 
 			if (numPhysicalCores > 0)
 			{
-				capabilities.Properties.Add($"PhysicalCores={numPhysicalCores}");
+				capabilities.Properties.Add($"{KnownResourceNames.PhysicalCores}={numPhysicalCores}");
 			}
 		}
 
diff --git a/Engine/Source/Programs/Shared/EpicGames.Horde/Agents/KnownResourceNames.cs b/Engine/Source/Programs/Shared/EpicGames.Horde/Agents/KnownResourceNames.cs
--- a/Engine/Source/Programs/Shared/EpicGames.Horde/Agents/KnownResourceNames.cs
+++ b/Engine/Source/Programs/Shared/EpicGames.Horde/Agents/KnownResourceNames.cs
@@ -8,9 +8,14 @@ namespace EpicGames.Horde.Agents
 	public static class KnownResourceNames
 	{
 		/// <summary>
-		/// Number of logical cores
+		/// Number of logical cores including SMT cores (hyper-threading)
 		/// </summary>
 		public const string LogicalCores = "LogicalCores";
+		
+		/// <summary>
+		/// Number of physical cores
+		/// </summary>
+		public const string PhysicalCores = "PhysicalCores";
 
 		/// <summary>
 		/// Amount of RAM, in GB
