Horde: Added a flag to allow agents to auto restart on failure

When installing the agent via HordeAgent service install, you can specify -AutoRestartService=true/false, defaults to true. It will set the three failure recovery options to Restart the Service

diff --git a/Engine/Source/Programs/Horde/HordeAgent/Commands/Service/ServiceInstall.cs b/Engine/Source/Programs/Horde/HordeAgent/Commands/Service/ServiceInstall.cs
--- a/Engine/Source/Programs/Horde/HordeAgent/Commands/Service/ServiceInstall.cs
+++ b/Engine/Source/Programs/Horde/HordeAgent/Commands/Service/ServiceInstall.cs
@@ -15,7 +15,8 @@ namespace HordeAgent.Commands.Service
 		string? ServerProfile = null,
 		string? ConfigFiles = null,
 		string DotNetExecutable = "dotnet",
-		bool StartService = true
+		bool StartService = true,
+		bool AutoRestartService = true
 	);
 	
 	/// <summary>
@@ -53,6 +54,10 @@ namespace HordeAgent.Commands.Service
 		[Description("Paths to additional config JSON files to load, separated by colon (Linux/macOS) or semi-colon (Windows)")]
 		public string? ConfigFiles { get; set; } = null;
 
+		[CommandLine("-AutoRestartService=")]
+		[Description("Configures the Windows Service to autostart when HordeAgent crashes/exits abnormally. (true/false)")]
+		public string AutoRestartService { get; set; } = "true";
+
 		/// <summary>
 		/// Installs the service
 		/// </summary>
@@ -73,6 +78,7 @@ namespace HordeAgent.Commands.Service
 				ServerProfile = Server,
 				DotNetExecutable = DotNetExecutable,
 				StartService = Start != null && Start.Equals("true", StringComparison.OrdinalIgnoreCase),
+				AutoRestartService = AutoRestartService != null && AutoRestartService.Equals("true", StringComparison.OrdinalIgnoreCase)
 			};
 			return Task.FromResult(InstallService(options, logger));
 		}
@@ -140,6 +146,13 @@ namespace HordeAgent.Commands.Service
 			using (WindowsService service = serviceManager.Create(ServiceName, "Horde Agent", commandLine.ToString(), options.UserName, options.Password))
 			{
 				service.SetDescription("Allows this machine to participate in a Horde farm.");
+
+				// Set the service to auto restart if deemed to.
+				if (options.AutoRestartService)
+				{
+					service.ChangeServiceFailureActions([WindowsServiceRecoverAction.Restart, WindowsServiceRecoverAction.Restart, WindowsServiceRecoverAction.Restart]);
+				}
+
 				if (options.StartService)
 				{
 					logger.LogInformation("Starting...");
diff --git a/Engine/Source/Programs/Horde/HordeAgent/Utility/WindowsServiceManager.cs b/Engine/Source/Programs/Horde/HordeAgent/Utility/WindowsServiceManager.cs
--- a/Engine/Source/Programs/Horde/HordeAgent/Utility/WindowsServiceManager.cs
+++ b/Engine/Source/Programs/Horde/HordeAgent/Utility/WindowsServiceManager.cs
@@ -117,10 +117,60 @@ namespace HordeAgent.Utility
 			public string lpDescription;
 		}
 
+		public const int SERVICE_CONFIG_FAILURE_ACTIONS = 0x2;
+
+		[StructLayout(LayoutKind.Sequential)]
+		[SuppressMessage("Style", "IDE1006:Naming Styles", Justification = "Native struct")]
+		public struct SERVICE_FAILURE_ACTIONS
+		{
+			public int dwResetPeriod;
+
+			[MarshalAs(UnmanagedType.LPWStr)]
+			public string? lpRebootMsg;
+
+			[MarshalAs(UnmanagedType.LPWStr)]
+			public string? lpCommand;
+
+			public int cActions;
+			public IntPtr lpsaActions;
+		}
+
+		[StructLayout(LayoutKind.Sequential)]
+		[SuppressMessage("Style", "IDE1006:Naming Styles", Justification = "Native struct")]
+		public struct SERVICE_ACTION
+		{
+			public WindowsServiceRecoverAction type;
+			public uint dwDelay;
+		}
+
 		[DllImport("advapi32.dll", SetLastError = true, CharSet = CharSet.Auto)]
 		[return: MarshalAs(UnmanagedType.Bool)]
-		public static extern bool ChangeServiceConfig2(ServiceHandle hService, int dwInfoLevel, ref SERVICE_DESCRIPTION lpInfo);
+		public static extern bool ChangeServiceConfig2(ServiceHandle hService, int dwInfoLevel, IntPtr lpInfo);
 #pragma warning restore CS0649
+
+	}
+
+	enum WindowsServiceRecoverAction : int
+	{
+		/// <summary>
+		/// No Recover Action
+		/// </summary>
+		None = 0,
+
+		/// <summary>
+		/// Restart the service.
+		/// </summary>
+		Restart = 1,
+
+		/// <summary>
+		/// Reboot the host.
+		/// </summary>
+		Reboot = 2,
+
+		/// <summary>
+		/// Run a command
+		/// </summary>
+		RunCommand = 3
 	}
 
 	enum WindowsServiceStatus : uint
@@ -229,7 +279,77 @@ namespace HordeAgent.Utility
 		{
 			Native.SERVICE_DESCRIPTION descriptionData = new Native.SERVICE_DESCRIPTION();
 			descriptionData.lpDescription = description;
-			Native.ChangeServiceConfig2(_handle, Native.SERVICE_CONFIG_DESCRIPTION, ref descriptionData);
+
+			nint descriptionDataPointer = Marshal.AllocHGlobal(Marshal.SizeOf(descriptionData));
+			Marshal.StructureToPtr(descriptionData, descriptionDataPointer, false);
+
+			try
+			{
+				Native.ChangeServiceConfig2(_handle, Native.SERVICE_CONFIG_DESCRIPTION, descriptionDataPointer);
+			}
+			finally
+			{
+				Marshal.FreeHGlobal(descriptionDataPointer);
+			}
+		}
+
+		/// <summary>
+		/// Set the recovery actions for the windows service
+		/// </summary>
+		/// <param name="recoveryActions"></param>
+		/// <returns></returns>
+		public void ChangeServiceFailureActions(params WindowsServiceRecoverAction[] recoveryActions)
+		{
+			const int ActionsCount = 3;
+			const int ActionsSize = ActionsCount * sizeof(int) * 2;
+
+			if (recoveryActions == null)
+			{
+				throw new InvalidDataException("No recoveryActions specified, they must be specified.");
+			}
+
+			if (recoveryActions.Length != ActionsCount)
+			{
+				throw new InvalidDataException($"recoveryActions should contain {ActionsCount} and only contains {recoveryActions.Length} actions.");
+			}
+
+			// Represents SERVICE_ACTION a collection of ints
+			// Back of restart time for each retry attempt to accommodate for server/network issues
+			int[] serviceActions =
+			{
+				(int)recoveryActions[0], 10,
+				(int)recoveryActions[1], 30,
+				(int)recoveryActions[2], 60,
+			};
+
+			nint actionsBuffer = Marshal.AllocHGlobal(ActionsSize);
+			Marshal.Copy(serviceActions, 0, actionsBuffer, serviceActions.Length);
+
+			Native.SERVICE_FAILURE_ACTIONS serviceFailureActions = new()
+			{
+				cActions = ActionsCount,
+				dwResetPeriod = (int)TimeSpan.FromDays(30).TotalSeconds,
+				lpCommand = null,
+				lpRebootMsg = null,
+				lpsaActions = new IntPtr(actionsBuffer)
+			};
+
+			nint serviceActionsPointer = Marshal.AllocHGlobal(Marshal.SizeOf(serviceFailureActions));
+			Marshal.StructureToPtr(serviceFailureActions, serviceActionsPointer, false);
+
+			try
+			{
+				bool result = Native.ChangeServiceConfig2(_handle, Native.SERVICE_CONFIG_FAILURE_ACTIONS, serviceActionsPointer);
+				if (!result)
+				{
+					throw new Win32Exception($"Unable to set Service Failure Actions (error 0x{Marshal.GetLastWin32Error():X8})");
+				}
+			}
+			finally
+			{
+				Marshal.FreeHGlobal(actionsBuffer);
+				Marshal.FreeHGlobal(serviceActionsPointer);
+			}
 		}
 
 		/// <summary>
@@ -239,7 +359,7 @@ namespace HordeAgent.Utility
 		{
 			if (!Native.StartService(_handle, 0, null))
 			{
-				throw new Win32Exception(String.Format("Unable to start service (error 0x{0:X8})", Marshal.GetLastWin32Error()));
+				throw new Win32Exception($"Unable to start service (error 0x{Marshal.GetLastWin32Error():X8})");
 			}
 		}
 
@@ -254,7 +374,7 @@ namespace HordeAgent.Utility
 				int error = Marshal.GetLastWin32Error();
 				if (error != Native.ERROR_SERVICE_NOT_ACTIVE)
 				{
-					throw new Win32Exception(String.Format("Unable to stop service (error 0x{0:X8})", error));
+					throw new Win32Exception($"Unable to stop service (error 0x{error:X8})");
 				}
 			}
 		}
@@ -266,7 +386,7 @@ namespace HordeAgent.Utility
 		{
 			if (!Native.DeleteService(_handle))
 			{
-				throw new Win32Exception(String.Format("Unable to delete service (error 0x{0:X8})", Marshal.GetLastWin32Error()));
+				throw new Win32Exception($"Unable to delete service (error 0x{Marshal.GetLastWin32Error():X8})");
 			}
 		}
 
@@ -304,7 +424,7 @@ namespace HordeAgent.Utility
 				}
 				else
 				{
-					throw new Win32Exception(String.Format("Unable to open service manager (0x{0:X8}).", errorCode));
+					throw new Win32Exception($"Unable to open service manager (0x{errorCode:X8}).");
 				}
 			}
 		}
