Further improvement to the Device telemetry graphs
- introduce a graph selector
- add failure heat map at the bottom of the Usage (main) graph
- fix issue displaying long list of failure
- add visual cue to the Health list
- add a distribution graph using a pie chart combine with the timeline selector
- fix several visual and UX glitches

#rnx

diff --git a/Engine/Source/Programs/Horde/HordeDashboard/src/App.tsx b/Engine/Source/Programs/Horde/HordeDashboard/src/App.tsx
--- a/Engine/Source/Programs/Horde/HordeDashboard/src/App.tsx
+++ b/Engine/Source/Programs/Horde/HordeDashboard/src/App.tsx
@@ -15,7 +15,7 @@ import { AuditLogView } from './components/AuditLog';
 import { AutomationView } from './components/AutomationView';
 import { DashboardView } from './components/DashboardView';
 import { DebugView } from './components/DebugView';
-import { DeviceView } from './components/DeviceView';
+import { DeviceView } from './components/devices/DeviceView';
 import { ErrorDialog, ErrorHandler } from './components/ErrorHandler';
 import { JobRedirector } from './components/JobRedirector';
 import { LogView } from './components/LogView';
diff --git a/Engine/Source/Programs/Horde/HordeDashboard/src/backend/DeviceTelemetryHandler.ts b/Engine/Source/Programs/Horde/HordeDashboard/src/backend/DeviceTelemetryHandler.ts
new file mode 100644
--- /dev/null
+++ b/Engine/Source/Programs/Horde/HordeDashboard/src/backend/DeviceTelemetryHandler.ts
@@ -0,0 +1,540 @@
+import { action, makeObservable, observable } from "mobx";
+import moment from "moment";
+import backend from "../backend";
+import { projectStore } from "../backend/ProjectStore";
+import { DevicePoolTelemetryQuery, DevicePoolType, DeviceTelemetryQuery, GetDevicePlatformResponse, GetDevicePlatformTelemetryResponse, GetDevicePoolResponse, GetDevicePoolTelemetryResponse, GetDeviceResponse, GetDeviceTelemetryResponse, GetTelemetryInfoResponse } from "../backend/Api";
+import { msecToElapsed } from "../base/utilities/timeUtils";
+
+type PoolId = string;
+
+type PoolTelemetryData = {
+   date: Date;
+   dateString: string;
+   pools: Record<PoolId, GetDevicePlatformTelemetryResponse[]>;
+};
+
+export type PlatformTelemetry = GetDevicePlatformTelemetryResponse & {
+   date: Date;
+   dateString: string;
+};
+
+type DeviceTelemetry = GetTelemetryInfoResponse & {
+   createTime: Date;
+   reservationStart?: Date;
+   reservationFinish?: Date;
+   problemTime?: Date;
+}
+
+export type UsagePlatformTelemetry = { reservations: number, date: Date }
+
+export type TelemetryData = {
+   pool: GetDevicePoolResponse;
+   platform: GetDevicePlatformResponse;
+   streamIds: string[];
+   // stream id => full stream name
+   streamNames: Record<string, string>;
+   // stream full stream name => stream id
+   streamNamesReverse: Record<string, string>;
+   // stream name => valid selector
+   streamSelectors: Record<string, string>;
+   // selector => stream name
+   streamSelectorsReverse: Record<string, string>;
+   deviceIds: string[];
+   devices: Record<string, GetDeviceResponse>;
+   telemetry: PlatformTelemetry[];
+}
+
+type StepMetrics = {
+   numKits: number;
+   numProblems: number;
+   duration: string;
+}
+
+type ProblemStep = {
+   streamId?: string;
+   jobId?: string;
+   jobName?: string;
+   stepId?: string;
+   stepName?: string;
+}
+
+export type DeviceProblem = { deviceId: string, deviceName: string, poolName: string, problems: number, problemsRate: number, reservations: number, problemsDesc: string, latestProblem?: DeviceTelemetry, reservable: boolean };
+
+export class PoolTelemetryHandler {
+
+   constructor() {
+      makeObservable(this);
+      this.set();
+   }
+
+   getProblemSteps(deviceIds: string[], time: Date): Record<string, ProblemStep> {
+      const problems: Record<string, ProblemStep> = {};
+
+      deviceIds.forEach(id => {
+         const telemetry = this.deviceTelemetry.get(id);
+         if (!telemetry) {
+            return;
+         }
+
+         let best: DeviceTelemetry | undefined;
+         let bestMS = Number.MAX_VALUE;
+
+         telemetry.forEach(t => {
+
+            if (!t.problemTime) {
+               return;
+            }
+
+            const ctime = Math.abs(t.problemTime.getTime() - time.getTime());
+
+            if (ctime > 60 * 1000 * 30) {
+               return;
+            }
+
+            if (bestMS > ctime) {
+               best = t;
+               bestMS = ctime;
+            }
+
+         });
+
+         if (!best) {
+            return;
+         }
+
+         problems[id] = {
+            streamId: best.streamId,
+            jobId: best.jobId,
+            jobName: best.jobName,
+            stepId: best.stepId,
+            stepName: best.stepName
+         };
+
+      });
+
+      return problems;
+   }
+
+   getSteps(platformId: string, startDate: Date | undefined, endDate: Date | undefined) {
+
+      const stepData: Map<string, { stream: string, pool: string, stepName: string, reservations: number, problems: number, durationMS: number }> = new Map();
+
+      if(!!startDate && !!endDate) {
+         // reorder if inverted
+         if (startDate > endDate) {
+            const tempDate = endDate;
+            endDate = startDate;
+            startDate = tempDate;
+         }
+      }
+
+      this.deviceTelemetry.forEach((telemetry, deviceId) => {
+
+         const device = this.devices.get(deviceId);
+         if (!device || device.platformId !== platformId) {
+            return;
+         }
+
+         const pool = this.pools.get(device.poolId)!.name ?? "Unknown Pool"
+
+         telemetry.forEach(t => {
+
+            if (!!t.reservationStart) {
+               if (!!startDate && t.reservationStart < startDate) {
+                  return;
+               }
+
+               if (!!endDate && t.reservationStart > endDate) {
+                  return;
+               }
+            }
+
+            const stepName = t.stepName ?? "Unknown Step";
+            const stream = projectStore.streamById(t.streamId)?.fullname ?? `${t.streamId}`;
+            const key = stepName + stream + pool;
+
+            if (!stepData.has(key)) {
+               stepData.set(key, { stream: stream, pool: pool, stepName: stepName, reservations: 0, problems: 0, durationMS: 0 });
+            }
+
+            const step = stepData.get(key)!;
+
+            if (t.problemTimeUtc) {
+               step.problems++;
+            } else {
+               step.reservations++;
+               if (t.reservationStart && t.reservationFinish) {
+                  step.durationMS += (t.reservationFinish.getTime() - t.reservationStart.getTime());
+               }
+            }
+         });
+      });
+
+      return stepData;
+   }
+
+   getPlatformTelemetry(platformId: string) {
+      if (!this.data) {
+         return [];
+      }
+      const timeSeries: UsagePlatformTelemetry[] = [];
+
+      this.data.forEach(d => {
+         const tdata: UsagePlatformTelemetry = { reservations: 0, date: d.date };
+         Object.keys(d.pools).map(key => d.pools[key]).forEach(p => {
+            let platformData: GetDevicePlatformTelemetryResponse | undefined;
+            platformData = p.find(pd => pd.platformId === platformId);
+            if (!platformData || !platformData.reserved) {
+               return;
+            }
+            tdata.reservations += Object.keys(platformData.reserved).length;
+         });
+         timeSeries.push(tdata);
+      });
+
+      return timeSeries
+   }
+
+   getProblemDevices() {
+
+      const platforms: Map<string, DeviceProblem[]> = new Map();
+
+      this.deviceTelemetry.forEach((telemetry, deviceId) => {
+
+         const device = this.devices.get(deviceId);
+         if (!device) {
+            return;
+         }
+
+         const platform = this.platforms.get(device.platformId);
+         if (!platform) {
+            return;
+         }
+
+         let count = 0;
+         let latest: DeviceTelemetry | undefined;
+         telemetry.filter(t => t.problemTime).forEach(t => {
+            count++;
+            if (!latest || latest.problemTime! < t.problemTime!) {
+               latest = t;
+            }
+
+            if (t.jobName?.indexOf("- Kicked By") !== -1) {
+               t.jobName = t.jobName?.split("- Kicked By")[0];
+            }
+
+         });
+
+         if (!count) {
+            return;
+         }
+
+         if (!platforms.has(platform.id)) {
+            platforms.set(platform.id, []);
+         }
+
+         const poolName = this.pools.get(device.poolId)!.name;
+         const ratePercent = Math.floor(count / telemetry.length * 100);
+         const reservable = device.enabled && !device.maintenanceTime;
+
+         platforms.get(platform.id)!.push({ deviceId: deviceId, deviceName: device.name, poolName: poolName, problems: count, problemsRate: ratePercent, reservations: telemetry.length, problemsDesc: `${count} / ${telemetry.length} (${ratePercent}%)`, latestProblem: latest, reservable: reservable });
+
+      });
+
+      platforms.forEach((devices, key) => {
+         devices = devices.sort((a, b) => b.problems - a.problems);
+         platforms.set(key, devices);
+      });
+
+      return platforms;
+
+   }
+
+   getStepMetrics(jobId: string, stepId: string): StepMetrics | undefined {
+
+      const telemetry: DeviceTelemetry[] = [];
+
+      const devices = new Set<string>();
+
+      this.deviceTelemetry.forEach((t, id) => {
+
+         const stepTelemetry = t.filter(td => { return (td.jobId === jobId) && (td.stepId === stepId) });
+
+         if (stepTelemetry.length) {
+            devices.add(id);
+            telemetry.push(...stepTelemetry);
+         }
+
+      });
+
+      if (!telemetry.length) {
+         return undefined;
+      }
+
+      let minDate: Date | undefined;
+      let maxDate: Date | undefined;
+
+      telemetry.forEach(t => {
+
+         if (t.reservationStart) {
+            if (!minDate || minDate > t.reservationStart) {
+               minDate = t.reservationStart;
+            }
+         }
+
+         if (t.reservationFinish) {
+            if (!maxDate || maxDate < t.reservationFinish) {
+               maxDate = t.reservationFinish;
+            }
+         }
+      });
+
+      let duration = "";
+      if (minDate && maxDate) {
+         duration = msecToElapsed(maxDate.getTime() - minDate.getTime(), true, false);
+      }
+
+      let numProblems = 0;
+
+      telemetry.forEach(t => {
+         if (t.problemTime) {
+            numProblems++;
+         }
+      });
+
+      return {
+         numKits: devices.size,
+         numProblems: numProblems,
+         duration: duration
+      }
+   }
+
+   async set() {
+
+      const requests: any = [];
+
+      // get 2 weeks worth
+      const date = moment().subtract(14, 'days');
+
+      const poolQuery: DevicePoolTelemetryQuery = { minCreateTime: date.toISOString(), count: 65536 };
+      const deviceQuery: DeviceTelemetryQuery = { minCreateTime: date.toISOString(), count: 65536 };
+
+      requests.push(backend.getDevicePlatforms());
+      requests.push(backend.getDevicePools());
+      requests.push(backend.getDevices());
+      requests.push(backend.getDevicePoolTelemetry(poolQuery));
+      requests.push(backend.getDeviceTelemetry(deviceQuery));
+
+      const responses = await Promise.all(requests);
+
+      const platforms = responses[0] as GetDevicePlatformResponse[];
+      const pools = responses[1] as GetDevicePoolResponse[];
+      const devices = responses[2] as GetDeviceResponse[];
+      const telemetry = responses[3] as GetDevicePoolTelemetryResponse[];
+      const deviceTelemetry = responses[4] as GetDeviceTelemetryResponse[];
+
+      deviceTelemetry.forEach(d => {
+         const telemetry: DeviceTelemetry[] = [];
+         d.telemetry.forEach(t => {
+
+            const t2: DeviceTelemetry = { ...t, createTime: new Date(t.createTimeUtc) };
+
+            if (t.reservationStartUtc) {
+               t2.reservationStart = new Date(t.reservationStartUtc);
+            }
+
+            if (t.problemTimeUtc) {
+               t2.problemTime = new Date(t.problemTimeUtc);
+            }
+
+            if (t.reservationFinishUtc) {
+               t2.reservationFinish = new Date(t.reservationFinishUtc);
+            }
+
+            telemetry.push(t2);
+         });
+
+         this.deviceTelemetry.set(d.deviceId, telemetry);
+      });
+
+      for (const platform of platforms) {
+         this.platforms.set(platform.id, platform);
+      }
+
+      for (const pool of pools) {
+         if (pool.poolType !== DevicePoolType.Automation) {
+            continue;
+         }
+         this.pools.set(pool.id, pool);
+      }
+
+      for (const device of devices) {
+         this.devices.set(device.id, device);
+      }
+
+      // patch up telemetry data 
+      telemetry.forEach(t => {
+         for (const poolId in t.telemetry) {
+            const values = t.telemetry[poolId];
+            values.forEach(v => {
+
+               if (!v.reserved) {
+                  return;
+               }
+
+               for (const rkey in v.reserved) {
+                  const rd = v.reserved[rkey];
+                  rd.forEach(rdv => {
+
+                     if (!rdv.jobName) {
+                        rdv.jobName = `Unknown Job in stream ${rkey}, ${rdv.jobId}`;
+                     }
+                     if (!rdv.stepName) {
+                        rdv.stepName = `Unknown Step in stream ${rkey}`;
+                     }
+
+                     if (rdv.jobName.indexOf("- Kicked By") !== -1) {
+                        rdv.jobName = rdv.jobName.split("- Kicked By")[0];
+                     }
+                  })
+               }
+            });
+         };
+      });
+
+      this.data = telemetry.map(t => {
+         const date = new Date(t.createTimeUtc);
+         return { date: date, dateString: date.toDateString(), pools: t.telemetry }
+      }).sort((a, b) => a.date.getTime() - b.date.getTime());
+
+      this.setUpdated();
+
+   }
+
+   @action
+   setUpdated() {
+      this.updated++;
+   }
+
+   getData(poolId: string, platformId: string): TelemetryData | undefined {
+
+      if (!this.data) {
+         return undefined;
+      }
+
+      const pool = this.pools.get(poolId);
+      const platform = this.platforms.get(platformId);
+
+      if (!pool || !platform) {
+         return undefined;
+      }
+
+      const telemetry: PlatformTelemetry[] = [];
+      const streamIds = new Set<string>();
+      const deviceIds = new Set<string>();
+
+      const streamNames: Record<string, string> = {};
+      const streamNamesReverse: Record<string, string> = {};
+      const streamSelectors: Record<string, string> = {};
+      const streamSelectorsReverse: Record<string, string> = {};
+
+      this.data.forEach(d => {
+
+         const poolData = d.pools[poolId];
+
+         let platformData: GetDevicePlatformTelemetryResponse | undefined;
+
+         if (poolData) {
+            platformData = poolData.find(pd => pd.platformId === platformId);
+         }
+
+         if (!platformData) {
+            return;
+         }
+
+         const tdata = { ...platformData, date: d.date, dateString: d.dateString } as any;
+
+         tdata["Available"] = platformData.available?.length ?? 0;
+         tdata["Disabled"] = platformData.disabled?.length ?? 0;
+         tdata["Maintenance"] = platformData.maintenance?.length ?? 0;
+         tdata["Problem"] = platformData.problem?.length ?? 0;
+
+         platformData.available?.forEach(d => { deviceIds.add(d); });
+         platformData.maintenance?.forEach(d => { deviceIds.add(d); });
+         platformData.disabled?.forEach(d => { deviceIds.add(d); });
+         platformData.problem?.forEach(d => { deviceIds.add(d); });
+
+         if (platformData.reserved) {
+            for (const streamId in platformData.reserved) {
+               if (!streamIds.has(streamId)) {
+
+                  streamIds.add(streamId);
+
+                  streamNames[streamId] = streamId;
+                  const stream = projectStore.streamById(streamId);
+                  if (stream) {
+                     streamNames[streamId] = stream.fullname ?? streamId;
+                  }
+
+                  const name = streamNames[streamId];
+                  streamNamesReverse[name] = streamId;
+
+                  streamSelectors[name] = name.replaceAll(".", "_").replaceAll("/", "_");
+                  streamSelectorsReverse[streamSelectors[name]] = name;
+               }
+
+               // make stream data accessible by id and fullname
+               tdata[streamId] = platformData.reserved[streamId].length;
+               tdata[streamNames[streamId]] = platformData.reserved[streamId].length;
+
+               platformData.reserved[streamId].forEach(d => { deviceIds.add(d.deviceId); })
+            }
+         }
+
+         telemetry.push(tdata);
+      });
+
+      const deviceIdArray = Array.from(deviceIds).sort((a, b) => a.localeCompare(b));
+      const deviceResponses = deviceIdArray.map(id => this.devices.get(id)).filter(d => !!d);
+      const devices: Record<string, GetDeviceResponse> = {};
+
+      deviceResponses.forEach(d => devices[d!.id] = d!);
+
+      const streamIdArray = Array.from(streamIds).sort((a, b) => a.localeCompare(b))
+
+      for (const streamId of streamIdArray) {
+         telemetry.forEach((t: any) => {
+            if (!t[streamId]) {
+               t[streamId] = 0;
+               t[streamNames[streamId]] = 0;
+            }
+         });
+      }
+
+      return {
+         pool: pool,
+         platform: platform,
+         streamIds: streamIdArray,
+         streamNames: streamNames,
+         streamNamesReverse: streamNamesReverse,
+         streamSelectors: streamSelectors,
+         streamSelectorsReverse: streamSelectorsReverse,
+         telemetry: telemetry,
+         deviceIds: deviceIdArray,
+         devices: devices
+      }
+   }
+
+   get loaded(): boolean { return this.data ? true : false; }
+
+   @observable
+   updated: number = 0;
+
+   data?: PoolTelemetryData[];
+
+   platforms = new Map<string, GetDevicePlatformResponse>();
+   pools = new Map<string, GetDevicePoolResponse>();
+   devices = new Map<string, GetDeviceResponse>();
+
+   deviceTelemetry = new Map<string, DeviceTelemetry[]>();
+}
diff --git a/Engine/Source/Programs/Horde/HordeDashboard/src/base/utilities/timeUtils.ts b/Engine/Source/Programs/Horde/HordeDashboard/src/base/utilities/timeUtils.ts
--- a/Engine/Source/Programs/Horde/HordeDashboard/src/base/utilities/timeUtils.ts
+++ b/Engine/Source/Programs/Horde/HordeDashboard/src/base/utilities/timeUtils.ts
@@ -22,12 +22,26 @@ export const msecToElapsed = (millisec: number, includeMinutes: boolean = true,
     let duration = "";
     const d = moment.duration(millisec);
 
+    let aggregatedHours = 0;
+
+    if (d.years()) {
+        aggregatedHours += d.years() * 365 * 24;
+    }
+
+    if (d.months()) {
+        aggregatedHours += d.months() * 31 * 24;
+    }
+
     if (d.days()) {
-        duration += `${d.days()}d `;
+        aggregatedHours += d.days() * 24;
     }
 
     if (d.hours()) {
-        duration += `${d.hours()}h `;
+        aggregatedHours += d.hours();
+    }
+
+    if (aggregatedHours) {
+        duration += `${aggregatedHours}h `;
     }
 
     if (includeMinutes) {
diff --git a/Engine/Source/Programs/Horde/HordeDashboard/src/components/DeviceEditor.tsx b/Engine/Source/Programs/Horde/HordeDashboard/src/components/DeviceEditor.tsx
deleted file mode 100644
diff --git a/Engine/Source/Programs/Horde/HordeDashboard/src/components/DeviceInfoView.tsx b/Engine/Source/Programs/Horde/HordeDashboard/src/components/DeviceInfoView.tsx
deleted file mode 100644
diff --git a/Engine/Source/Programs/Horde/HordeDashboard/src/components/DevicePoolTelemetry.tsx b/Engine/Source/Programs/Horde/HordeDashboard/src/components/DevicePoolTelemetry.tsx
deleted file mode 100644
diff --git a/Engine/Source/Programs/Horde/HordeDashboard/src/components/DeviceView.tsx b/Engine/Source/Programs/Horde/HordeDashboard/src/components/DeviceView.tsx
deleted file mode 100644
diff --git a/Engine/Source/Programs/Horde/HordeDashboard/src/components/devices/DeviceEditor.tsx b/Engine/Source/Programs/Horde/HordeDashboard/src/components/devices/DeviceEditor.tsx
new file mode 100644
--- /dev/null
+++ b/Engine/Source/Programs/Horde/HordeDashboard/src/components/devices/DeviceEditor.tsx
@@ -0,0 +1,512 @@
+import { SpinnerSize, DialogType, IDropdownOption, MessageBarType, Text, Checkbox, DefaultButton, Dialog, DialogFooter, Dropdown, MessageBar, Modal, PrimaryButton, Spinner, Stack, TextField } from "@fluentui/react";
+import moment from "moment";
+import React, { useState } from "react";
+import backend from "../../backend";
+import { GetDevicePlatformResponse, GetDevicePoolResponse, GetDeviceReservationResponse, GetDeviceResponse, GetUserResponse } from "../../backend/Api";
+import dashboard from "../../backend/Dashboard";
+import { PollBase } from "../../backend/PollBase";
+import { getHordeStyling } from "../../styles/Styles";
+
+type DeviceEditData = {
+   // if defined, existing device
+   id?: string;
+   platformId: string;
+   poolId: string;
+   modelId?: string;
+   name: string;
+   address?: string;
+   enabled: boolean;
+   maintenance: boolean;
+   notes?: string;
+};
+
+// in sort order
+export enum DeviceStatus {
+   Available = 0,
+   Reserved = 1,
+   Maintenance = 2,
+   Problem = 3,
+   Disabled = 4
+}
+
+export class DeviceHandler extends PollBase {
+
+   constructor(pollTime = 5000) {
+      super(pollTime);
+   }
+
+   clear() {
+      super.stop();
+      this.loaded = false;
+      this.devices = [];
+      this.reservations = [];
+      this.pools.clear();
+      this.platforms.clear();
+   }
+
+   async poll(): Promise<void> {
+
+      try {
+
+         this.devices = await backend.getDevices();
+         this.reservations = await backend.getDeviceReservations();
+
+         // filter out utilization to entries which contain job id (legacy reservations don't)
+         this.devices.forEach(d => d.utilization = d.utilization?.filter(u => !!u.jobId));
+
+         const userIds = new Set<string>();
+         this.devices.forEach(d => {
+            if (d.checkedOutByUserId && !this.users.has(d.checkedOutByUserId)) {
+               userIds.add(d.checkedOutByUserId);
+            }
+         });
+
+         if (!this.users.get(dashboard.userId)) {
+            userIds.add(dashboard.userId);
+         }
+
+         if (userIds.size) {
+            const users = await backend.getUsers({ ids: Array.from(userIds) });
+            users.forEach(u => this.users.set(u.id, u));
+         }
+
+         if (!this.devices.length || this.devices.find(d => !this.platforms.has(d.platformId))) {
+            // we need to get platforms
+            const platforms = await backend.getDevicePlatforms();
+            this.platforms.clear();
+            platforms.forEach(p => {
+               this.platforms.set(p.id, p);
+            });
+
+            const pools = await backend.getDevicePools();
+            this.pools.clear();
+            pools.forEach(p => {
+               this.pools.set(p.id, p);
+            });
+
+         }
+         
+         this.loaded = true;         
+         this.setUpdated();
+
+      } catch (err) {
+
+      }
+
+   }
+
+   getUserDeviceCheckouts(userId: string): GetDeviceResponse[] {
+      return this.devices.filter(d => d.checkedOutByUserId === userId);
+   }
+
+   getReservation(device: GetDeviceResponse): GetDeviceReservationResponse | undefined {
+      return this.reservations.find(r => !!r.devices.find(d => d === device.id));
+   }
+
+   getDeviceWriteAccess(device: GetDeviceResponse): boolean {
+
+      const pool = this.pools.get(device.poolId);
+
+      if (!pool) {
+         return false;
+      }
+
+      return pool.writeAccess;
+
+   }
+
+   getDeviceStatus(device: GetDeviceResponse): DeviceStatus {
+
+      if (this.getReservation(device)) {
+         return DeviceStatus.Reserved;
+      }
+
+      if (!device.enabled) {
+         return DeviceStatus.Disabled;
+      }
+
+      if (device.maintenanceTime) {
+         return DeviceStatus.Maintenance;
+      }
+
+      if (device.problemTime) {
+
+         const end = moment.utc();
+         const d = moment.duration(end.diff(moment(device.problemTime)));
+
+         // note this must match reservation selection for problem time in backend
+         if (d.asMinutes() < dashboard.deviceProblemCooldownMinutes) {
+            return DeviceStatus.Problem;
+         }
+      }
+
+      return DeviceStatus.Available;
+   }
+
+   getDevices(platformId?: string) {
+      return this.devices.filter(d => !platformId || d.platformId === platformId);
+   }
+
+   loaded = false;
+
+   private devices: GetDeviceResponse[] = [];
+
+   platforms: Map<string, GetDevicePlatformResponse> = new Map();
+
+   pools: Map<string, GetDevicePoolResponse> = new Map();
+
+   users: Map<string, GetUserResponse> = new Map();
+
+   private reservations: GetDeviceReservationResponse[] = [];
+}
+
+
+export const DeviceEditor: React.FC<{ handler: DeviceHandler, deviceIn?: GetDeviceResponse | undefined, editNote?: boolean, onClose: () => void }> = ({ handler, deviceIn, editNote, onClose }) => {
+
+   const [state, setState] = useState<{ device?: DeviceEditData, title?: string }>({});
+   const [error, setError] = useState<string | undefined>();
+   const [submitting, setSubmitting] = useState(false);
+   const [confirmDelete, setConfirmDelete] = useState(false);
+   
+   const { hordeClasses } = getHordeStyling();
+
+   if (submitting) {
+
+      return <Modal isOpen={true} isBlocking={true} styles={{ main: { padding: 8, width: 400 } }} >
+         <Stack style={{ paddingTop: 32 }}>
+            <Stack tokens={{ childrenGap: 24 }} styles={{ root: { padding: 8 } }}>
+               <Stack horizontalAlign="center">
+                  <Text variant="large">Please wait...</Text>
+               </Stack>
+               <Stack verticalAlign="center" style={{ paddingBottom: 32 }}>
+                  <Spinner size={SpinnerSize.large} />
+               </Stack>
+            </Stack>
+         </Stack>
+      </Modal>
+
+   }
+
+   const platforms = handler.platforms;
+   const pools = handler.pools;
+
+   const device = state.device!;
+   const existing = !!deviceIn;
+
+   let deviceStatus = DeviceStatus.Available;
+
+   if (deviceIn) {
+      deviceStatus = handler.getDeviceStatus(deviceIn);
+   }
+
+   const reservation = deviceStatus === DeviceStatus.Reserved;
+
+   const close = async () => {
+      // @todo: state should be squashed
+
+      await handler.forceUpdate();
+
+      setState({});
+      setSubmitting(false);
+      setError("");
+      setConfirmDelete(false);
+      onClose();
+   };
+
+   const onDelete = async () => {
+
+      try {
+
+         setSubmitting(true);
+         await backend.deleteDevice(device.id!);
+         await close();
+
+      } catch (reason: any) {
+         setError(reason.toString());
+         setSubmitting(false);
+      }
+
+   }
+
+
+   if (confirmDelete) {
+      return <Dialog
+         hidden={false}
+         onDismiss={() => setConfirmDelete(false)}
+         minWidth={400}
+         dialogContentProps={{
+            type: DialogType.normal,
+            title: `Delete device ${state.device?.name} ?`,
+         }}
+         modalProps={{ isBlocking: true }} >
+         <DialogFooter>
+            <PrimaryButton onClick={() => { setConfirmDelete(false); onDelete() }} text="Delete" />
+            <DefaultButton onClick={() => setConfirmDelete(false)} text="Cancel" />
+         </DialogFooter>
+      </Dialog>
+   }
+
+   const platformOptions = Array.from(platforms.values()).sort((a, b) => {
+
+      return a.name.localeCompare(b.name);
+
+   }).map(p => {
+      return { key: `platform_${p.id}`, text: p.name, selected: device?.platformId === p.id, data: p.id } as IDropdownOption
+   });
+
+   let modelOptions: IDropdownOption[] = [];
+
+   const platform = platforms.get(device?.platformId);
+
+   if (platform) {
+      modelOptions = platform.modelIds.sort((a, b) => {
+         return a.localeCompare(b);
+      }).map(model => {
+         return { key: `model_${model}`, text: model, selected: device?.modelId === model, data: model } as IDropdownOption
+      });
+   }
+
+   modelOptions.unshift({ key: `model_base`, text: "Base", selected: device?.modelId === undefined, data: undefined });
+
+   let defaultModelKey = "model_base";
+   if (device?.modelId) {
+      defaultModelKey = `model_${device.modelId}`;
+   }
+
+
+   const poolOptions: IDropdownOption[] = [];
+   pools.forEach(p => {
+      poolOptions.push({ key: `pool_${p.id}`, text: p.name, selected: device?.poolId === p.id, data: p.id });
+   });
+
+   if (!state.device) {
+
+      if (deviceIn) {
+         setState({
+            device: {
+               id: deviceIn?.id,
+               platformId: deviceIn.platformId,
+               enabled: deviceIn.enabled,
+               address: deviceIn.address,
+               name: deviceIn.name,
+               poolId: deviceIn.poolId,
+               modelId: deviceIn.modelId,
+               notes: deviceIn.notes,
+               maintenance: !!deviceIn.maintenanceTime
+            }, title: `Edit ${deviceIn.name}`
+         });
+      } else {
+         setState({
+            device: {
+               enabled: true,
+               maintenance: false,
+               address: "",
+               name: "",
+               platformId: platformOptions.length > 0 ? platformOptions[0].data : "",
+               poolId: poolOptions.length > 0 ? poolOptions[0].data : ""
+            },
+            title: "Add New Device"
+         });
+      }
+
+      return null;
+   }
+
+
+   const statusOptions: IDropdownOption[] = [];
+
+   if (reservation) {
+      statusOptions.push({ key: 'status_reserved', text: "Reserved", selected: true, data: "Reserved" });
+   } else {
+
+      statusOptions.push({ key: 'status_available', text: "Available", selected: device.enabled && !device.maintenance, data: "Available" });
+      statusOptions.push({ key: 'status_disabled', text: "Disabled", selected: !device.enabled, data: "Disabled" });
+      statusOptions.push({ key: 'status_maintenance', text: "Maintenance", selected: device.enabled && device.maintenance, data: "Maintenance" });
+   };
+
+   // need to handle delete reservation case
+   const onSave = async () => {
+
+      if (!device.name) {
+         setError("Device must have a name");
+         return;
+      }
+
+      if (!device.address) {
+         setError("Device must have an address");
+         return;
+      }
+
+      setSubmitting(true);
+
+      try {
+
+         if (reservation) {
+
+            if (editNote) {
+
+               await backend.modifyDevice(device.id!, {
+                  notes: device.notes ? device.notes.trim() : ""
+               }).then(() => {
+                  close();
+               }).catch((reason) => {
+                  setError(`Problem modifying device, ${reason}`);
+               });
+
+            } else {
+
+               await backend.modifyDevice(device.id!, {
+                  problem: false,
+                  maintenance: device.maintenance,
+                  enabled: true
+               }).then(() => {
+                  close();
+               }).catch((reason) => {
+                  setError(`Problem modifying device, ${reason}`);
+               });
+            }
+
+         }
+         else if (!existing) {
+            const result = await backend.addDevice({
+               name: device.name.trim(),
+               address: device.address,
+               poolId: device.poolId,
+               platformId: device.platformId,
+               modelId: device.modelId
+            });
+
+            if (!result) {
+               setError("Problem adding device");
+            } else {
+               close();
+            }
+         } else {
+
+            await backend.modifyDevice(device.id!, {
+               name: device.name.trim(),
+               address: device.address,
+               poolId: device.poolId,
+               modelId: device.modelId ?? "Base",
+               notes: device.notes,
+               problem: false, // always clear automatically generated problem state when saving
+               maintenance: device.maintenance,
+               enabled: device.enabled
+            }).then(() => {
+               close();
+            }).catch((reason) => {
+               setError(`Problem modifying device,: ${reason}`);
+               setSubmitting(false);
+            })
+         }
+         
+      } catch (reason: any) {
+         setError(reason.toString());
+         setSubmitting(false);
+      }
+
+   }
+
+   return <Modal className={hordeClasses.modal} isOpen={true} isBlocking={true} styles={{ main: { padding: 8, width: 700 } }}>
+      {!!error && <MessageBar
+         messageBarType={MessageBarType.error}
+         isMultiline={false}
+         onDismiss={() => setError("")}
+      >
+         <Text>{error}</Text>
+      </MessageBar>}
+
+      <Stack style={{ padding: 8 }}>
+         <Stack style={{ paddingBottom: 16 }}>
+            <Text variant="mediumPlus" style={{ fontFamily: "Horde Open Sans SemiBold" }}>{state.title}</Text>
+         </Stack>
+         <Stack tokens={{ childrenGap: 8 }} style={{ padding: 8 }}>
+            <TextField label="Name" disabled={existing} defaultValue={device.name} onChange={(ev, value) => {
+               device.name = value ?? "";
+            }} />
+            <Dropdown label="Platform" disabled={existing} options={platformOptions} onChange={(ev, option) => {
+
+               if (option) {
+                  device.platformId = option.data as string;
+                  device.modelId = undefined;
+
+                  setState({ device: device, title: state.title });
+               }
+
+            }} />
+
+
+            {(!editNote && !reservation) && <Dropdown label="Pool" options={poolOptions} onChange={(ev, option) => {
+
+               if (option) {
+                  device.poolId = option.data as string;
+               }
+
+            }} />}
+
+            {(!editNote && !reservation) && <Dropdown label="Model" defaultSelectedKey={defaultModelKey} options={modelOptions} onChange={(ev, option) => {
+
+               if (option?.data === undefined) {
+                  device.modelId = undefined; // "Base"
+               }
+               else if (option) {
+                  device.modelId = option.data as string;
+               }
+
+            }} />}
+
+
+            <TextField label="Address" defaultValue={device.address} disabled={(editNote || reservation)} onChange={(ev, value) => {
+               device.address = value ?? "";
+            }} />
+
+            {(!editNote && !reservation) && <Dropdown label="Status" options={statusOptions} onChange={(ev, option) => {
+
+               if (option) {
+
+                  if (option.data === "Available") {
+                     device.maintenance = false;
+                     device.enabled = true;
+                  }
+
+                  if (option.data === "Maintenance") {
+                     device.maintenance = true;
+                     device.enabled = true;
+                  }
+
+                  if (option.data === "Disabled") {
+                     device.maintenance = false;
+                     device.enabled = false;
+                  }
+
+               }
+
+            }} />}
+
+            <TextField label="Notes" defaultValue={device.notes} multiline rows={5} resizable={false} onChange={(ev, newValue) => {
+               device.notes = newValue;
+            }} />
+
+            {!!reservation && !editNote && <Stack style={{ paddingTop: 24 }}><Checkbox label="Disable device for maintenance once reservation finishes"
+               checked={!!device.maintenance}
+               onChange={(ev, checked) => {
+                  device.maintenance = checked ? true : false;
+                  setState({ device: device, title: state.title });
+               }} />
+            </Stack>}
+
+
+         </Stack>
+
+         <Stack horizontal style={{ paddingTop: 64 }}>
+            {!!existing && !editNote && !reservation && <PrimaryButton style={{ backgroundColor: "red", border: 0 }} onClick={() => setConfirmDelete(true)} text="Delete Device" />}
+            <Stack grow />
+            <Stack horizontal tokens={{ childrenGap: 28 }}>
+               <PrimaryButton onClick={() => onSave()} text="Save" />
+               <DefaultButton onClick={() => close()} text="Cancel" />
+            </Stack>
+         </Stack>
+      </Stack>
+   </Modal>
+
+};
\ No newline at end of file
diff --git a/Engine/Source/Programs/Horde/HordeDashboard/src/components/devices/DeviceHealthWidgets.tsx b/Engine/Source/Programs/Horde/HordeDashboard/src/components/devices/DeviceHealthWidgets.tsx
new file mode 100644
--- /dev/null
+++ b/Engine/Source/Programs/Horde/HordeDashboard/src/components/devices/DeviceHealthWidgets.tsx
@@ -0,0 +1,24 @@
+import { Stack, Text } from "@fluentui/react";
+import dashboard, { StatusColor } from "../../backend/Dashboard";
+
+export const deviceHealthRateBar = (rate: number, width: number, height: number, text: string): JSX.Element => {
+   const colors = dashboard.getStatusColors();
+   return (
+      <Stack verticalFill={true} verticalAlign="center" horizontal>
+         <div style={{ width: width, height: height, verticalAlign: 'middle', marginRight: '3px' }}>
+            <svg width={width} height={height}>
+               <defs>
+                  <clipPath id="cutout">
+                     <rect x={0} y={0} width={width} height={height} rx={4} ry={4} />
+                  </clipPath>
+               </defs>
+               <g clipPath="url(#cutout)">
+                  <rect x={0} y={0} width={width} height={height} fill={colors.get(StatusColor.Success)} />
+                  <rect x={0} y={0} width={Math.ceil(width * rate / 100)} height={height} fill={colors.get(StatusColor.Failure)} />
+               </g>
+            </svg>
+         </div>
+         <Text>{text}</Text>
+      </Stack>
+   );
+}
diff --git a/Engine/Source/Programs/Horde/HordeDashboard/src/components/devices/DeviceInfoView.tsx b/Engine/Source/Programs/Horde/HordeDashboard/src/components/devices/DeviceInfoView.tsx
new file mode 100644
--- /dev/null
+++ b/Engine/Source/Programs/Horde/HordeDashboard/src/components/devices/DeviceInfoView.tsx
@@ -0,0 +1,327 @@
+import { DefaultButton, DetailsList, DetailsRow, IColumn, IconButton, IDetailsListProps, Label, Modal, PrimaryButton, SelectionMode, Spinner, SpinnerSize, Stack, Text, TextField } from "@fluentui/react";
+import { observer } from "mobx-react-lite";
+import moment from "moment";
+import { useState } from "react";
+import { Link, useNavigate } from "react-router-dom";
+import backend from "../../backend";
+import { DeviceStatus, DeviceTelemetryQuery, GetDeviceResponse, GetDeviceTelemetryResponse, GetDeviceUtilizationResponse, GetTelemetryInfoResponse, JobData } from "../../backend/Api";
+import dashboard, { StatusColor } from "../../backend/Dashboard";
+import { projectStore } from "../../backend/ProjectStore";
+import { displayTimeZone } from "../../base/utilities/timeUtils";
+import { getHordeStyling } from "../../styles/Styles";
+import { DeviceHandler } from "./DeviceEditor";
+import { deviceHealthRateBar } from "./DeviceHealthWidgets"
+import { DeviceStatusIcon } from "../StatusIcon";
+
+const streamIdToFullname = new Map<string, string>();
+
+export const DeviceInfoModal: React.FC<{ handler: DeviceHandler, deviceIn?: GetDeviceResponse | undefined, onEdit: (device: GetDeviceResponse) => void, onClose: () => void }> = observer(({ handler, deviceIn, onEdit, onClose }) => {
+
+   type JobItem = {
+      job: JobData,
+      utilization: GetDeviceUtilizationResponse
+   };
+
+   type JobTelemetryItem = {
+      job: GetTelemetryInfoResponse,
+   }
+
+   const [jobState, setJobState] = useState<{ items: JobItem[], telemetryItems: JobTelemetryItem[], querying: boolean, queried: boolean, modifiedByUser?: string }>({ items: [], telemetryItems: [], querying: true, queried: false });
+   const navigate = useNavigate();
+
+   if (!deviceIn) {
+      return null;
+   }
+
+   const { hordeClasses } = getHordeStyling();
+
+   // get unique ids
+   let jobIds = deviceIn.utilization?.map(u => u.jobId).filter(jobId => !!jobId) as string[];
+   if (jobIds) {
+      jobIds = Array.from(new Set(jobIds)).slice(0, 100);
+   }
+
+   const queryJobs = async () => {
+
+      if (!deviceIn.modifiedByUser && !jobIds?.length) {
+         setJobState({ ...jobState, queried: true, querying: false });
+         return;
+      }
+
+      let modifiedByUser = "";
+
+      if (deviceIn.modifiedByUser) {
+         const user = await backend.getUsers({ ids: [deviceIn.modifiedByUser!] });
+         if (user.length && user[0].name) {
+            modifiedByUser = user[0].name;
+         }
+
+      }
+
+      let deviceTelemetry: GetDeviceTelemetryResponse[] = [];
+
+      try {
+
+         const maxCount = 512;
+         const telemetryQuery: DeviceTelemetryQuery = {
+            id: [deviceIn.id],
+            count: maxCount,
+            minCreateTime: new Date(new Date().getTime() - 86400000 * 14).toISOString()
+         };
+
+         deviceTelemetry = await backend.getDeviceTelemetry(telemetryQuery);
+
+      } catch (reason) {
+
+      } finally {
+
+         const streamIds = new Set<string>();
+
+         if (deviceTelemetry.length) {
+            deviceTelemetry[0].telemetry.forEach(t => {
+               if (t.streamId) {
+                  const stream = projectStore.streamById(t.streamId);
+                  if (!stream || !stream.fullname) {
+                     streamIds.add(t.streamId);
+                     streamIdToFullname.set(t.streamId, t.streamId);
+                     return;
+                  };
+                  streamIds.add(t.streamId);
+                  streamIdToFullname.set(t.streamId, stream.fullname);
+               }
+            });
+         }
+
+         const items: JobItem[] = [];
+         const telemetryItems: JobTelemetryItem[] = [];
+
+         deviceIn.utilization?.forEach(u => {
+            const jobsFromTelemetry = deviceTelemetry[0]?.telemetry;
+            const telemetryJob = jobsFromTelemetry?.find(j => j.jobId === u.jobId && j.stepId === u.stepId);
+            if (!telemetryJob) { return; }
+            telemetryItems.push({ job: telemetryJob });
+         });
+
+         setJobState({ ...jobState, items: items, telemetryItems: telemetryItems, queried: true, querying: false, modifiedByUser: modifiedByUser });
+      }
+
+   }
+
+
+   if (!jobState.queried) {
+      setJobState({ ...jobState, queried: true, querying: true });
+      queryJobs();
+      return null;
+   }
+
+   if (jobState.querying) {
+
+      return <Modal isOpen={true} isBlocking={true} topOffsetFixed={true} styles={{ main: { padding: 8, width: 700, hasBeenOpened: false, top: "80px", position: "absolute" } }} className={hordeClasses.modal}>
+         <Stack tokens={{ childrenGap: 40 }} styles={{ root: { padding: 8 } }}>
+            <Stack grow verticalAlign="center">
+               <Text variant="mediumPlus" styles={{ root: { fontWeight: "unset", fontFamily: "Horde Open Sans SemiBold" } }}>Loading Device {deviceIn.name}</Text>
+            </Stack>
+            <Stack verticalAlign="center">
+               <Spinner size={SpinnerSize.large} />
+            </Stack>
+
+         </Stack>
+      </Modal>
+   }
+
+
+   // results
+   const columns: IColumn[] = [
+      //{ key: 'column1', name: 'Change', data: 'string', isRowHeader: true, minWidth: 80, maxWidth: 80, isResizable: false },
+      { key: 'column1', name: 'Device Health', isRowHeader: true, minWidth: 15, maxWidth: 15, isResizable: false },
+      { key: 'column2', name: 'Name', data: 'string', isRowHeader: true, minWidth: 530, maxWidth: 530, isResizable: false },
+      { key: 'column3', name: 'Created', data: 'date', isRowHeader: true, minWidth: 120, maxWidth: 120, isResizable: false },
+      //{ key: 'column3', name: 'StartedBy', minWidth: 140, maxWidth: 140, isResizable: false },
+
+   ];
+
+   const renderRow: IDetailsListProps['onRenderRow'] = (props) => {
+
+      if (props) {
+
+         const item = props!.item as JobTelemetryItem;
+
+         let url = `/job/${item.job.jobId}`;
+
+         if (item.job?.stepId) {
+            url += `?step=${item.job.stepId}`
+         }
+
+         const commonSelectors = { ".ms-DetailsRow-cell": { "overflow": "visible", padding: 0 } };
+
+         props.styles = { ...props.styles, root: { selectors: { ...commonSelectors as any } } };
+
+         return <Link to={url} onClick={(ev) => { if (!ev.ctrlKey) onClose() }}><div className="job-item"><DetailsRow {...props} /> </div></Link>;
+
+      }
+      return null;
+   };
+
+
+
+   const renderItem = (item: JobTelemetryItem, index?: number, column?: IColumn) => {
+
+      if (!column) {
+         return <div />;
+      }
+
+      if (column.name === "Device Health") {
+         let deviceStatus = item.job.problemTimeUtc ? DeviceStatus.Error : DeviceStatus.Normal;
+         return <Stack horizontal verticalFill={true} verticalAlign="center" tokens={{ childrenGap: 0, padding: 0 }} style={{ overflow: "hidden" }} ><DeviceStatusIcon status={deviceStatus} /></Stack>
+      }
+
+      if (column.name === "Name") {
+
+         let name = `${streamIdToFullname.get(item.job.streamId ?? "")!}`;
+         if (item.job.stepName) {
+            name += ` - ${item.job.stepName}`;
+         }
+
+         return <Stack horizontal verticalFill={true} verticalAlign="center" tokens={{ childrenGap: 0, padding: 0 }} style={{ overflow: "hidden" }} ><Text variant="small">{name}</Text></Stack>;
+      }
+
+      /*
+      if (column.name === "Change") {
+          
+          const batch = item.job.batches?.find(b => b.steps.find(s => s.id === item.utilization.stepId));
+          const step = batch?.steps.find(s => s.id === item.utilization.stepId);
+
+          if (step) {
+              return <Stack horizontal verticalFill={true} verticalAlign="center" tokens={{ childrenGap: 0, padding: 0 }} style={{ paddingTop: 4 }} ><StepStatusIcon step={step} style={{ fontSize: 12, paddingBottom: 4 }} /><ChangeButton job={item.job} /></Stack>;
+          }
+
+          
+
+          return <Stack horizontal verticalFill={true} verticalAlign="center" tokens={{ childrenGap: 0, padding: 0 }} style={{ paddingTop: 4 }} ><ChangeButton job={item.job} /></Stack>;
+      }
+      */
+
+      /*
+      if (column.name === "StartedBy") {
+          let startedBy = item.job.startedByUserInfo;
+          if (!startedBy) {
+              startedBy = "Scheduler";
+          }
+          return <Stack verticalAlign="center" verticalFill={true} horizontalAlign={"center"}>{startedBy}</Stack>;
+      }
+      */
+
+
+      if (column.name === "Created") {
+
+         if (item.job.createTimeUtc) {
+
+            const displayTime = moment(item.job.reservationStartUtc).tz(displayTimeZone());
+            const format = dashboard.display24HourClock ? "HH:mm:ss z" : "LT z";
+
+            let displayTimeStr = displayTime.format('MMM Do') + ` at ${displayTime.format(format)}`;
+
+
+            return <Stack verticalAlign="center" horizontalAlign="end" tokens={{ childrenGap: 0, padding: 0 }} style={{ height: "100%", paddingRight: 18 }}>
+               <Text variant="small">{displayTimeStr}</Text>
+            </Stack>;
+
+         } else {
+            return "???";
+         }
+      }
+
+      return <Stack />;
+   }
+
+   const failureRate = jobState.telemetryItems.length == 0? 0 : Math.floor(jobState.telemetryItems.filter(item => !!item.job.problemTimeUtc).length / jobState.telemetryItems.length * 100);
+
+   return (<Modal isOpen={true} styles={{ main: { padding: 8, width: 1140, height: '800px', backgroundColor: '#FFFFFF' } }} className={hordeClasses.modal} onDismiss={() => { onClose() }}>
+      <Stack styles={{ root: { paddingTop: 8, paddingLeft: 24, paddingRight: 12, paddingBottom: 8 } }}>
+         <Stack tokens={{ childrenGap: 12 }}>
+            <Stack horizontal styles={{ root: { padding: 0 } }}>
+               <Stack style={{ paddingLeft: 0, paddingTop: 4 }} grow>
+                  <Text variant="mediumPlus" styles={{ root: { fontFamily: "Horde Open Sans SemiBold" } }}>{deviceIn.name}</Text>
+               </Stack>
+               <Stack grow horizontalAlign="end">
+                  <IconButton
+                     iconProps={{ iconName: 'Cancel' }}
+                     onClick={() => { onClose(); }}
+                  />
+               </Stack>
+            </Stack>
+
+            <Stack horizontal tokens={{ childrenGap: 48 }} >
+
+               <Stack>
+                  <Stack tokens={{ childrenGap: 12 }}>
+
+                     <Stack>
+                        <TextField label="Pool" value={deviceIn.poolId?.toUpperCase()} disabled={true} />
+                     </Stack>
+
+                     <Stack>
+                        <TextField label="Address" value={deviceIn.address} disabled={true} />
+                     </Stack>
+
+                     <Stack>
+                        <TextField label="Model" value={deviceIn.modelId ? deviceIn.modelId : "Base"} disabled={true} />
+                     </Stack>
+
+
+                     <Stack>
+                        <TextField label="Last Modified By" value={jobState.modifiedByUser ? jobState.modifiedByUser : ""} disabled={true} />
+                     </Stack>
+
+                     <Stack horizontal style={{ paddingTop: 24 }}>
+                        <Stack style={{ paddingTop: 8, paddingRight: 24 }}>
+                           <PrimaryButton text="Edit" onClick={() => { onEdit(deviceIn) }} />
+                        </Stack>
+                        <Stack style={{ paddingTop: 8, paddingRight: 24 }}>
+                           <DefaultButton text="Audit" onClick={() => { navigate(`/audit/device/${deviceIn.id}`) }} />
+                        </Stack>
+                     </Stack>
+                  </Stack>
+               </Stack>
+
+               <Stack>
+                  <Stack styles={{ root: { paddingLeft: 4, paddingRight: 0, paddingBottom: 4 } }}>
+                     <Stack>
+                        <Label>Device Reservations History</Label>
+                        {!!failureRate && deviceHealthRateBar(failureRate, 70, 15, `${failureRate}% Failure Rate`)}
+                     </Stack>
+
+                     {!jobState.telemetryItems.length && <Stack>
+                        <Text>No Results</Text>
+                     </Stack>}
+
+                     {!!jobState.telemetryItems.length && <Stack>
+                        <div style={{ overflowY: 'auto', overflowX: 'hidden', height: "670px" }} data-is-scrollable={true}>
+                           <Stack tokens={{ childrenGap: 12 }} style={{ paddingRight: 12 }}>
+                              <DetailsList
+                                 compact={true}
+                                 isHeaderVisible={false}
+                                 indentWidth={0}
+                                 items={jobState.telemetryItems}
+                                 columns={columns}
+                                 setKey="set"
+                                 selectionMode={SelectionMode.none}
+                                 onRenderItemColumn={renderItem}
+                                 onRenderRow={renderRow}
+
+                              />
+                           </Stack>
+                        </div>
+                     </Stack>}
+                  </Stack>
+               </Stack>
+
+            </Stack>
+
+         </Stack>
+      </Stack>
+
+   </Modal >);
+});
+
+
diff --git a/Engine/Source/Programs/Horde/HordeDashboard/src/components/devices/DevicePoolTelemetry.tsx b/Engine/Source/Programs/Horde/HordeDashboard/src/components/devices/DevicePoolTelemetry.tsx
new file mode 100644
--- /dev/null
+++ b/Engine/Source/Programs/Horde/HordeDashboard/src/components/devices/DevicePoolTelemetry.tsx
@@ -0,0 +1,402 @@
+import { DefaultButton, DetailsList, DetailsListLayoutMode, Dropdown, FocusZone, FocusZoneDirection, IColumn, IconButton, IDropdownOption, Modal, ScrollablePane, ScrollbarVisibility, SelectionMode, Spinner, SpinnerSize, Stack, Text } from "@fluentui/react";
+import { observer } from "mobx-react-lite";
+import { useState, useRef, useEffect, useCallback, useMemo } from "react";
+import { Link } from "react-router-dom";
+import { deviceHealthRateBar } from "./DeviceHealthWidgets"
+import { PoolTelemetryHandler, DeviceProblem } from "../../backend/DeviceTelemetryHandler"
+import dashboard, { StatusColor } from "../../backend/Dashboard";
+import { getHumanTime } from "../../base/utilities/timeUtils";
+import { getHordeStyling } from "../../styles/Styles";
+import { mergeStyles } from '@fluentui/react/lib/Styling';
+import { PoolTelemetryGraph, PoolDistributionGraph, PoolDistributionTimelineGraph} from "./DeviceTelemetryGraphs";
+import { GetDevicePlatformResponse } from "../../backend/Api";
+
+import * as d3 from "d3";
+
+// Handle bad "@types/d3" types, fix if addressed upstream
+const _d3 = d3 as any;
+
+const scrollableOverflowOverrideClass = mergeStyles({
+   // avoid tooltip to be cut by modal
+   overflow: "visible !important" as any
+});
+
+const HealthGraph: React.FC<{ handler: PoolTelemetryHandler, platform?: string, onChange: (options: object) => void }> = observer(({ handler, platform, onChange }) => {
+
+   const selectedPlatformItem = Array.from(handler.platforms.values()).find((p) => p.id === platform);
+   if (!selectedPlatformItem) {
+      return <Stack style={{ paddingLeft: 12, paddingRight: 32 }}>
+         <Text>Unknown platform {platform}</Text>
+      </Stack>
+   }
+
+   const problems = useMemo(() => handler.getProblemDevices(), [handler]);
+
+   const renderPlatform = useCallback((platform: GetDevicePlatformResponse) => {
+      let platformProblems = problems.get(platform.id) ?? [];
+      platformProblems = platformProblems.filter(p => p.problems > 8 && p.reservable).sort((a, b) => {
+         if (a.problemsRate < b.problemsRate) return 1;
+         if (a.problemsRate > b.problemsRate) return -1;
+         if (a.problems < b.problems) return 1;
+         if (a.problems > b.problems) return -1;
+         return 0;
+      });
+
+      const problemColumns: IColumn[] = [
+         { key: 'health-column1', name: 'Problem Devices', minWidth: 60, maxWidth: 150 },
+         { key: 'health-column2', name: 'Pool', minWidth: 60, maxWidth: 80 },
+         { key: 'health-column3', name: 'Issues/Reservations', minWidth: 100, maxWidth: 150 },
+         { key: 'health-column4', name: 'Latest Issue', minWidth: 300, maxWidth: 640 },
+      ];
+
+      const schemaColor = d3.scaleOrdinal()
+         .domain(Array.from(handler.pools.values()).map((p) => p.name))
+         .range(dashboard.darktheme ? d3.schemeDark2 : d3.schemeSet2);
+
+      let column = problemColumns.find(c => c.name === "Pool")!;
+      column.onRender = ((item: DeviceProblem) => {
+
+         return <Stack verticalFill={true} verticalAlign="center" horizontal>
+            <svg style={{ width: 5, height: 15, verticalAlign: 'middle', marginRight: '3px' }}>
+                  <rect x={0} y={0} width={5} height={15} fill={`${schemaColor(item.poolName)}`} />
+            </svg>
+            <Text>{item.poolName}</Text>
+         </Stack>
+      });
+
+      column = problemColumns.find(c => c.name === "Issues/Reservations")!;
+      column.onRender = ((item: DeviceProblem) => deviceHealthRateBar(item.problemsRate, 60, 15, item.problemsDesc));
+
+      column = problemColumns.find(c => c.name === "Latest Issue")!;
+      column.onRender = ((item: DeviceProblem) => {
+
+         const latest = item.latestProblem;
+
+         if (!latest) {
+            return null;
+         }
+
+         if (!latest.jobId || !latest.stepId) {
+            return <Stack verticalFill={true} verticalAlign="center">
+               <Text>{getHumanTime(latest.problemTime)}</Text>
+            </Stack>
+         }
+
+         const url = `/job/${latest.jobId}?step=${latest.stepId}`;
+
+         return <Stack verticalFill={true} verticalAlign="center">
+            <Link to={url} target="_blank">
+               <Text variant="small">{`${getHumanTime(latest.problemTime!)} - ${latest.jobName ?? "Unknown Job Name"} - ${latest.stepName ?? "Unknown Step Name"}`}</Text>
+            </Link>
+         </Stack>
+      });
+
+      column = problemColumns.find(c => c.name === "Problem Devices")!;
+      column.onRender = ((item: DeviceProblem) => {
+         const deviceName = item.deviceName;
+         const url = `/devices?pivotKey=pivot-key-automation&filter=${deviceName}`;
+         return <Stack verticalFill={true} verticalAlign="center">
+            <Link to={url} target="_blank"><Text variant="small">{deviceName}</Text></Link>
+         </Stack>
+      });
+
+      return <Stack style={{ paddingLeft: 12, paddingRight: 32 }}>
+         {!!platformProblems?.length &&
+            <Stack style={{ paddingLeft: 12, paddingTop: 12, width: 1220 }}>
+               <DetailsList
+                  isHeaderVisible={true}
+                  items={platformProblems}
+                  columns={problemColumns}
+                  layoutMode={DetailsListLayoutMode.justified}
+                  compact={true}
+                  selectionMode={SelectionMode.none}
+               />
+            </Stack>}
+      </Stack>
+   }, [handler]);
+
+   let title = "Device Health Per Platform";
+
+   if (handler.data && handler.data.length > 1) {
+      const minDate = getHumanTime(handler.data[0].date);
+      const maxDate = getHumanTime(handler.data[handler.data.length - 1].date);
+      title = `Device Health Per Platform / ${minDate} - ${maxDate}`;
+   }
+
+   return <Stack style={{ paddingBottom: 18 }}>
+         <Stack horizontal verticalAlign="start">
+            <Stack style={{ paddingLeft: 4 }}>
+               <Text variant="mediumPlus" styles={{ root: { fontWeight: "unset", fontFamily: "Horde Open Sans SemiBold" } }}>{title}</Text>
+            </Stack>
+            <Stack grow />
+            <TelemetryDropdown
+               handler={handler}
+               isPoolDisabled={true}
+               selectedPlatform={platform}
+               onChange={(options: object) => {
+                  onChange(options);
+               }}
+            />
+         </Stack>
+         <Stack style={{ paddingLeft: 12, paddingTop: 12 }}>
+            <FocusZone direction={FocusZoneDirection.vertical}>
+               <div style={{ position: 'relative', width: 1400, height: 650 }} data-is-scrollable>
+                  <ScrollablePane onScroll={() => { }}>
+                     <Stack style={{ paddingTop: 12 }} tokens={{ childrenGap: 24 }}>
+                        {renderPlatform(selectedPlatformItem)}
+                     </Stack>
+                  </ScrollablePane>
+               </div>
+            </FocusZone>
+         </Stack>
+      </Stack>
+});
+
+const DistributionGraph: React.FC<{ handler: PoolTelemetryHandler, platform?: string, onChange: (options: object) => void }> = observer(({ handler, platform, onChange }) => {
+   const [container, setContainer] = useState<HTMLDivElement | null>(null);
+   const graph = useRef<PoolDistributionGraph>(new PoolDistributionGraph(handler));
+
+   const [timelineContainer, setTimelineContainer] = useState<HTMLDivElement | null>(null);
+   const timelineGraph = useRef<PoolDistributionTimelineGraph>(new PoolDistributionTimelineGraph(handler));
+
+   const drawGraph = useCallback(() => {
+      if (container && platform) {
+         graph.current.draw(container, platform, 1400, 460, 200, timelineGraph.current?.startDate, timelineGraph.current?.endDate);
+      }
+   }, [container, platform, timelineGraph]);
+   drawGraph();
+
+   const drawTimelineGraph = useCallback(() => {
+      if (timelineContainer && platform) {
+         timelineGraph.current.draw(timelineContainer, platform, 800, 120, () => {
+            if (graph.current) {
+               graph.current.forceRender = true;
+               drawGraph();
+            }
+         });
+      }
+   }, [timelineContainer, platform, timelineGraph, drawGraph]);
+   drawTimelineGraph();
+
+   let title = "Device Reservation Distribution Per Job Step";
+
+   if (handler.data && handler.data.length > 1) {
+      const minDate = getHumanTime(handler.data[0].date);
+      const maxDate = getHumanTime(handler.data[handler.data.length - 1].date);
+      title = `Device Reservation Distribution Per Job Step / ${minDate} - ${maxDate}`;
+   }
+
+   return <Stack>
+         <Stack horizontal verticalAlign="start">
+            <Stack style={{ paddingLeft: 4 }}>
+               <Text variant="mediumPlus" styles={{ root: { fontWeight: "unset", fontFamily: "Horde Open Sans SemiBold" } }}>{title}</Text>
+            </Stack>
+            <Stack grow />
+            <TelemetryDropdown
+               handler={handler}
+               isPoolDisabled={true}
+               selectedPlatform={platform}
+               onChange={(options: object) => {
+                  if (graph.current) {
+                     graph.current.forceRender = true;
+                  }
+                  if (timelineGraph.current) {
+                     timelineGraph.current.forceRender = true;
+                  }
+                  onChange(options);
+               }}
+            />
+         </Stack>
+         <Stack style={{ paddingLeft: 12, paddingTop: 12}}>
+            <Stack.Item>
+               <div id="pool_distribution_graph_container" className="horde-no-darktheme" style={{ userSelect: "none" }} ref={(ref: HTMLDivElement) => setContainer(ref)} onMouseEnter={() => { }} onMouseLeave={() => { }} />
+            </Stack.Item>
+            <Stack.Item align="center">
+               <div id="pool_distribution_timeline_graph_container" className="horde-no-darktheme" style={{ userSelect: "none" }} ref={(ref: HTMLDivElement) => setTimelineContainer(ref)} onMouseEnter={() => { }} onMouseLeave={() => { }} />
+            </Stack.Item>
+         </Stack>
+      </Stack>
+});
+
+export const DevicePoolGraph: React.FC<{ handler: PoolTelemetryHandler, pool?: string, platform?: string, onChange: (options: object) => void }> = observer(({ handler, pool, platform, onChange }) => {
+
+   const [container, setContainer] = useState<HTMLDivElement | null>(null);
+   const graph = useRef<PoolTelemetryGraph>(new PoolTelemetryGraph(handler));
+
+   if (container && pool && platform) {
+      graph.current.draw(container, pool, platform, 1400, 600);
+   }
+
+   let title = "Device Usage Telemetry";
+
+   if (handler.data && handler.data.length > 1) {
+      const minDate = getHumanTime(handler.data[0].date);
+      const maxDate = getHumanTime(handler.data[handler.data.length - 1].date);
+      title = `Device Usage Telemetry / ${minDate} - ${maxDate}`;
+   }
+
+   return <Stack style={{ position: "relative" }}>
+         <Stack horizontal verticalAlign="start" style={{ paddingBottom: 18 }}>
+            <Stack style={{ paddingLeft: 4 }}>
+               <Text variant="mediumPlus" styles={{ root: { fontWeight: "unset", fontFamily: "Horde Open Sans SemiBold" } }}>{title}</Text>
+            </Stack>
+
+            <Stack grow />
+            <TelemetryDropdown
+               handler={handler}
+               selectedPool={pool}
+               selectedPlatform={platform}
+               onChange={(options: object) => {
+                  if (graph.current) {
+                     graph.current.forceRender = true;
+                  }
+                  onChange(options);
+               }}
+            />
+         </Stack>
+
+         <Stack style={{ paddingLeft: 12 }}>
+            <div id="pool_graph_container" className="horde-no-darktheme" style={{ userSelect: "none" }} ref={(ref: HTMLDivElement) => setContainer(ref)} onMouseEnter={() => { }} onMouseLeave={() => { }} />
+         </Stack>
+      </Stack>
+});
+
+export const TelemetryDropdown: React.FC<{ handler: PoolTelemetryHandler, selectedPool?: string, selectedPlatform?: string, isPoolDisabled?: boolean, isPlatformDisabled?: boolean, onChange: (options: object) => void }> = ({ handler, selectedPool, selectedPlatform, isPoolDisabled, isPlatformDisabled, onChange }) => {
+   const pools = Array.from(handler.pools.values()).sort((a, b) => a.name.localeCompare(b.name));
+   const platforms = Array.from(handler.platforms.values()).sort((a, b) => a.name.localeCompare(b.name));
+
+   if (!isPoolDisabled && !pools.length) {
+      return <Stack>No Pools</Stack>
+   }
+   if (!isPlatformDisabled && !platforms.length) {
+      return <Stack>No Platforms</Stack>
+   }
+
+   const defaultSelectedPool = selectedPool ?? pools[0].id;
+   const defaultSelectedPlatform = selectedPlatform ?? platforms[0].id;
+
+   useEffect(() => {
+      // Assign default values on post load
+      if ((!isPoolDisabled && !selectedPool) || (!isPlatformDisabled && !selectedPlatform)) {
+         onChange({pool: defaultSelectedPool, platform: defaultSelectedPlatform});
+      }
+   }, [selectedPool, selectedPlatform, defaultSelectedPool, defaultSelectedPlatform, isPoolDisabled, isPlatformDisabled]);
+
+   const poolOptions: IDropdownOption[] = pools.map(p => { return { key: `pool_key_${p.id}`, text: p.name, data: p.id } });
+   const platformOptions: IDropdownOption[] = platforms.map(p => { return { key: `platform_key_${p.id}`, text: p.name, data: p.id } });
+
+   return <Stack horizontal tokens={{ childrenGap: 12 }} style={{ paddingTop: 1 }}>
+            <Dropdown
+               defaultSelectedKey={`platform_key_${defaultSelectedPlatform}`}
+               options={platformOptions}
+               onChange={(event, option) =>  onChange({platform: option!.data})}
+               disabled={!!isPlatformDisabled}
+               style={{ width: 120 }}
+            />
+            <Dropdown
+               defaultSelectedKey={`pool_key_${defaultSelectedPool}`}
+               options={poolOptions}
+               onChange={(event, option) => onChange({pool: option!.data})}
+               disabled={!!isPoolDisabled}
+               style={{ width: 120 }}
+            />
+         </Stack>
+}
+
+const renderMenuButton = (item: any): JSX.Element => {
+   return <Stack key={item.key}><DefaultButton text={item.title} onClick={item.onClick} checked={item.checked} style={{width: '80px', paddingTop: 16 , paddingBottom: 16}} /></Stack>
+}
+
+export const DevicePoolTelemetryModal: React.FC<{ poolsIn?: Array<string>, platformsIn?: Array<string>, onClose?: () => void }> = ({ poolsIn, platformsIn, onClose }) => {
+
+   const { hordeClasses, modeColors } = getHordeStyling();
+   const [selected, setSelected] = useState<{ pool?: string, platform?: string }>({
+      pool: poolsIn?.find((e) => true), platform: platformsIn?.find((e) => true)});
+   const [graphOption, setGraphOption] = useState<{ usage?: boolean, health?: boolean, distribution?: boolean }>({usage: true});
+   const handler = useRef<PoolTelemetryHandler | undefined>(undefined);
+
+   if (!handler.current) {
+      handler.current = new PoolTelemetryHandler();
+   }
+
+   const menuItems = [
+      {
+         key: 'graph1',
+         title: 'Usage',
+         onClick: () => setGraphOption({usage: true}),
+         checked: graphOption.usage
+      },
+      {
+         key: 'graph2',
+         title: 'Health',
+         onClick: () => setGraphOption({health: true}),
+         checked: graphOption.health
+      },
+      {
+         key: 'graph3',
+         title: 'Distribution',
+         onClick: () => setGraphOption({distribution: true}),
+         checked: graphOption.distribution
+      }
+   ]
+
+   return <Stack>
+      <Modal className={hordeClasses.modal} scrollableContentClassName={scrollableOverflowOverrideClass} isOpen={true} isBlocking={true} topOffsetFixed={true} styles={{ main: { padding: 8, width: 1560, backgroundColor: modeColors.background, hasBeenOpened: false, top: "64px", position: "absolute", height: "760px", overflow: "visible !important" as any } }} onDismiss={() => { if (onClose) { onClose() } }}>
+         <Stack>
+            <Stack horizontal>
+               <Stack grow />
+               <Stack horizontalAlign="end" style={{ paddingTop: 8, paddingRight: 12 }}>
+                  <IconButton
+                     iconProps={{ iconName: 'Cancel' }}
+                     ariaLabel="Close popup modal"
+                     onClick={() => { if (onClose) onClose() }}
+                  />
+               </Stack>
+            </Stack>
+            <Stack style={{ paddingLeft: 24 }} horizontal tokens={{ childrenGap: 6 }}>
+               <Stack grow>
+                  {!handler.current.loaded &&
+                     <Stack horizontalAlign="center" style={{ marginLeft: 12, width: 1400 }} tokens={{ childrenGap: 32 }} grow>
+                        <Text variant="mediumPlus" styles={{ root: { fontWeight: "unset", fontFamily: "Horde Open Sans SemiBold" } }}>Loading Device Pool Telemetry</Text>
+                        <Spinner size={SpinnerSize.large} />
+                     </Stack>
+                  }
+                  {handler.current.loaded && !!graphOption.usage &&
+                     <DevicePoolGraph
+                        handler={handler.current}
+                        pool={selected.pool}
+                        platform={selected.platform}
+                        onChange={(options: object) => {
+                           setSelected({ ...selected, ...options });
+                        }}
+                     />
+                  }
+                  {handler.current.loaded && !!graphOption.health &&
+                     <HealthGraph
+                        handler={handler.current}
+                        platform={selected.platform}
+                        onChange={(options: object) => {
+                           setSelected({ ...selected, ...options });
+                        }}
+                     />
+                  }
+                  {handler.current.loaded && !!graphOption.distribution &&
+                     <DistributionGraph
+                        handler={handler.current}
+                        platform={selected.platform}
+                        onChange={(options: object) => {
+                           setSelected({ ...selected, ...options });
+                        }}
+                     />
+                  }
+               </Stack>
+               <Stack tokens={{ childrenGap: 4 }} style={{ paddingLeft: 5, paddingRight: 24 }}>{menuItems.map((i) => renderMenuButton(i))}</Stack>
+            </Stack>
+         </Stack>
+      </Modal>
+   </Stack>
+
+}
+
+
+
diff --git a/Engine/Source/Programs/Horde/HordeDashboard/src/components/devices/DeviceTelemetryGraphs.tsx b/Engine/Source/Programs/Horde/HordeDashboard/src/components/devices/DeviceTelemetryGraphs.tsx
new file mode 100644
--- /dev/null
+++ b/Engine/Source/Programs/Horde/HordeDashboard/src/components/devices/DeviceTelemetryGraphs.tsx
@@ -0,0 +1,1113 @@
+import { PoolTelemetryHandler, UsagePlatformTelemetry, PlatformTelemetry } from "../../backend/DeviceTelemetryHandler"
+import { getHordeStyling } from "../../styles/Styles";
+import { getHumanTime, msecToElapsed } from "../../base/utilities/timeUtils";
+import dashboard, { StatusColor } from "../../backend/Dashboard";
+import { TelemetryData } from "../../backend/DeviceTelemetryHandler";
+import moment from "moment";
+
+import * as d3 from "d3";
+
+// Handle bad "@types/d3" types, fix if addressed upstream
+const _d3 = d3 as any;
+
+type SelectionType = d3.Selection<SVGGElement, unknown, null, undefined>;
+type ScalarTime = d3.ScaleTime<number, number>
+type DivSelectionType = d3.Selection<HTMLDivElement, unknown, null, undefined>;
+
+export class PoolTelemetryGraph {
+   constructor(handler: PoolTelemetryHandler) {
+      this.handler = handler;
+      this.margin = { top: 32, right: 28, bottom: 40, left: 48 };
+   }
+
+   set(poolId: string, platformId: string) {
+
+      this.poolId = poolId;
+      this.platformId = platformId;
+      this.streamIds.clear();
+      this.data = undefined;
+
+      const handler = this.handler;
+      this.data = handler.getData(poolId, platformId);
+
+   }
+
+   draw(container: HTMLDivElement, poolId: string, platformId: string, width: number, height: number) {
+
+      if (this.hasRendered && !this.forceRender) {
+         return;
+      }
+      
+      this.container = container;
+      this.width = width;
+      this.height = height;
+      const margin = this.margin;
+      this.clear();
+
+      this.set(poolId, platformId);
+
+      const data = this.data;
+
+      if (!data) {
+         return;
+      }
+
+      this.hasRendered = true;
+      this.forceRender = false;
+      const { modeColors } = getHordeStyling();
+
+      const xPlot = this.xPlot = this.xPlotMethod();
+      const xPlotAxis = this.xPlotMethod();
+
+      const totalDevices = data.deviceIds.length;
+
+      const streams = Array.from(this.streamIds).sort((a, b) => a.localeCompare(b));
+
+      const schemaColor = d3.scaleOrdinal()
+         .domain(streams)
+         .range(dashboard.darktheme ? d3.schemeDark2 : d3.schemeSet2);
+
+      const scolors = dashboard.getStatusColors();
+      const colors: Record<string, string> = {
+         "Available": scolors.get(StatusColor.Success)!,
+         "Problem": scolors.get(StatusColor.Failure)!,
+         "Maintenance": dashboard.darktheme ? "#413F3D" : "#A19F9D",
+         "Disabled": "#715F5D"!
+      };
+
+      const y = d3.scaleLinear()
+         .domain([0, totalDevices])
+         .range([height - margin.top - margin.bottom, 0]);
+
+      const svg = d3.select(container)
+         .append("svg")
+         .attr("width", width)
+         .attr("height", height)
+         .attr("viewBox", [0, 0, width, height] as any)
+         .style("background-color", modeColors.crumbs);
+
+      // define a clipping region
+      svg.append("clipPath")
+         .attr("id", this.clipId)
+         .append("rect")
+         .attr("x", margin.left + 1)
+         .attr("y", margin.top)
+         .attr("width", width - margin.left - margin.right - 241)
+         .attr("height", height - margin.top);
+
+      const zoom = d3.zoom()
+         .scaleExtent([1, 12])
+         .extent([[margin.left, 0], [width - margin.right - 240, height]])
+         .translateExtent([[margin.left, -Infinity], [width - margin.right - 240, Infinity]])
+         .on("zoom", zoomed);
+
+      const graph = this;
+
+      function zoomed(event: any) {
+         if (graph.freezeTooltip) {
+            return;
+         }
+
+         // update bottom axis drawing
+         xAxis!.call(graph.xAxisMethod as any, event.transform.rescaleX(xPlotAxis));
+         // update x plot mapping
+         xPlot!.range([margin.left, width - margin.right - 240].map(d => event.transform.applyX(d)));
+         // update graph drawing
+         areaChart
+            .selectAll("path")
+            .attr("d", area as any)
+         // update markers
+         makersChart
+            .selectAll("path")
+            .attr("transform", plotMarker as any)
+      }
+
+      svg.call(zoom as any);
+
+      // left axis
+      svg.append("g")
+         .attr("transform", `translate(${this.margin.left},${margin.bottom - 8})`)
+         .call(d3.axisLeft(y).ticks(totalDevices))
+
+      const area = d3.area()
+         .x((d: any) => xPlot!(d.data.date))
+         .y0((d) => y(d[0]) + 32)
+         .y1((d) => y(d[1]) + 32)
+         .curve(d3.curveStep);
+
+      const keys = ["Disabled", "Maintenance", "Problem"];
+
+      data.streamIds.forEach(sid => keys.push(data.streamNames[sid]));
+
+      const stackedData = d3.stack()
+         .keys(keys)
+         (data.telemetry as any)
+
+      const highlight = (e: any, d: any, fromArea?: boolean) => {
+         this.currentTarget = fromArea ? d : undefined;
+         if (this.freezeTooltip) {
+            return;
+         }
+         d3.selectAll(".telemetryArea").style("opacity", .2)
+         this.freezedHighlights.forEach((s) => {
+            d3.selectAll("." + s).style("opacity", 1)
+         });
+         const stackTarget = data.streamSelectors[d] ?? d;
+         if (!this.freezedHighlights.has(stackTarget)) {
+            d3.selectAll("." + stackTarget).style("opacity", 1)
+         }
+      }
+
+      const noHighlight = (e: any, d: any) => {
+         this.currentTarget = undefined;
+         if (this.freezeTooltip) {
+            return;
+         }
+         if (!this.freezedHighlights.size) {
+            d3.selectAll(".telemetryArea").style("opacity", 1)
+         } else {
+            d3.selectAll(".telemetryArea").style("opacity", .2)
+            this.freezedHighlights.forEach((s) => {
+               d3.selectAll("." + s).style("opacity", 1)
+            });
+         }
+      }
+
+      const selectHighlight = (e: any, d: any) => {
+         if (this.freezeTooltip) {
+            return;
+         }
+         const stackTarget = data.streamSelectors[d] ?? d;
+         if (this.freezedHighlights.has(stackTarget)) {
+            this.freezedHighlights.delete(stackTarget);
+         } else {
+            this.freezedHighlights.add(stackTarget);
+         }
+      }
+
+      const areaChart = svg.append('g')
+         .attr("clip-path", `url(#${this.clipId})`);
+
+      areaChart
+         .selectAll("telemetrylayers")
+         .data(stackedData)
+         .enter()
+         .append("path")
+         .attr("shape-rendering", "criptEdges")
+         .attr("class", function (d) { return "telemetryArea " + (data.streamSelectors[d.key] ?? d.key); })
+         .style("shapeRendering", "geometricPrecision")
+         .style("fill", function (d) { return colors[d.key] ?? schemaColor(d.key) as any })
+         .attr("d", area as any)
+         .on("mouseover", (e:any) => { highlight(e, e.target.classList[1], true) })
+         .on("mouseleave", (e:any) => { noHighlight(e, null) })
+         .on("click", (e:any) => { selectHighlight(e, e.target.classList[1]) });
+
+      // problem markers
+      const triangleMark = d3.symbol().type(d3.symbolDiamond).size(60)();
+      const plotMarker = (d) => `translate(${xPlot!(d.date)},${height - 30})`;
+      const markersData = data.telemetry.map(function(d) { 
+         return {
+            date: d.date,
+            value: d["Problem"] > 1? Math.log(d["Problem"]) / Math.log(totalDevices) : 0
+         };
+      });
+      const makersChart = svg.append('g')
+         .attr("clip-path", `url(#${this.clipId})`);
+      makersChart
+         .selectAll("makers")
+         .data(markersData)
+         .enter()
+         .append("path")
+         .style("fill", colors["Problem"])
+         .attr("d", (d) => d.value > 0? triangleMark : null)
+         .attr("opacity", (d) => d.value)
+         .attr("transform", plotMarker as any)
+         .attr('stroke', 'transparent').attr('stroke-width', 5)
+         .on("mouseover", (e, d) => { highlight(e, "Problem", true) })
+         .on("mouseleave", (e, d) => { noHighlight(e, null) })
+         .on("click", (e, d) => { selectHighlight(e, "Problem") });
+
+      // bottom axis
+      const xAxis = this.xAxis = svg.append("g")
+         .call(this.xAxisMethod as any, xPlot)
+         .style("pointer-events", "none");
+
+      // legend
+      const size = 15
+      const posx = 1160;
+      const posy = margin.top;
+      svg.selectAll("legendrect")
+         .data(keys)
+         .enter()
+         .append("rect")
+         .attr("class", function (d) { return "telemetryArea " + (data.streamSelectors[d] ?? d) })
+         .attr("x", posx)
+         .attr("y", function (d, i) { return posy + 10 + i * (size + 5 + 4) })
+         .attr("width", size)
+         .attr("height", size)
+         .style("fill", function (d: any) { return colors[d] ?? schemaColor(d) as any })
+         .on("mouseover", highlight as any)
+         .on("mouseleave", noHighlight as any)
+         .on("click", selectHighlight as any);
+
+      svg.selectAll("legendlabels")
+         .data(keys)
+         .enter()
+         .append("text")
+         .attr("x", posx + size * 1.2 + 8)
+         .attr("y", function (d, i) { return posy + 10 + i * (size + 5 + 4) + (size / 2) })
+         .style("fill", modeColors.text)
+         .text(function (d) { return d })
+         .attr("text-anchor", "left")
+         .attr("class", function (d) { return "telemetryArea " + (data.streamSelectors[d] ?? d) })
+         .style("alignment-baseline", "middle")
+         .style("font-family", "Horde Open Sans Regular")
+         .style("font-size", "13px")
+         .on("mouseover", highlight as any)
+         .on("mouseleave", noHighlight as any)
+         .on("click", selectHighlight as any);
+
+      // tooltip
+      this.tooltip = d3.select(container)
+         .append("div")
+         .style("visibility", "hidden")
+         .style("background-color", modeColors.background)
+         .style("border", "solid")
+         .style("border-width", "1px")
+         .style("border-radius", "3px")
+         .style("border-color", dashboard.darktheme ? "#413F3D" : "#2D3F5F")
+         .style("padding", "8px")
+         .style("position", "absolute");
+
+      svg.on("mousemove", (event) => this.handleMouseMove(event));
+      // dblclick for double click events
+      svg.on("click", (event) => {
+         if (this.freezeTooltip) {
+            this.freezeTooltip = false;
+            if (this.currentTarget !== undefined && !this.freezedHighlights.has(this.currentTarget))
+            {
+               // force the tooltip to hide
+               this.currentTarget = undefined;
+            }
+            // update the tooltip
+            this.handleMouseMove(event);
+         }
+         if (this.currentTarget) {
+            this.freezeTooltip = true;
+         }
+         // force refreshing hightlights
+         noHighlight(event, null);
+      });
+
+      // Average load
+      const avgLoad = svg.append("g")
+         .attr("transform", "translate(8,6)");
+      const colorScale = d3.scaleLinear<string>()
+         .domain([10, 50, 80])
+         .range([scolors.get(StatusColor.Success)!, scolors.get(StatusColor.Warnings)!, scolors.get(StatusColor.Failure)!]);
+      const avgLoadWidth = 100;
+      const load = totalDevices > 0 ? Math.floor(data.telemetry.reduce((acru, item) => acru + totalDevices - (item.available?.length ?? 0), 0) / (data.telemetry.length * totalDevices) * 100) : 0;
+      avgLoad.append("rect")
+         .attr("width", avgLoadWidth)
+         .attr("height", 20)
+         .style("fill", "none")
+         .style("stroke", modeColors.text)
+         .style("stroke-width", 1);
+      avgLoad.append("text")
+         .attr("x", avgLoadWidth + 6)
+         .attr("y", 10)
+         .text(`${load}% Avg Load`)
+         .style("fill", modeColors.text)
+         .attr("text-anchor", "left")
+         .style("alignment-baseline", "middle")
+         .style("font-family", "Horde Open Sans Regular")
+         .style("font-size", "13px");
+      avgLoad.selectAll("averageLoad")
+         .data([load])
+         .enter()
+         .append("rect")
+         .attr("width", (d) => Math.floor(d / 100 * (avgLoadWidth-1)))
+         .attr("height", 19)
+         .style("fill", (d) => colorScale(d));
+   }
+
+   private clear() {
+      if (!!this.container) {
+         d3.select(this.container).selectAll("*").remove();
+      }
+   }
+
+   private closestData(x: number): PlatformTelemetry | undefined {
+
+      const xPlot = this.xPlot;
+      if (!xPlot) {
+         return;
+      }
+
+      const data = this.data!.telemetry;
+
+      let closest = data.reduce((best, value, i) => {
+
+         let absx = Math.abs(xPlot(value.date) - x)
+         if (absx < best.value) {
+            return { index: i, value: absx };
+         }
+         else {
+            return best;
+         }
+      }, { index: 0, value: Number.MAX_SAFE_INTEGER });
+
+      return data[closest.index];
+   }
+
+   private updateTooltip(show: boolean, x?: number, y?: number, html?: string) {
+      if (!this.tooltip) {
+         return;
+      }
+
+      x = x ?? 0;
+      y = y ?? 0;
+
+      this.tooltip
+         .html(html ?? "")
+         .style("left", (x + 40) + "px")
+         .style("top", y + "px")
+         .style("font-family", "Horde Open Sans Regular")
+         .style("font-size", "10px")
+         .style("line-height", "16px")
+         .style("shapeRendering", "crispEdges")
+         .style("stroke", "none")
+         .style("visibility", show ? "visible" : "hidden");
+   }
+
+   private handleMouseMove(event: any) {
+      if (this.freezeTooltip) {
+         return;
+      }
+
+      const target = this.currentTarget;
+      if (!target) {
+         this.updateTooltip(false);
+         return;
+      }
+
+      const data = this.data!;
+      const streamId = data.streamIds.find(sid => {
+         const streamName = data.streamSelectorsReverse[target!];
+         return data.streamNamesReverse[streamName] === sid;
+      });
+
+      const closest = this.closestData(_d3.pointer(event)[0]);
+      if (!closest) {
+         this.updateTooltip(false);
+         return;
+      }
+
+      const displayTime = moment(closest.date);
+      const format = dashboard.display24HourClock ? "HH:mm:ss z" : "LT z";
+
+      let displayTimeStr = displayTime.format('MMM Do YYYY') + ` - ${displayTime.format(format)}`;
+
+      let targetName = "???";
+      if (target) {
+         targetName = data.streamSelectorsReverse[target];
+         if (!targetName) {
+            targetName = target;
+         }
+      }
+
+      let text = `<div style="padding-left:4px;padding-right:12px;padding-bottom:12px">`;
+
+      text += `<div style="padding-bottom:8px;padding-top:2px"><span style="font-family:Horde Open Sans Semibold;font-size:11px">${targetName}</span></div>`;
+
+      if (streamId && closest.reserved && closest.reserved[streamId]) {
+         closest.reserved[streamId].forEach(v => {
+            const d = data.devices[v.deviceId]?.name ?? v.deviceId;
+            text += `<div style="padding-bottom:8px">`;
+            const href = `/job/${v.jobId}?step=${v.stepId}`;
+            text += `<a href="${href}" target="_blank"><span style="font-family:Horde Open Sans Semibold;font-size:11px;padding-left:8px">${d}: ${v.stepName}</span></a><br/>`
+            if (v.jobId && v.stepId) {
+               const metrics = this.handler.getStepMetrics(v.jobId, v.stepId);
+               if (metrics && (metrics.duration || metrics.numKits > 1 || metrics.numProblems > 0)) {
+
+                  text += `<div style="padding-left: 16px">`;
+
+                  if (metrics.numKits > 1) {
+                     text += `<span style="">Devices:</span> ${metrics.numKits} <br/>`;
+                  }
+
+                  if (metrics.numProblems > 0) {
+                     text += `<span style="font-family:Horde Open Sans Semibold;color:#FF0000">Problems: ${metrics.numProblems}</span> <br/>`;
+                  }
+
+                  if (metrics.duration) {
+                     text += `<span style="">Duration:</span> ${metrics.duration} <br/>`;
+                  }
+
+                  text += `</div>`;
+               }
+
+               text += `</div>`;
+            }
+         });
+      }
+      else if (target === "Maintenance" && closest.maintenance) {
+         closest.maintenance.forEach(d => {
+            d = data.devices[d]?.name ?? d;
+            text += `<span style="font-family:Horde Open Sans Semibold;font-size:11px;padding-left:8px">${d}</span><br/>`
+         });
+         text += `<br/>`;
+      }
+      else if (target === "Problem") {
+         if (!closest.problem?.length) {
+            this.updateTooltip(false);
+            return;
+         }
+
+         const problems = this.handler.getProblemSteps(closest.problem, closest.date);
+
+         closest.problem.sort((a,b) => {
+            if (!problems[a] || !problems[b]) return 0;
+            // group problems by job step name
+            const streamA = !!problems[a].streamId ? data.streamNames[problems[a].streamId] : "Unknown stream";
+            const streamB = !!problems[b].streamId ? data.streamNames[problems[b].streamId] : "Unknown stream";
+            const nameA = `${streamA} - ${(problems[a].stepName ?? problems[a].jobName)}`;
+            const nameB = `${streamB} - ${(problems[b].stepName ?? problems[b].jobName)}`;
+            if (nameA < nameB) return -1;
+            if (nameA > nameB) return 1;
+            return 0;
+         }).forEach(d => {
+            const problem = problems[d];
+            d = data.devices[d]?.name ?? d;
+
+            text += `<span style="font-family:Horde Open Sans Semibold;font-size:11px;padding-left:8px">${d}</span><br/>`
+
+            if (problem && problem.streamId && problem.jobName) {
+
+               const name = problem.stepName ?? problem.jobName;
+               let desc = `${data.streamNames[problem.streamId]} - ${name}`;
+
+               const href = `/job/${problem.jobId}?step=${problem.stepId}`;
+
+               text += `<a href="${href}" target="_blank"><div style="padding-bottom:8px;padding-left:8px"><span style="">${desc}</span></div></a>`
+
+            }
+
+         });
+      }
+      else if (target === "Disabled" && closest.disabled) {
+         closest.disabled.forEach(d => {
+            d = data.devices[d]?.name ?? d;
+            text += `<span style="font-family:Horde Open Sans Semibold;font-size:11px;padding-left:8px">${d}</span><br/>`
+         });
+
+         text += `<br/>`;
+      }
+
+      text += `${displayTimeStr}<br/>`;
+
+      text += "</div>";
+
+      this.updateTooltip(true, _d3.pointer(event)[0], _d3.pointer(event)[1], text);
+
+   }
+
+   private get xAxisMethod() {
+      return (g: SelectionType, x: ScalarTime) => {
+         const tickSet: Set<string> = new Set();
+
+         g.attr("transform", `translate(0,${this.height - this.margin.bottom})`)
+            .style("font-family", "Horde Open Sans Regular")
+            .style("font-size", "11px")
+            .call(d3.axisBottom(x)
+               .ticks(8)
+               .tickFormat(d => {
+                  const v = (d as Date).toDateString();
+                  if (!tickSet.has(v)) {
+                     tickSet.add(v)
+                     return getHumanTime(d as Date)
+                  }
+
+                  return (d as Date).toLocaleString(undefined, { hour: 'numeric', minute: "numeric", hour12: true })
+               })
+               .tickSizeOuter(0))
+      }
+   }
+
+   private xPlot?: ScalarTime;
+
+   private xPlotMethod() {
+
+      const data = this.data!.telemetry;
+
+      return d3.scaleTime()
+         .domain(d3.extent(data, (d => d.date)) as any)
+         .range([this.margin.left, this.width - this.margin.right - 240])
+   }
+
+   handler: PoolTelemetryHandler;
+
+   container?: HTMLDivElement;
+   private margin: { top: number, right: number, bottom: number, left: number }
+   private width = 0;
+   private height = 0;
+
+   private data?: TelemetryData;
+   private poolId = "";
+   private platformId = "";
+   private streamIds = new Set<string>();
+
+   private freezeTooltip = false;
+   private freezedHighlights = new Set<string>();
+
+   private xAxis?: SelectionType;
+   private tooltip?: DivSelectionType;
+   private readonly clipId = "pool_history_clip_path";
+
+   private currentTarget?: string;
+
+   private hasRendered = false;
+   forceRender = false;
+
+   extent?: any;
+}
+
+export class PoolDistributionGraph {
+   constructor(handler: PoolTelemetryHandler) {
+      this.handler = handler;
+      this.margin = { top: 0, right: 0, bottom: 0, left: 0 };
+   }
+
+   set(platformId: string, startDate: Date | undefined, endDate: Date | undefined) {
+      this.platformId = platformId;
+      this.stepIds.clear();
+      this.data = this.handler.getSteps(platformId, startDate, endDate);
+      this.data.forEach((s) => { this.stepIds.add(s.stepName) });
+   }
+
+   draw(container: HTMLDivElement, platformId: string, width: number, height: number, radius: number, startDate: Date | undefined, endDate: Date | undefined) {
+
+      if (this.hasRendered && !this.forceRender) {
+         return;
+      }
+      
+      this.container = container;
+      this.width = width;
+      this.height = height;
+      const margin = this.margin;
+      this.clear();
+
+      this.set(platformId, startDate, endDate);
+
+      const data = this.data;
+
+      if (!data) {
+         return;
+      }
+
+      this.hasRendered = true;
+      this.forceRender = false;
+
+      const root = this.modelData();
+      if (!root) {
+         return;
+      }
+      d3.partition().size([2 * Math.PI, radius])(root);
+
+      const { modeColors } = getHordeStyling();
+
+      const stepIds = Array.from(this.stepIds).sort((a, b) => a.localeCompare(b));
+
+      const schemaColor = d3.scaleOrdinal()
+         .domain(stepIds)
+         .range(dashboard.darktheme ? d3.schemeDark2 : d3.schemeSet2);
+
+      const svg = d3.select(container)
+         .append("svg")
+         .attr("width", width)
+         .attr("height", height)
+         .attr("viewBox", [0, 0, width, height] as any)
+         .style("background-color", modeColors.crumbs);
+
+      // define a clipping region
+      svg.append("clipPath")
+         .attr("id", this.clipId)
+         .append("rect")
+         .attr("x", margin.left + 1)
+         .attr("y", margin.top)
+         .attr("width", width - margin.left - margin.right)
+         .attr("height", height - margin.top);
+
+      // arc generator
+      const arc = d3.arc()
+         .cornerRadius(2)
+         .startAngle((d: any) => d.x0)
+         .endAngle((d: any) => d.x1)
+         .padAngle(0.025)
+         .padRadius(radius / 4)
+         .innerRadius((d: any) => d.y0)
+         .outerRadius((d: any) => d.y1 - 2);
+
+      // draw the pieChart
+      const center = [width / 2, height / 2];
+      const pieChart = svg.append('g')
+         .attr("clip-path", `url(#${this.clipId})`)
+         .append('g')
+         .attr("transform", `translate(${center})`);
+
+      const arcOpacity = (d) => {
+         const family = d.ancestors().reverse()[1]?.data.name;
+         if (!family || !d.depth) return 0;
+         if (d.depth == 1) return 1;
+         if (d.depth <= this.displayDepth
+               && (!!this.displayFamily && family == this.displayFamily)) return 1;
+         if ((!!this.selectedFamily && family == this.selectedFamily)
+               || (!!this.displayFamily && family == this.displayFamily)) return 0.3;
+         return 0;
+      }
+
+      const labelOpacity = (d) => {
+         const family = d.ancestors().reverse()[1]?.data.name;
+         if (!family || d.depth != 1) return 0;
+         const isNarrow = (d.y0 + d.y1) /2 * (d.x1 - d.x0) < 25;
+         if (d.depth == this.displayDepth && ((!this.displayFamily && !isNarrow) || family == this.selectedFamily)) return 1;
+         if (!isNarrow || family == this.selectedFamily) return 0.5;
+         return 0;
+      }
+
+      const arcPointerEvent = (d) => {
+         return d.depth <= this.displayDepth ? "auto" : "none";
+      }
+
+      const highlight = (e: any, d: any) => {
+         if (!d.depth) return;
+         this.displayDepth = d.depth + 1;
+         this.displayFamily = d.ancestors().reverse()[1]?.data.name;
+         d3.selectAll(".arcStep")
+            .attr("opacity", arcOpacity)
+            .style("pointer-events", arcPointerEvent);
+         d3.selectAll(".arcStepLabel")
+            .attr("opacity", labelOpacity);
+         d3.selectAll(`.${d.data.id}`)
+            .attr("opacity", 1)
+            .raise();
+      }
+
+      const noHighlight = (e: any, d: any) => {
+         this.displayDepth = this.selectedDepth;
+         this.displayFamily = this.selectedFamily;
+         d3.selectAll(".arcStep")
+            .attr("opacity", arcOpacity)
+            .style("pointer-events", arcPointerEvent);
+         d3.selectAll(".arcStepLabel")
+            .attr("opacity", labelOpacity);
+         if (!!this.selectedItem) {
+            d3.selectAll(`.${this.selectedItem}`)
+               .attr("opacity", 1)
+               .raise();
+         }
+      }
+
+      const selectHighlight = (e: any, d: any) => {
+         if (!d.depth) {
+            this.selectedDepth = 1;
+            this.selectedFamily = "";
+            this.selectedItem = "";
+         } else {
+            this.selectedDepth = d.depth + 1;
+            this.selectedFamily = d.ancestors().reverse()[1]?.data.name;
+            this.selectedItem = d.data.id;
+         }
+         noHighlight(e, d);
+      }
+
+      const pieceColor = (d) => {
+         const pColor = d3.color(schemaColor(d.ancestors().reverse()[1]?.data.name) as any);
+         return dashboard.darktheme ? pColor?.brighter(d.depth / 2) : pColor?.darker(d.depth / 2)
+      }
+
+      pieChart
+         .selectAll("steps")
+         .data(root.descendants())
+         .enter()
+         .append("path")
+         .attr("d", arc as any)
+         .attr("class", "arcStep")
+         .attr("fill", (d: any) => pieceColor(d) as any)
+         .attr('stroke', 'transparent').attr('stroke-width', 3)
+         .attr("opacity", arcOpacity)
+         .style("pointer-events", arcPointerEvent)
+         .on("mouseover", highlight as any)
+         .on("mouseleave", noHighlight as any)
+         .on("click", selectHighlight as any);
+
+      // center label
+      const centerLabel = svg.append('g')
+         .attr("transform", `translate(${center})`)
+         .style("pointer-events", "none");
+
+      centerLabel.append("text")
+         .attr("dy", "0em")
+         .attr("fill", modeColors.text)
+         .style("text-anchor", "middle")
+         .text((root.data as any).name)
+      centerLabel.append("text")
+         .attr("dy", "1em")
+         .attr("fill", modeColors.text)
+         .style("text-anchor", "middle")
+         .text(msecToElapsed((root.data as any).durationMS, false, false));
+
+      // labels
+      const labels = svg.append('g')
+         .attr("clip-path", `url(#${this.clipId})`)
+         .append('g')
+         .attr("transform", `translate(${center})`)
+         .style("pointer-events", "none");
+
+      const addAlongVector = (vector, factor) => {
+         const [x, y] = vector;
+         const length = Math.sqrt(x * x + y * y);
+         const [xn, yn] = [x / length, y / length];
+         return [x + xn * factor, y + yn * factor];
+      }
+
+      const labelPosition = (d, offsetx) => {
+         const textAnchor = addAlongVector(arc.centroid(d), radius * 0.2);
+         if (!!offsetx){
+            textAnchor[0] += offsetx * (textAnchor[0] >= 0 ? 1 : -1);
+         }
+         return textAnchor;
+      }
+
+      labels.selectAll("lines")
+         .data(root.descendants())
+         .enter()
+         .append("polyline")
+         .attr("class", (d: any) => `arcStepLabel ${d.data.id}`)
+         .style("fill", "none")
+         .style("stroke", modeColors.background).style("stroke-width", 2)
+         .attr("points", (d: any) => {
+            const textAnchor = labelPosition(d, 0);
+            textAnchor[1] += 4;
+            const textUnderline = labelPosition(d, 40);
+            textUnderline[1] += 4;
+            const startAnchor = addAlongVector(arc.centroid(d), radius * 0.05)
+            return ([startAnchor, textAnchor, textUnderline]).toString();
+         })
+         .attr("opacity", labelOpacity)
+         .clone(true)
+         .style("stroke", modeColors.text)
+         .style("stroke-width", 1);
+
+      labels.selectAll("durations")
+         .data(root.descendants())
+         .enter()
+         .append("text")
+         .attr("class", (d: any) => `arcStepLabel ${d.data.id}`)
+         .text((d: any) => msecToElapsed((d.data as any).durationMS, false, false))
+         .style("fill", modeColors.text)
+         .attr("font-weight", "bold")
+         .attr("font-size", "smaller")
+         .style("stroke", modeColors.background).style("stroke-width", 3)
+         .attr("transform", (d: any) => `translate(${arc.centroid(d)})`)
+         .style("text-anchor", "middle")
+         .attr("opacity", labelOpacity)
+         .clone(true)
+         .style("fill", modeColors.text)
+         .style("stroke", "none");
+
+      labels.selectAll("labels")
+         .data(root.descendants())
+         .enter()
+         .append("text")
+         .attr("class", (d: any) => `arcStepLabel ${d.data.id}`)
+         .text((d: any) => d.data.name)
+         .style("fill", modeColors.text)
+         .style("stroke", modeColors.background).style("stroke-width", 3)
+         .attr("transform", (d: any) => `translate(${labelPosition(d, 5)})`)
+         .style("text-anchor", (d: any) => arc.centroid(d)[0] >= 0 ? "start" : "end")
+         .attr("opacity", labelOpacity)
+         .clone(true)
+         .style("fill", modeColors.text)
+         .style("stroke", "none");
+
+      // reset selection highlight
+      pieChart.call(() => { noHighlight(null, null) });
+   }
+
+   private clear() {
+      if (!!this.container) {
+         d3.select(this.container).selectAll("*").remove();
+      }
+   }
+
+   private modelData() {
+      if (!this.data) return;
+
+      const items: Map<string, { id: string, name: string, type: string, parentId: string, durationMS: number }> = new Map();
+      items.set(this.platformId, {id: this.platformId, name: this.handler.platforms.get(this.platformId)!.name, type: "Platform", parentId: "", durationMS: 0});
+      let totalTimeMS = 0;
+      const invalidClassCharacters = /[^_a-z0-9-]/gi;
+
+      this.data.forEach(s => {
+         totalTimeMS += s.durationMS;
+         // step layer
+         const stepKey = s.stepName.replaceAll(invalidClassCharacters, "_");
+         if (!items.has(stepKey)) {
+            items.set(stepKey, { id: stepKey, name: s.stepName, type: "Step", parentId: this.platformId, durationMS: 0 })
+         }
+         // stream layer
+         const streamKey = stepKey + s.stream?.replaceAll(invalidClassCharacters, "_");
+         if (!items.has(streamKey)) {
+            items.set(streamKey, { id: streamKey, name: s.stream, type: "Stream", parentId: stepKey, durationMS: 0 })
+         }
+         // pool layer
+         const poolKey = streamKey + s.pool?.replaceAll(invalidClassCharacters, "_");
+         if (!items.has(poolKey)) {
+            items.set(poolKey, { id: poolKey, name: s.pool, type: "Pool", parentId: streamKey, durationMS: s.durationMS })
+         } else {
+            const item = items.get(stepKey)!;
+            item.durationMS += s.durationMS;
+         }
+      });
+
+      const rootItem = items.get(this.platformId)!;
+
+      const root = d3.stratify()
+         .id((d: any) => d.id)
+         .parentId((d: any) => d.parentId)
+         (Array.from(items.values()));
+
+      // compute time layout
+      root.sum((d: any) => d.durationMS);
+      root.children = root.children?.filter((d: any) => d.value / totalTimeMS > 0.02);
+      root.sum((d: any) => d.durationMS);
+      root.each((d: any) => { d.data.durationMS = d.value });
+      root.sort((a, b) => b.value! - a.value!);
+
+      return root
+   }
+
+
+   handler: PoolTelemetryHandler;
+
+   container?: HTMLDivElement;
+   private margin: { top: number, right: number, bottom: number, left: number }
+   private width = 0;
+   private height = 0;
+
+   private data?: Map<string, any>;
+   private telemetry?: UsagePlatformTelemetry[];
+   private platformId = "";
+   private stepIds = new Set<string>();
+
+   private displayDepth = 1;
+   private displayFamily = "";
+   private selectedDepth = 1;
+   private selectedFamily = "";
+   private selectedItem = "";
+
+   private readonly clipId = "pool_distribution_clip_path";
+
+   private hasRendered = false;
+   forceRender = false;
+}
+
+export class PoolDistributionTimelineGraph {
+   constructor(handler: PoolTelemetryHandler) {
+      this.handler = handler;
+      this.margin = { top: 10, right: 30, bottom: 20, left: 30 };
+   }
+
+   set(platformId: string) {
+      this.platformId = platformId;
+      this.data = this.handler.getPlatformTelemetry(platformId);
+   }
+
+   draw(container: HTMLDivElement, platformId: string, width: number, height: number, onDateChange: () => void) {
+
+      if (this.hasRendered && !this.forceRender) {
+         return;
+      }
+      
+      this.container = container;
+      this.width = width;
+      this.height = height;
+      const margin = this.margin;
+      this.clear();
+
+      this.set(platformId);
+
+      const data = this.data;
+
+      if (!data) {
+         return;
+      }
+
+      this.hasRendered = true;
+      this.forceRender = false;
+
+      const { modeColors } = getHordeStyling();
+
+      const maxReservations = data.reduce((a, d) => Math.max(a, d.reservations), 1);
+
+      const svg = d3.select(container)
+         .append("svg")
+         .attr("width", width)
+         .attr("height", height)
+         .attr("viewBox", [0, 0, width, height] as any)
+         .style("background-color", modeColors.crumbs);
+
+      // plots
+      const xPlot = d3.scaleTime()
+         .domain(d3.extent(data, (d => d.date)) as any)
+         .range([this.margin.left, this.width - this.margin.right]);
+
+      const yPlot = d3.scaleLinear()
+         .domain([0, maxReservations])
+         .range([height - margin.bottom, margin.top]);
+
+      // bottom axis
+      const xAxis = svg.append("g")
+         .call(this.xAxisMethod as any, xPlot)
+         .style("pointer-events", "none");
+
+      // graph
+      const area = d3.area()
+         .x((d: any) => xPlot(d.date))
+         .y0(yPlot(0))
+         .y1((d: any) => yPlot(d.reservations));
+
+      const areaColor = dashboard.darktheme ?
+         (d3.color(modeColors.text) as any)?.darker(0.5)
+         : (d3.color(modeColors.text) as any)?.brighter(0.5)
+
+      svg.append("path")
+         .datum(data)
+         .attr("fill", areaColor)
+         .attr("stroke", modeColors.text).attr("stroke-width", 1.5)
+         .attr("opacity", 0.4)
+         .attr("d", area as any)
+         .style("pointer-events", "none");
+
+      // dragable timeline
+      const drag = (handle: any, update: (value: number) => void) => {
+         return d3.drag()
+            .on("start", (e: any) => { handle.attr("cursor", "grabbing") })
+            .on("end", (e: any) => { handle.attr("cursor", "grab") })
+            .on("drag", (e: any) => {
+               const x = Math.min(Math.max(e.x, xPlot.range()[0]), xPlot.range()[1]);
+               handle.attr("transform", `translate(${x},${yPlot(0)})`);
+               // update date
+               update(x);
+               // update time stamp
+               handle.selectAll("text").text(toLocaleTimeString(xPlot.invert(x)))
+               // update area selected
+               updateAreaSelected();
+               // trigger callback
+               onDateChange();
+            });
+      }
+
+      [this.startDate, this.endDate] = xPlot.domain();
+
+      const areaSelected = svg.append("g")
+         .style("pointer-events", "none");
+
+      const updateAreaSelected = () => {
+         areaSelected.selectAll("*").remove();
+         areaSelected.append("rect")
+            .attr("fill", modeColors.text)
+            .attr("opacity", 0.1)
+            .attr("x", Math.min(xPlot(this.startDate!), xPlot(this.endDate!)))
+            .attr("y", yPlot(maxReservations))
+            .attr("width", Math.abs(xPlot(this.endDate!) - xPlot(this.startDate!)))
+            .attr("height", yPlot(0) - yPlot(maxReservations));
+         areaSelected.append("line")
+            .style("stroke", modeColors.text).style("stroke-width", 3)
+            .attr("x1", xPlot(this.startDate!))
+            .attr("x2", xPlot(this.endDate!))
+            .attr("y1", yPlot(0))
+            .attr("y2", yPlot(0));
+      }
+
+      areaSelected.call(updateAreaSelected as any);
+
+      const toLocaleTimeString = (d: Date) => d.toLocaleTimeString(undefined, { hour: 'numeric', minute: "numeric", hour12: true });
+
+      const constructHandle = (handle, dateValue) => {
+         handle.append("circle")
+            .attr("r", 6)
+            .attr("fill", (d3.color(modeColors.text) as any).darker(0.6))
+            .attr("stroke", modeColors.text).attr("stroke-width", 1.5);
+         handle.append("text")
+            .text(toLocaleTimeString(dateValue))
+            .attr("text-anchor", "middle")
+            .attr("y", -11)
+            .style("fill", modeColors.text)
+            .style("font-family", "Horde Open Sans Regular")
+            .style("font-size", "11px")
+            .style("pointer-events", "none");
+      }
+
+      const startHandle = svg.append("g")
+         .attr("cursor", "grab")
+         .attr("transform", `translate(${xPlot(this.startDate!)},${yPlot(0)})`);
+      startHandle.call(
+            drag(startHandle, (v) => { this.startDate = xPlot.invert(v) }) as any
+         );
+      constructHandle(startHandle, this.startDate);
+
+      const endHandle = svg.append("g")
+         .attr("cursor", "grab")
+         .attr("transform", `translate(${xPlot(this.endDate!)},${yPlot(0)})`);
+      endHandle.call(
+            drag(endHandle, (v) => { this.endDate = xPlot.invert(v) }) as any
+         );
+      constructHandle(endHandle, this.endDate);
+   }
+
+   private clear() {
+      if (!!this.container) {
+         d3.select(this.container).selectAll("*").remove();
+      }
+   }
+
+   private get xAxisMethod() {
+      return (g: SelectionType, x: ScalarTime) => {
+         const tickSet: Set<string> = new Set();
+
+         g.attr("transform", `translate(0,${this.height - this.margin.bottom})`)
+            .style("font-family", "Horde Open Sans Regular")
+            .style("font-size", "11px")
+            .call(d3.axisBottom(x)
+               .ticks(8)
+               .tickFormat(d => {
+                  const v = (d as Date).toDateString();
+                  if (!tickSet.has(v)) {
+                     tickSet.add(v)
+                     return getHumanTime(d as Date)
+                  }
+
+                  return (d as Date).toLocaleString(undefined, { hour: 'numeric', minute: "numeric", hour12: true })
+               })
+               .tickSizeOuter(0))
+      }
+   }
+
+   handler: PoolTelemetryHandler;
+
+   container?: HTMLDivElement;
+   private margin: { top: number, right: number, bottom: number, left: number }
+   private width = 0;
+   private height = 0;
+
+   private data?: UsagePlatformTelemetry[];
+   private platformId = "";
+
+   startDate?: Date;
+   endDate?: Date;
+
+   private hasRendered = false;
+   forceRender = false;
+}
diff --git a/Engine/Source/Programs/Horde/HordeDashboard/src/components/devices/DeviceView.tsx b/Engine/Source/Programs/Horde/HordeDashboard/src/components/devices/DeviceView.tsx
new file mode 100644
--- /dev/null
+++ b/Engine/Source/Programs/Horde/HordeDashboard/src/components/devices/DeviceView.tsx
@@ -0,0 +1,1061 @@
+import { DefaultButton, DetailsList, DetailsListLayoutMode, Dropdown, FocusZone, FocusZoneDirection, IColumn, Icon, IconButton, IGroup, Label, Link as FluentLink, mergeStyleSets, Modal, Pivot, PivotItem, PrimaryButton, ScrollablePane, ScrollbarVisibility, SelectionMode, Spinner, SpinnerSize, Stack, Text, TextField, themeRulesStandardCreator, TooltipHost, DirectionalHint, } from "@fluentui/react";
+import { observer } from "mobx-react-lite";
+import { observable, action, makeObservable } from "mobx";
+import moment from "moment";
+import React, { useEffect, useState } from "react";
+import { Link, useNavigate, useSearchParams } from "react-router-dom";
+import backend from "../../backend";
+import { DevicePoolType, GetDeviceResponse } from "../../backend/Api";
+import dashboard from "../../backend/Dashboard";
+import { useWindowSize } from "../../base/utilities/hooks";
+import { getNiceTime } from "../../base/utilities/timeUtils";
+import { Breadcrumbs } from "../Breadcrumbs";
+import { DeviceEditor, DeviceHandler, DeviceStatus } from "./DeviceEditor";
+import { DeviceInfoModal } from "./DeviceInfoView";
+import { DevicePoolTelemetryModal } from "./DevicePoolTelemetry";
+import { TopNav } from "../TopNav";
+import { getHordeStyling } from "../../styles/Styles";
+import { getHordeTheme } from "../../styles/theme";
+
+const handler = new DeviceHandler();
+
+type DeviceItem = {
+   device: GetDeviceResponse;
+   // Whether user has this device checked out
+   userCheckout?: boolean;
+   status?: string;
+}
+
+let groups: IGroup[] = [];
+
+type DevicesSearchState = {
+   filterPools?: Array<string>;
+   filterPlatforms?: Array<string>;
+   pivotKey?: string;
+   filterString?: string;
+   historyItem?: string;
+}
+
+class LocalState {
+
+   @observable searchUpdated: number = 0;
+
+   constructor() {
+      makeObservable(this);
+   }
+
+   @action
+   setSearchUpdated() {
+      this.searchUpdated++;
+   }
+
+
+   searchState: DevicesSearchState = {};
+   search: URLSearchParams = new URLSearchParams(window.location.search);
+
+   updateSearch(): boolean {
+      
+      const state = { ...this.searchState } as DevicesSearchState;
+
+      state.filterPools = state.filterPools?.sort((a, b) => a.localeCompare(b));
+
+      const search = new URLSearchParams();
+      const csearch = this.search.toString();
+
+      state.filterPools?.forEach(f => {
+         if (f) {
+            search.append("pool", f);
+         }
+      });
+
+      state.filterPlatforms?.forEach(f => {
+         if (f) {
+            search.append("platform", f);
+         }
+      });
+
+      if (state.pivotKey?.length) {
+         search.append("pivotKey", state.pivotKey);
+      }
+
+      if (state.historyItem?.length) {
+         search.append("history-item", state.historyItem);
+      }
+
+      if (state.filterString?.length) {
+         search.append("filter", state.filterString);
+      }
+
+      if (search.toString() !== csearch) {
+         this.search = search;
+         this.setSearchUpdated();
+         return true;
+      }
+
+      return false;
+   }
+
+   resetState() {
+      this.initialized = false;
+      this.search = new URLSearchParams();
+      this.searchState = {};
+   }
+   
+   private stateFromSearch() {
+
+      const state: DevicesSearchState = {};
+
+      const pools = this.search.getAll("pool") ?? undefined;
+      const platforms = this.search.getAll("platform") ?? undefined;
+      const pivot = this.search.get("pivotKey") ?? undefined;
+      const filter = this.search.get("filter") ?? undefined;
+      const historyItem = this.search.get("history-item") ?? undefined;
+
+      state.filterPools = pools?.sort((a, b) => a.localeCompare(b));
+      state.filterPlatforms = platforms?.sort((a, b) => a.localeCompare(b));
+      state.pivotKey = pivot;
+      state.filterString = filter;
+      state.historyItem = historyItem;
+
+      this.searchState = state;
+
+      if (state.filterPools?.length) {
+         this.searchState.filterPools = state.filterPools;
+      }
+
+      if (state.filterPlatforms?.length) {
+         this.searchState.filterPlatforms = state.filterPlatforms;
+      }
+
+      if (state.pivotKey?.length) {
+         this.searchState.pivotKey = state.pivotKey;
+      }
+
+      if (state.filterString?.length) {
+         this.searchState.filterString = state.filterString;
+      }
+
+      return state;
+   }
+
+   @action
+   setFilterWithoutUrlUpdate(filter: string | undefined) {
+      if (filter) {
+         this.searchState.filterString = filter;
+      }
+   }
+
+   @action
+   setFilter(filter?: string) {
+      this.searchState.filterString = filter;
+      this.updateSearch();
+   }
+
+   @action
+   setPivotKey(pivotKey: string) {
+      this.searchState.filterPlatforms = undefined;
+      this.searchState.filterPools = undefined;
+      this.searchState.filterString = undefined;
+      this.searchState.pivotKey = pivotKey;
+      this.updateSearch();
+   }
+
+   @action
+   setPoolFilters(filter: Set<string>) {
+      const pools = Array.from(filter);
+      this.searchState.filterPools = pools.length ? pools : undefined;
+      this.updateSearch();
+   }
+
+   @action
+   setPlatformFilters(filter: Set<string>) {
+      const platforms = Array.from(filter);
+      this.searchState.filterPlatforms = platforms.length ? platforms : undefined;
+      this.updateSearch();
+   }
+
+   @action
+   setHistoryItem(item: string) {
+      this.searchState.historyItem = item;
+      this.updateSearch();
+   }
+
+   @action
+   resetHistoryItem() {
+      this.searchState.historyItem = undefined;
+      this.updateSearch();
+   }
+
+   initialize() 
+   {
+      if (!this.initialized) {
+         this.initialized = true;
+         this.stateFromSearch();
+      }
+   }
+
+   private initialized = false;
+}
+
+export const StatusNames = new Map<DeviceStatus, string>([
+   [DeviceStatus.Available, "Available"],
+   [DeviceStatus.Problem, "Problem"],
+   [DeviceStatus.Reserved, "Reserved"],
+   [DeviceStatus.Disabled, "Disabled"],
+   [DeviceStatus.Maintenance, "Maintenance"]
+]);
+
+const dropDownStyle: any = () => {
+
+   return {
+      dropdown: {},
+      callout: {
+         selectors: {
+            ".ms-Callout-main": {
+               padding: "4px 4px 12px 12px",
+               overflow: "hidden"
+            }
+         }
+      },
+      dropdownItemHeader: { fontSize: 12 },
+      dropdownOptionText: { fontSize: 12 },
+      dropdownItem: {
+         minHeight: 28, lineHeight: 28
+      },
+      dropdownItemSelected: {
+         minHeight: 28, lineHeight: 28, backgroundColor: "inherit"
+      }
+   }
+}
+
+
+const customStyles = mergeStyleSets({
+   raised: {
+      boxShadow: "0 1.6px 3.6px 0 rgba(0,0,0,0.132), 0 0.3px 0.9px 0 rgba(0,0,0,0.108)",
+      padding: "25px 30px 25px 30px"
+   },
+   details: {
+      selectors: {
+         '.ms-GroupHeader-title': {
+            cursor: "default"
+         },
+         '.ms-DetailsRow': {
+            animation: "none",
+            background: "unset"
+         }
+      },
+   }
+});
+
+const pivotKeyAutomation = "pivot-key-automation";
+const pivotKeyShared = "pivot-key-shared";
+const localState = new LocalState();
+
+export const SearchUpdate: React.FC = observer(() => {
+
+   const [, setSearchParams] = useSearchParams();
+
+   const csearch = localState.search.toString();
+
+   useEffect(() => {
+      setSearchParams(csearch, { replace: true });
+   }, [csearch, setSearchParams])
+
+   // subscribe
+   if (localState.searchUpdated) { }
+
+   return null;
+});
+
+const DevicePanel: React.FC = observer(() => {
+
+   const [editState, setEditState] = useState<{ shown?: boolean, infoShown?: boolean, editNote?: boolean, device?: DeviceItem | undefined }>({});
+   const [pivotState, setPivotState] = useState<{ key: string, poolFilter: Set<string>, platformFilter: Set<string> }>({ key: localState?.searchState?.pivotKey ?? pivotKeyShared, platformFilter: new Set(localState?.searchState?.filterPlatforms) ?? [], poolFilter: new Set(localState?.searchState?.filterPools ?? []) });
+   const [checkoutState, setCheckoutState] = useState<{ checkoutId?: string, checkinId?: string, showConfirm?: "in" | "out" | "error" }>({});
+   const [telemState, setTelemState] = useState(false);
+   
+   const navigate = useNavigate();
+
+   useEffect(() => {
+
+      handler.start();
+
+      return () => {
+         handler.clear();
+      };
+
+   }, []);
+
+   const { hordeClasses } = getHordeStyling();
+   const hordeTheme = getHordeTheme();
+
+   if (handler.updated) { }
+
+   let message = "";
+   let linkText = "device service configuration.";
+   let link = "/docs/Config/Devices.md";
+   if (handler.loaded) {
+
+      if (!handler.platforms.size) {
+         message = "There are no device platforms configured, please see";
+      } else if (!handler.pools.size) {
+         message = "There are no device pools configured, please see";
+      }
+   }
+
+   if (!handler.loaded) {
+      return <Stack horizontalAlign="center">
+         <Spinner size={SpinnerSize.large} />
+      </Stack>
+   }
+
+   if (message) {
+      return <Stack horizontal tokens={{ childrenGap: 6 }} horizontalAlign="center">
+         <Text variant="mediumPlus">{message}</Text>
+         {!!linkText && !!link && <a href={link} style={{ fontSize: "18px", "cursor": "pointer" }} onClick={(ev) => { ev.preventDefault(); ev.stopPropagation(); navigate(link) }}>{linkText}</a>}
+      </Stack>
+   }
+
+   const StatusColors = new Map<DeviceStatus, string>([
+      [DeviceStatus.Available, "#52C705"],
+      [DeviceStatus.Problem, "#DE4522"],
+      [DeviceStatus.Reserved, "#00BCF2"],
+      [DeviceStatus.Disabled, dashboard.darktheme ? "#D3D2D1" : "#F3F2F1"],
+      [DeviceStatus.Maintenance, "#0078D4"]
+   ]);
+
+
+   const InfoModal: React.FC = () => {
+
+      const checkout = !!checkoutState.checkoutId;
+
+
+      return <Modal isOpen={true} isBlocking={true} topOffsetFixed={true} styles={{ main: { padding: 8, width: 840, hasBeenOpened: false, top: "80px", position: "absolute" } }} className={hordeClasses.modal}>
+         <Stack tokens={{ childrenGap: 24 }} styles={{ root: { padding: 8 } }}>
+            <Stack grow verticalAlign="center">
+               <Text variant="mediumPlus" styles={{ root: { fontWeight: "unset", fontFamily: "Horde Open Sans SemiBold" } }}>{`Checking ${checkout ? "out" : "in"} Device`}</Text>
+            </Stack>
+            <Stack horizontalAlign="center">
+               <Text variant="mediumPlus">{`The device is being checked ${checkout ? "out" : "in"}, one moment please.`}</Text>
+            </Stack>
+            <Stack verticalAlign="center" style={{ paddingBottom: 32 }}>
+               <Spinner size={SpinnerSize.large} />
+            </Stack>
+         </Stack>
+      </Modal>
+   }
+
+   const automationTab = pivotState.key === pivotKeyAutomation;
+
+   const platforms = handler.platforms;
+
+   const poolWidth = automationTab ? 80 : 140;
+   const detailsWidth = automationTab ? 160 : 100;
+
+   const columns: IColumn[] = [
+      { key: 'column_platform', name: 'Platform', fieldName: 'platformId', minWidth: 72, maxWidth: 72, isResizable: false },
+      { key: 'column_pool', name: 'Pool', fieldName: 'poolId', minWidth: poolWidth, maxWidth: poolWidth, isResizable: false },
+      { key: 'column_model', name: 'Model', fieldName: 'modelId', minWidth: 80, maxWidth: 80, isResizable: false },
+      { key: 'column_name', name: 'Name', fieldName: 'name', minWidth: 120, maxWidth: 120, isResizable: false },
+      { key: 'column_address', name: 'Address', fieldName: 'address', minWidth: 120, maxWidth: 120, isResizable: false },
+      { key: 'column_status', name: 'Status', fieldName: 'status', minWidth: 120, maxWidth: 120, isResizable: false },
+      { key: 'column_details', name: 'Details', fieldName: 'details', minWidth: detailsWidth, maxWidth: detailsWidth, isResizable: false },
+      { key: 'column_notes', name: 'Notes', fieldName: 'notes', minWidth: 240, maxWidth: 240, isResizable: false },
+   ];
+
+   columns.forEach(c => {
+      c.isPadded = false;
+      c.onRender = ((item: DeviceItem, index, column) => {
+         return <Stack verticalFill={true} verticalAlign="center">
+            <Text>{(item.device as any)[column?.fieldName!]}</Text>
+         </Stack>
+
+      });
+   });
+
+   let column = columns.find(c => c.name === "Status")!
+   column.onRender = ((item: DeviceItem, index, column) => {
+
+      const write = handler.getDeviceWriteAccess(item.device);
+
+      if (!write) {
+         return null;
+      }
+
+      const device = item.device;
+      const status = handler.getDeviceStatus(device);
+
+      let color = status === DeviceStatus.Disabled ? "#777777 !important" : "#ffffff";
+      let backgroundColor = StatusColors.get(status)!
+      let text = StatusNames.get(status)!;
+
+
+      if (device.checkedOutByUserId) {
+         if (status === DeviceStatus.Available) {
+            text = "Checked Out";
+            backgroundColor = StatusColors.get(DeviceStatus.Reserved)!;
+         }
+      }
+
+      // @todo: support deleting reservations
+      return <Stack verticalFill={true} verticalAlign="center">
+         <PrimaryButton styles={{ root: { border: "0px", width: 84, height: 24, fontSize: 12, fontFamily: "Horde Open Sans SemiBold !important", selectors: { ".ms-Button-label": { color: color } }, backgroundColor: backgroundColor } }} text={text}
+            onClick={(ev) => {
+               ev.stopPropagation(); ev.preventDefault();
+               setEditState({ shown: true, device: item })
+            }}
+         />
+      </Stack>
+
+   })
+
+   column = columns.find(c => c.name === "Platform")!;
+   column.onRender = ((item: DeviceItem, index, column) => {
+
+      const platform = handler.platforms.get(item.device.platformId);
+
+      return <Stack verticalFill={true} verticalAlign="center">
+         <Text>{platform?.name ?? "???"}</Text>
+      </Stack>
+
+   })
+
+   column = columns.find(c => c.name === "Pool")!;
+   column.onRender = ((item: DeviceItem, index, column) => {
+
+      const pool = handler.pools.get(item.device.poolId);
+
+      return <Stack verticalFill={true} verticalAlign="center">
+         <Text>{pool?.name ?? "???"}</Text>
+      </Stack>
+
+   })
+
+   column = columns.find(c => c.name === "Model")!;
+   column.onRender = ((item: DeviceItem, index, column) => {
+
+
+      return <Stack verticalFill={true} verticalAlign="center">
+         <Text>{item.device.modelId ?? "Base"}</Text>
+      </Stack>
+
+   })
+
+   column = columns.find(c => c.name === "Details")!
+   column.onRender = ((item: DeviceItem, index, column) => {
+
+      const device = item.device;
+
+      // Check whether shared kit
+      if (!automationTab) {
+
+         const status = handler.getDeviceStatus(item.device);
+
+         const write = handler.getDeviceWriteAccess(item.device);
+
+         const checkOutDisabled = !write || status !== DeviceStatus.Available;
+
+         if (device.id === checkoutState.checkoutId) {
+            return <Stack verticalFill={true} verticalAlign="center" horizontalAlign="start"><Spinner size={SpinnerSize.small} /></Stack>;
+         }
+
+         if (!device.checkedOutByUserId) {
+
+            return <Stack verticalFill={true} verticalAlign="center">
+               {!checkOutDisabled && <PrimaryButton styles={{ root: { border: "0px", width: 90, height: 24, fontSize: 12, fontFamily: "Horde Open Sans SemiBold !important" } }} text="Check Out"
+                  onClick={async (ev) => {
+                     ev.stopPropagation(); ev.preventDefault();
+                     setCheckoutState({ checkoutId: device.id });
+                     await backend.checkoutDevice(device.id, true).then(() => {
+                        handler.updating = false;
+                        handler.update().then(() => {
+                           setCheckoutState({ showConfirm: "out" })
+                        }).catch((reason) => {
+                           console.error(reason);
+                           setCheckoutState({ showConfirm: "error" });
+                        })
+                     });
+                  }}
+               />}
+            </Stack>
+         }
+
+         const user = handler.users.get(device.checkedOutByUserId);
+
+         if (!user) {
+            return <Text>Error: Missing user {device.checkedOutByUserId} for checkout</Text>
+         }
+
+         if (user.id === dashboard.userId) {
+            return <Stack verticalFill={true} verticalAlign="center">
+               <PrimaryButton styles={{ root: { border: "0px", width: 90, height: 24, fontSize: 12, fontFamily: "Horde Open Sans SemiBold !important", backgroundColor: StatusColors.get(DeviceStatus.Reserved)! } }} text="Check In"
+                  onClick={async (ev) => {
+                     ev.stopPropagation(); ev.preventDefault();
+                     setCheckoutState({ checkinId: device.id });
+                     await backend.checkoutDevice(device.id, false).then(() => {
+                        handler.updating = false;
+                        handler.update().then(() => setCheckoutState(
+                           { showConfirm: handler.getUserDeviceCheckouts(dashboard.userId).length > 0 ? "in" : undefined }))
+                     }).catch((reason) => {
+                        console.error(reason);
+                        setCheckoutState({ showConfirm: "error" });
+                     });
+                  }}
+               />
+            </Stack>
+         }
+
+         return <Stack verticalFill={true} verticalAlign="center"><Text variant="tiny">{`${user.name}`}</Text></Stack>
+
+      }
+
+      // Automation Kit
+
+      const r = handler.getReservation(device);
+
+      let url = "";
+
+      if (r?.jobId) {
+         url = `/job/${r.jobId}`;
+         if (r.stepId) {
+            url += `?step=${r.stepId}`;
+         }
+      }
+      return <Stack verticalFill={true} verticalAlign="center">
+         <Stack horizontal tokens={{ childrenGap: 12 }}>
+            {automationTab && <Stack verticalFill={true} verticalAlign="center">
+               <TooltipHost content="History" delay={0} closeDelay={250}
+                  styles={{root: {display: 'inline-block'}}} directionalHint={DirectionalHint.leftCenter}>
+                  <div style={{ cursor: "pointer" }} onClick={(ev) => {
+                     ev.stopPropagation(); ev.preventDefault();
+                     setEditState({
+                        infoShown: true, device: item
+                     })
+                     localState.setHistoryItem(item.device.id);
+                  }}>
+                     <Stack horizontal tokens={{ childrenGap: 18 }}>
+                        <Icon style={{ padding: 2 }} iconName="History"/>
+                     </Stack>
+                  </div>
+               </TooltipHost>
+
+            </Stack>}
+            {!!r && <Stack>
+               {!!url && <Link to={url}><Text variant="small">Job Details</Text></Link>}
+               {!url && !!r.reservationDetails && r.reservationDetails.startsWith("https://horde") && <FluentLink href={r.reservationDetails} target="_blank" underline={false}><Text variant="small" onClick={(ev) => { ev.stopPropagation(); }}>Job Details</Text></FluentLink>}
+               <Text variant="tiny">Host: {r.hostname}</Text>
+               <Text variant="tiny">Duration: {Math.floor(moment.duration(moment.utc().diff(moment(r.createTimeUtc))).asMinutes())}m</Text>
+            </Stack>}
+         </Stack>
+      </Stack>
+
+   })
+
+   column = columns.find(c => c.name === "Notes")!
+   column.onRender = ((item: DeviceItem, index, column) => {
+
+      const device = item.device;
+      let notes = device.notes ?? "";
+
+      if (device.checkOutExpirationTime && device.checkedOutByUserId) {
+         notes = `Checked out until ${getNiceTime(device.checkOutExpirationTime)}.  ` + notes;
+      }
+
+      return <Stack verticalFill={true} verticalAlign="center">
+         <div style={{ cursor: "pointer" }} onClick={(ev) => {
+            ev.stopPropagation(); ev.preventDefault();
+            setEditState({
+               shown: true, device: item, editNote: true
+            })
+         }}>
+            <Stack horizontal tokens={{ childrenGap: 12 }}>
+               <Icon style={{ paddingTop: 2 }} iconName="Edit" />
+               <Text variant="small" style={{ whiteSpace: "pre-wrap" }}>{notes}</Text>
+            </Stack>
+         </div>
+      </Stack>
+   })
+
+   /**
+    *
+    * @returns List of devices filtered by pool and search string
+    */
+   const GetFilteredDeviceList = (filterString?: string): [GetDeviceResponse[], boolean] => {
+
+      const devices = handler.getDevices();
+
+      const exactMatch = devices.find(d => {
+         return d.name.toLowerCase() === filterString || d.address?.toLowerCase() === filterString;
+      })
+
+      if (exactMatch) {
+
+         return [[exactMatch], true];
+      }
+
+      let poolDevices = devices.filter((d) => {
+         const pool = handler.pools.get(d.poolId);
+         if (!pool) {
+            return false;
+         }
+
+         if (automationTab && pool.poolType !== DevicePoolType.Automation) {
+            return false;
+         }
+
+         if (!automationTab && pool.poolType === DevicePoolType.Automation) {
+            return false;
+         }
+
+         if (pivotState.poolFilter.size && !pivotState.poolFilter.has(d.poolId)) {
+            return false;
+         }
+
+         if (pivotState.platformFilter.size && !pivotState.platformFilter.has(d.platformId)) {
+            return false;
+         }
+
+         if (filterString) {
+            if (!(d.name.toLowerCase().includes(filterString)) && !(d.address?.toLowerCase().includes(filterString))) {
+               return false;
+            }
+         }
+
+         return true;
+
+      });
+      return [poolDevices, false];
+   }
+
+   /**
+    * Function to sort a supplied list of devices based on a specified sort column
+    * @param filteredDevices List of "GetDeviceResponse" objects used to create a sorted list of DeviceItems
+    * @param sortColumn Defines an optional sort column to use in addition to the primary 'platformId' and 'poolId' columns
+    * @returns Sorted list of Device Items
+    */
+   const GetSortedDeviceList = (filteredDevices: GetDeviceResponse[], sortColumn: undefined | string = undefined): DeviceItem[] => {
+      const devices = filteredDevices.sort((a, b) => {
+
+         if (a.platformId === b.platformId) {
+
+            if (a.poolId === b.poolId) {
+               if (sortColumn) { // If a sort column is defined, sort by that column
+                  if (sortColumn === "status") { // The status column is derived, so it isn't as simple as just comparing property values :(
+                     let statusCodeA = handler.getDeviceStatus(a);
+                     let statusCodeB = handler.getDeviceStatus(b);
+                     // "Checked Out" is defined as 'Available' AND checkedOutByUser
+                     let statusA = a.checkedOutByUserId && statusCodeA === DeviceStatus.Available ? 100 : statusCodeA;
+                     let statusB = b.checkedOutByUserId && statusCodeB === DeviceStatus.Available ? 100 : statusCodeB;
+                     if (statusA === statusB) {
+                        return 0; // Statuses are not unique, so we need to explicitly return the matching case
+                     }
+                     return statusA < statusB ? -1 : 1;
+                  }
+                  if (a[sortColumn] === b[sortColumn]) {
+                     return 0; // Names are always unique, but generic columns could have the same value, so this will account for those.
+                  }
+                  return (a[sortColumn] ?? "") < (b[sortColumn] ?? "") ? -1 : 1;
+               } else { // By default, innermost sort should be done by name
+                  return a.name < b.name ? -1 : 1;;
+               }
+            }
+
+            return a.poolId < b.poolId ? -1 : 1;
+         }
+
+         return a.platformId < b.platformId ? -1 : 1;
+
+      }).map(d => {
+         return { device: d } as DeviceItem;
+      });
+      return devices;
+   }
+
+   // Subscribe to get a new sorted/filtered device list as soon as the search is updated
+   if (localState.searchUpdated) { }
+
+   const [filteredDevices, exactMatch] = GetFilteredDeviceList(localState.searchState.filterString?.toLowerCase());
+
+   // Check for an exact match and switch pivot if necessary
+   if (exactMatch) {
+
+      const pool = handler.pools.get(filteredDevices[0].poolId);
+      const tabType = automationTab ? DevicePoolType.Automation : DevicePoolType.Shared;
+
+      if (pool && pool.poolType !== tabType) {
+         const key = pool.poolType === DevicePoolType.Automation ? pivotKeyAutomation : pivotKeyShared;
+         setPivotState({ ...pivotState, key: key });
+         return null;
+      }
+   }
+
+   const devices = GetSortedDeviceList(filteredDevices, !automationTab ? "status" : undefined);
+   const newGroups: IGroup[] = [];
+
+   let myDevices: DeviceItem[] = [];
+   if (pivotState.key === "pivot-key-shared") {
+      myDevices = devices.filter(d => d.device.checkedOutByUserId === dashboard.userId);
+      devices.unshift(...myDevices.map(d => { return { ...d, userCheckout: true } }));
+   }
+
+   let curPlatform: string | undefined;
+   devices.forEach((item, index) => {
+      const d = item.device;
+      if (item.userCheckout) {
+         if (!newGroups.length) {
+            const key = `user_checkout`;
+            newGroups.push({
+               key: key,
+               name: "My Checkouts",
+               startIndex: 0,
+               count: myDevices.length,
+               level: 0,
+               isCollapsed: groups.find(g => g.key === key)?.isCollapsed
+            });
+         }
+      }
+      else if (d.platformId !== curPlatform) {
+         curPlatform = d.platformId;
+         const key = `group_${d.platformId}`;
+         newGroups.push({
+            key: key,
+            name: platforms.get(d.platformId)!.name,
+            startIndex: index,
+            count: devices.filter(cd => cd.device.platformId === d.platformId).length,
+            level: 0,
+            isCollapsed: groups.find(g => g.key === key)?.isCollapsed
+         });
+      }
+   })
+
+   groups = newGroups;
+
+   const pivotItems: JSX.Element[] = [];
+
+   pivotItems.push(<PivotItem headerText="Shared" itemKey={pivotKeyShared} key={pivotKeyShared} />);
+   pivotItems.push(<PivotItem headerText="Automation" itemKey={pivotKeyAutomation} key={pivotKeyAutomation} />);
+
+
+   const poolItems = Array.from(handler.pools.values()).filter(pool => {
+
+      if (automationTab && pool.poolType !== DevicePoolType.Automation) {
+         return false;
+      }
+      if (!automationTab && pool.poolType === DevicePoolType.Automation) {
+         return false;
+      }
+
+      return true;
+
+   }).sort((a, b) => {
+      return a.name < b.name ? 1 : 0;
+   }).map(pool => {
+      const name = pool.name;
+      return {
+         text: name,
+         key: `dropdown_${pool.id}_key`,
+         poolId: pool.id
+      }
+   });
+
+   const platformItems = Array.from(platforms.values()).sort((a, b) => {
+      if (a.name === b.name) { return 0; }
+      return a.name < b.name ? 1 : -1;
+   }).map(platform => {
+      return {
+         text: platform.name,
+         key: `dropdown_${platform.id}_key`,
+         platformId: platform.id
+      }
+   })
+
+   const checkedOut = handler.getUserDeviceCheckouts(dashboard.userId);
+
+   const hasDevices = !!handler.getDevices().length
+
+   return (<Stack>
+      {!!telemState && <DevicePoolTelemetryModal poolsIn={localState.searchState.filterPools} platformsIn={localState.searchState.filterPlatforms} onClose={() => setTelemState(false)} />}
+      {(!!checkoutState.checkinId || !!checkoutState.checkoutId) && <InfoModal />}
+      {!!checkoutState.showConfirm && checkedOut.length > 0 && <CheckoutConfirmModal check={checkoutState.showConfirm} devices={checkedOut} onClose={() => setCheckoutState({})} />}
+      {editState.infoShown && <DeviceInfoModal handler={handler} deviceIn={editState.device?.device} onEdit={(device) => { setEditState({ infoShown: true, shown: true, device: { device: device } }) }} onClose={() => { setEditState({}); localState.resetHistoryItem(); }} />}
+      {editState.shown && <DeviceEditor handler={handler} deviceIn={editState.device?.device} editNote={editState.editNote} onClose={() => { setEditState({ ...editState, shown: false }); localState.resetHistoryItem(); }} />}
+      <Stack styles={{ root: { paddingLeft: 12, paddingRight: 12, width: "100%" } }} >
+         <Stack>
+            <Stack horizontal verticalAlign="center" style={{ paddingBottom: 12 }} tokens={{ childrenGap: 32 }}>
+               {hasDevices && <Stack>
+                  <Pivot className={hordeClasses.pivot}
+                     selectedKey={pivotState.key}
+                     linkSize="normal"
+                     linkFormat="links"
+                     onLinkClick={(item) => {
+                        localState.setPivotKey(item!.props.itemKey!);
+                        setPivotState({ key: item!.props.itemKey!, poolFilter: new Set(), platformFilter: new Set() })
+                     }}>
+                     {pivotItems}
+                  </Pivot>
+               </Stack>}
+               {hasDevices && <Stack>
+                  <TextField
+                     placeholder="Device Search"
+                     spellCheck={false}
+                     autoComplete="off"
+                     validateOnLoad={false}
+                     deferredValidationTime={1000}
+                     defaultValue={localState.searchState.filterString}
+                     styles={{
+                        root: { width: 280, fontSize: 12 }, fieldGroup: {
+                           borderWidth: 1
+                        }
+                     }}
+                     onKeyUp={(evt) => {
+                        if (evt.key === "Enter") {
+                           localState.setFilter(evt.currentTarget.value || undefined)
+                        }
+                     }}
+                     onGetErrorMessage={(newValue) => {
+                        localState.setFilter(newValue);
+                        return undefined;
+                     }}
+                  />
+               </Stack>}
+               {hasDevices && platforms.size > 1 && <Stack style={{ paddingLeft: 15 }}>
+                  <Dropdown
+                     placeholder="Filter Platforms"
+                     style={{ width: 200 }}
+                     styles={dropDownStyle}
+                     selectedKeys={Array.from(pivotState.platformFilter).map(id => `dropdown_${id}_key`)}
+                     multiSelect
+                     options={platformItems}
+                     onChange={(event, option, index) => {
+
+                        if (option) {
+                           if (option.selected) {
+                              pivotState.platformFilter.add((option as any).platformId);
+                           } else {
+                              pivotState.platformFilter.delete((option as any).platformId);
+                           }
+
+                           localState.setPlatformFilters(pivotState.platformFilter);
+                           setPivotState({...pivotState, platformFilter: new Set(pivotState.platformFilter)})
+                        }
+                     }}
+                  />
+               </Stack>}
+               {hasDevices && poolItems.length > 1 && <Stack style={{ paddingLeft: 15 }}>
+                  <Dropdown
+                     placeholder="Filter Pools"
+                     style={{ width: 200 }}
+                     styles={dropDownStyle}
+                     selectedKeys={Array.from(pivotState.poolFilter).map(id => `dropdown_${id}_key`)}
+                     multiSelect
+                     options={poolItems}
+                     onChange={(event, option, index) => {
+
+                        if (option) {
+
+                           if (option.selected) {
+                              pivotState.poolFilter.add((option as any).poolId);
+                           } else {
+                              pivotState.poolFilter.delete((option as any).poolId);
+                           }
+
+                           localState.setPoolFilters(pivotState.poolFilter);
+                           setPivotState({ ...pivotState })
+
+                        }
+                     }}
+                  />
+               </Stack>}
+
+               {hasDevices && automationTab && <Stack>
+                  <DefaultButton text="Pool Telemetry" onClick={() => { setTelemState(true) }} />
+               </Stack>}
+
+               <Stack grow />
+               <Stack>
+                  <PrimaryButton text="Add Device"
+                     styles={{ root: { fontFamily: "Horde Open Sans SemiBold !important" } }}
+                     onClick={(ev) => { ev.stopPropagation(); ev.preventDefault(); setEditState({ shown: true }) }}
+                  />
+               </Stack>
+
+            </Stack>
+
+
+         </Stack>
+
+         {!handler.getDevices().length && handler.loaded && <Stack horizontalAlign="center">
+            <Text variant="mediumPlus">No Devices Found</Text>
+         </Stack>}
+
+
+         {!!handler.getDevices().length && <Stack tokens={{ childrenGap: 12 }}>
+            <FocusZone direction={FocusZoneDirection.vertical}>
+               <div className={customStyles.details} style={{ height: "calc(100vh - 280px)", position: 'relative' }} data-is-scrollable>
+                  <ScrollablePane scrollbarVisibility={ScrollbarVisibility.always} onScroll={() => { }}>
+                     <DetailsList
+                        styles={{
+                           root: {
+                              overflowX: "hidden", width: 1312, selectors: {
+                                 '.ms-GroupHeader,.ms-GroupHeader:hover': {
+                                    background: dashboard.darktheme ? hordeTheme.horde.dividerColor : "#DFDEDD",
+                                 },
+                                 '.ms-GroupHeader-expand,.ms-GroupHeader-expand:hover': {
+                                    background: dashboard.darktheme ? hordeTheme.horde.dividerColor : "#DFDEDD",
+                                 }
+                              }
+                           }
+                        }}
+                        items={devices}
+                        groups={groups}
+                        columns={columns}
+                        selectionMode={SelectionMode.none}
+                        layoutMode={DetailsListLayoutMode.justified}
+                        compact={true}
+                        onShouldVirtualize={() => false}
+                     />
+                  </ScrollablePane>
+               </div>
+            </FocusZone>
+         </Stack>}
+      </Stack>
+   </Stack>
+   );
+});
+
+export const CheckoutConfirmModal: React.FC<{ check: "in" | "out" | "error", devices: GetDeviceResponse[], onClose: () => void }> = ({ check, devices, onClose }) => {
+
+   const { hordeClasses } = getHordeStyling();
+
+   type CheckoutItem = {
+      device: GetDeviceResponse;
+   }
+
+   const close = () => {
+      onClose();
+   };
+
+   const columns = [
+      { key: 'column1', name: 'Name', fieldName: 'name', minWidth: 240, maxWidth: 240 },
+      { key: 'column2', name: 'Date', fieldName: 'value', minWidth: 500, maxWidth: 500 },
+   ];
+
+   const deviceItems: CheckoutItem[] = devices.map(d => { return { device: d } });
+
+   const onRenderItemColumn = (item: CheckoutItem, index?: number, columnIn?: IColumn) => {
+
+      const column = columnIn!;
+
+      const device = item.device;
+
+      const platform = handler.platforms.get(device.platformId)?.name;
+
+      // simple cases
+      switch (column.name) {
+         case 'Name':
+            return <Text variant="medium" styles={{ root: { fontFamily: "Horde Open Sans SemiBold" } }}>{platform} : {device.name} : {device.address ?? ""} </Text>
+         case 'Date':
+            if (device.checkOutExpirationTime) {
+               return <Text variant="medium">{`Checked out until ${getNiceTime(device.checkOutExpirationTime)}`}</Text>
+            }
+            return null;
+      }
+
+      return null;
+   }
+
+   let text = check === "in" ? "Device checked in, you have the following devices checked out:" : "Device checked out, you have the following devices checked out:";
+   if (check === "error") {
+      text = "There was an error processing request, please try again.";
+   }
+
+
+   return <Modal isOpen={true} isBlocking={true} topOffsetFixed={true} styles={{ main: { padding: 8, width: 840, hasBeenOpened: false, top: "80px", position: "absolute" } }} className={hordeClasses.modal} onDismiss={() => { close() }}>
+      <Stack tokens={{ childrenGap: 12 }}>
+         <Stack horizontal styles={{ root: { paddingLeft: 12 } }}>
+            <Stack grow style={{ paddingTop: 12, paddingLeft: 8 }}>
+               <Label style={{ fontSize: 18, fontWeight: 400, fontFamily: "Horde Open Sans SemiBold" }}>Checked Out Devices</Label>
+            </Stack>
+            <Stack style={{ paddingRight: 4, paddingTop: 4 }}>
+               <IconButton
+                  iconProps={{ iconName: 'Cancel' }}
+                  ariaLabel="Close popup modal"
+                  onClick={() => { close(); }}
+               />
+            </Stack>
+         </Stack>
+
+         <Stack tokens={{ childrenGap: 32 }} style={{ paddingLeft: 8 }}>
+            <Stack style={{ paddingLeft: 18 }}>
+               <Stack>
+                  <Text variant="medium">{text}</Text>
+               </Stack>
+            </Stack>
+            <Stack horizontal>
+               <Stack style={{ minWidth: 700, paddingLeft: 18, paddingRight: 18 }}>
+                  {check !== "error" && <DetailsList
+                     items={deviceItems}
+                     columns={columns}
+                     setKey="set"
+                     layoutMode={DetailsListLayoutMode.justified}
+                     isHeaderVisible={false}
+                     selectionMode={SelectionMode.none}
+                     onRenderItemColumn={onRenderItemColumn}
+                  />}
+               </Stack>
+            </Stack>
+
+         </Stack>
+
+         <Stack styles={{ root: { padding: 12 } }}>
+            <Stack horizontal tokens={{ childrenGap: 16 }} styles={{ root: { paddingTop: 12, paddingLeft: 8, paddingBottom: 8 } }}>
+               <Stack grow />
+               <PrimaryButton text="Ok" disabled={false} onClick={() => { close(); }} />
+            </Stack>
+         </Stack>
+      </Stack>
+   </Modal>;
+
+};
+
+
+export const DeviceView: React.FC = () => {
+
+   const windowSize = useWindowSize();
+   const vw = Math.max(document.documentElement.clientWidth, window.innerWidth || 0);
+
+   const { hordeClasses, modeColors } = getHordeStyling();
+
+   localState.initialize();
+
+   useEffect(() => {
+
+      return () => {
+         localState.resetState();
+      };
+
+   }, []);
+
+
+
+   return <Stack className={hordeClasses.horde}>
+      <TopNav />
+      <Breadcrumbs items={[{ text: 'Devices' }]} />
+      <Stack horizontal>
+         <div key={`windowsize_streamview_${windowSize.width}_${windowSize.height}`} style={{ width: vw / 2 - (1440 / 2), flexShrink: 0, backgroundColor: modeColors.background }} />
+         <Stack tokens={{ childrenGap: 0 }} styles={{ root: { backgroundColor: modeColors.background, width: "100%" } }}>
+            <Stack style={{ maxWidth: 1440, paddingTop: 30, marginLeft: 4, height: 'calc(100vh - 8px)' }}>
+               <Stack horizontal className={hordeClasses.raised}>
+                  <Stack style={{ width: "100%", height: 'calc(100vh - 228px)' }} tokens={{ childrenGap: 18 }}>
+                     <DevicePanel />
+                     <SearchUpdate />
+                  </Stack>
+               </Stack>
+            </Stack>
+         </Stack>
+      </Stack>
+   </Stack>
+
+};
+
+
