CIDR Blocks requires inclusive checks

diff --git a/Engine/Source/Programs/Horde/HordeServer.Tests/Configuration/ConfigTests.cs b/Engine/Source/Programs/Horde/HordeServer.Tests/Configuration/ConfigTests.cs
--- a/Engine/Source/Programs/Horde/HordeServer.Tests/Configuration/ConfigTests.cs
+++ b/Engine/Source/Programs/Horde/HordeServer.Tests/Configuration/ConfigTests.cs
@@ -139,26 +139,28 @@ namespace HordeServer.Tests.Configuration
 			{
 				Networks = new List<NetworkConfig>()
 				{
-					new() { CidrBlock = "10.0.0.0/30", Id = "foo" },
+					new() { CidrBlock = "10.0.0.0/31", Id = "foo" },
 					new() { CidrBlock = "10.0.0.4/30", Id = "bar" },
 					new() { CidrBlock = "192.168.0.0/16", Id = "baz" },
 				}
 			};
 
-			Assert.AreEqual(null, GetNetworkConfig(gc, "10.0.0.0"));
+			Assert.AreEqual("foo", GetNetworkConfig(gc, "10.0.0.0")!.Id);
 			Assert.AreEqual("foo", GetNetworkConfig(gc, "10.0.0.1")!.Id);
-			Assert.AreEqual("foo", GetNetworkConfig(gc, "10.0.0.2")!.Id);
-			Assert.AreEqual("foo", GetNetworkConfig(gc, "10.0.0.3")!.Id);
+			Assert.AreEqual(null, GetNetworkConfig(gc, "10.0.0.2"));
+			Assert.AreEqual(null, GetNetworkConfig(gc, "10.0.0.3"));
 
-			Assert.AreEqual(null, GetNetworkConfig(gc, "10.0.0.4"));
+			Assert.AreEqual("bar", GetNetworkConfig(gc, "10.0.0.4")!.Id);
 			Assert.AreEqual("bar", GetNetworkConfig(gc, "10.0.0.5")!.Id);
 			Assert.AreEqual("bar", GetNetworkConfig(gc, "10.0.0.6")!.Id);
 			Assert.AreEqual("bar", GetNetworkConfig(gc, "10.0.0.7")!.Id);
-			Assert.AreEqual(null, GetNetworkConfig(gc, "10.0.0.4"));
+			Assert.AreEqual(null, GetNetworkConfig(gc, "10.0.0.8"));
 
-			Assert.AreEqual(null, GetNetworkConfig(gc, "192.168.0.0"));
+			Assert.AreEqual("baz", GetNetworkConfig(gc, "192.168.0.0")!.Id);
 			Assert.AreEqual("baz", GetNetworkConfig(gc, "192.168.0.1")!.Id);
 			Assert.AreEqual("baz", GetNetworkConfig(gc, "192.168.255.254")!.Id);
+			Assert.AreEqual("baz", GetNetworkConfig(gc, "192.168.255.255")!.Id);
+			Assert.AreEqual(null, GetNetworkConfig(gc, "192.169.0.0"));
 			Assert.AreEqual(null, GetNetworkConfig(gc, "192.169.0.1"));
 
 			Assert.AreEqual(null, GetNetworkConfig(gc, "11.0.0.1"));
diff --git a/Engine/Source/Programs/Horde/Plugins/Compute/HordeServer.Compute/ComputeConfig.cs b/Engine/Source/Programs/Horde/Plugins/Compute/HordeServer.Compute/ComputeConfig.cs
--- a/Engine/Source/Programs/Horde/Plugins/Compute/HordeServer.Compute/ComputeConfig.cs
+++ b/Engine/Source/Programs/Horde/Plugins/Compute/HordeServer.Compute/ComputeConfig.cs
@@ -155,7 +155,7 @@ namespace HordeServer
 			byte[] ipBytes = ip.GetAddressBytes();
 			Array.Reverse(ipBytes);
 			uint ipUint = BitConverter.ToUInt32(ipBytes, 0);
-			return ipUint > ipRangeStart && ipUint <= ipRangeEnd;
+			return ipUint >= ipRangeStart && ipUint <= ipRangeEnd;
 		}
 
 		/// <summary>
