Horde: Update release notes
#rnx
#skipci

diff --git a/Engine/Source/Programs/Horde/Docs/ReleaseNotes.md b/Engine/Source/Programs/Horde/Docs/ReleaseNotes.md
--- a/Engine/Source/Programs/Horde/Docs/ReleaseNotes.md
+++ b/Engine/Source/Programs/Horde/Docs/ReleaseNotes.md
@@ -1,5 +1,21 @@
 # Release Notes
 
+## 2024-12-17
+
+* Pass through OIDC scopes for API authorization.  Allows requesting custom delegated scopes for tools using OidcToken, like UGS interacting with Horde.
+* EpicGames.Core: Add extension method for redacting sensitive values in request URIs (38781177)
+* Added a flag to allow agents to auto restart on failure.  When installing the agent via HordeAgent service install, you can specify -AutoRestartService=true/false, defaults to true. It will set the three failure recovery options to Restart the Service (38780557)
+* Use consts for LogicalCores and PhysicalCores values (38779497)
+* Collect system telemetry on Linux (CPU/mem/disk usage) (38779426)
+* Support multiple templates in job redirector queries (38755881)
+* Make job redirector default to latest step regardless of state, with a complete query param option (38754699)
+* Add job/step information to device issue audit logging (38754008)
+* Cloud DDC - Added missing documentation on HttpContent extensions (38751527)
+* Cloud DDC - Fixed issue were put blob responses did not correctly support being converted to compact binary. This involves avoiding anon types so that we can flag the fields for cb serialization. Also moved the tooling to read compact binary from a http client from Jupiter into EpicGames.Serialization for easier access. (38748909)
+* Lower timeout for ComputeSocket send buffer detach (38748804)
+* Analytics documentation update to consider new plugins setup. (38731447)
+* Cleanup ResolveWarning API to remove buildconfiguration delegate, and instead apply target defaults in new API. (38674255)
+
 ## 2024-12-09
 
 * Disable low cardinality indexes for Storage.Blobs collection (38662017)
