Test to see if Content Metrics aggregate with Count rather than Sum or Average

diff --git a/Engine/Source/Programs/Horde/HordeServer/Defaults/default-metrics.telemetry.json b/Engine/Source/Programs/Horde/HordeServer/Defaults/default-metrics.telemetry.json
--- a/Engine/Source/Programs/Horde/HordeServer/Defaults/default-metrics.telemetry.json
+++ b/Engine/Source/Programs/Horde/HordeServer/Defaults/default-metrics.telemetry.json
@@ -543,7 +543,7 @@
 			"filter": "$[?(@.Payload.EventName == 'Core.VirtualAssets.MountPoint')]",
 			"groupBy": "$.Payload.Build_BranchName",
 			"property": "$.Payload.TotalFileSize",
-			"function": "Sum",
+			"function": "Count",
 			"interval": "00:05:00"
 		},
 		{
@@ -551,7 +551,7 @@
 			"filter": "$[?(@.Payload.EventName == 'Core.VirtualAssets.MountPoint')]",
 			"groupBy": "$.Payload.Build_BranchName,$.Payload.Name",
 			"property": "$.Payload.TotalFileSize",
-			"function": "Sum",
+			"function": "Count",
 			"interval": "00:05:00"
 		},
 		{
@@ -559,7 +559,7 @@
 			"filter": "$[?(@.Payload.EventName == 'Core.VirtualAssets.MountPoint')]",
 			"groupBy": "$.Payload.Build_BranchName",
 			"property": "$.Payload.VirtualizedPercent",
-			"function": "Sum",
+			"function": "Count",
 			"interval": "00:05:00"
 		},
 		{
diff --git a/Engine/Source/Programs/Horde/HordeServer/Defaults/default-views.telemetry.json b/Engine/Source/Programs/Horde/HordeServer/Defaults/default-views.telemetry.json
--- a/Engine/Source/Programs/Horde/HordeServer/Defaults/default-views.telemetry.json
+++ b/Engine/Source/Programs/Horde/HordeServer/Defaults/default-views.telemetry.json
@@ -574,6 +574,15 @@
 									"id": "core-content-virtualization-mp"
 								}
 							]
+						},
+						{
+							"name": "Content Event Count",
+							"display": "Value",
+							"metrics": [
+								{
+									"id": "core-content-events"
+								}
+							]
 						}
 					]
 				},
