Horde: Test changing order on contents metrics
#rnx
#skipci

diff --git a/Engine/Source/Programs/Horde/HordeServer/Defaults/default-metrics.telemetry.json b/Engine/Source/Programs/Horde/HordeServer/Defaults/default-metrics.telemetry.json
--- a/Engine/Source/Programs/Horde/HordeServer/Defaults/default-metrics.telemetry.json
+++ b/Engine/Source/Programs/Horde/HordeServer/Defaults/default-metrics.telemetry.json
@@ -531,18 +531,18 @@
 			"topN": 25
 		},
 		{
-			"id": "core-content-events",
+			"id": "core-content-totalsize",
 			"filter": "$[?(@.Payload.EventName == 'Core.VirtualAssets.MountPoint')]",
-			"groupBy": "$.Payload.Build_BranchName,$.Payload.Session_Label,$.Payload.ProjectName",
-			"property": "$.Payload.SchemaVersion",
+			"groupBy": "$.Payload.Build_BranchName",
+			"property": "$.Payload.TotalFileSize",
 			"function": "Count",
 			"interval": "00:05:00"
-		},
+		},		
 		{
-			"id": "core-content-totalsize",
+			"id": "core-content-events",
 			"filter": "$[?(@.Payload.EventName == 'Core.VirtualAssets.MountPoint')]",
-			"groupBy": "$.Payload.Build_BranchName",
-			"property": "$.Payload.TotalFileSize",
+			"groupBy": "$.Payload.Build_BranchName,$.Payload.Session_Label,$.Payload.ProjectName",
+			"property": "$.Payload.SchemaVersion",
 			"function": "Count",
 			"interval": "00:05:00"
 		},
