Horde: Auto convert string property values to numerics when possible, warn when cannot convert
#rnx

diff --git a/Engine/Source/Programs/Horde/Plugins/Analytics/HordeServer.Analytics/Telemetry/Metrics/MetricCollection.cs b/Engine/Source/Programs/Horde/Plugins/Analytics/HordeServer.Analytics/Telemetry/Metrics/MetricCollection.cs
--- a/Engine/Source/Programs/Horde/Plugins/Analytics/HordeServer.Analytics/Telemetry/Metrics/MetricCollection.cs
+++ b/Engine/Source/Programs/Horde/Plugins/Analytics/HordeServer.Analytics/Telemetry/Metrics/MetricCollection.cs
@@ -162,7 +162,6 @@ namespace HordeServer.Telemetry.Metrics
 
 				if (result.Matches != null && result.Matches.Count > 0)
 				{
-
 					foreach (Json.Path.Node match in result.Matches)
 					{
 						JsonValue? value = match.Value as JsonValue;
@@ -170,6 +169,21 @@ namespace HordeServer.Telemetry.Metrics
 						{
 							values.Add(doubleValue);
 						}
+						else if (value != null && value.TryGetValue(out string? stringValue))
+						{
+							if (Double.TryParse(stringValue!, out double numericValue))
+							{
+								values.Add(numericValue);
+							}
+							else
+							{
+								_logger.LogWarning("Unabled to parse string to double for metric {MetricId}", metric.Id);
+							}
+						}
+						else
+						{
+							_logger.LogWarning("Value for property is not a string or number for metric {MetricId}", metric.Id);
+						}
 					}
 				}
 			}
