Bringing the content metrics back to Sum and Average aggregation.

diff --git a/Engine/Source/Programs/Horde/HordeServer/Defaults/default-metrics.telemetry.json b/Engine/Source/Programs/Horde/HordeServer/Defaults/default-metrics.telemetry.json
--- a/Engine/Source/Programs/Horde/HordeServer/Defaults/default-metrics.telemetry.json
+++ b/Engine/Source/Programs/Horde/HordeServer/Defaults/default-metrics.telemetry.json
@@ -543,7 +543,7 @@
 			"filter": "$[?(@.Payload.EventName == 'Core.VirtualAssets.MountPoint')]",
 			"groupBy": "$.Payload.Build_BranchName,$.Payload.Session_Label,$.Payload.ProjectName",
 			"property": "$.Payload.SchemaVersion",
-			"function": "Count",
+			"function": "Sum",
 			"interval": "00:05:00"
 		},
 		{
@@ -551,7 +551,7 @@
 			"filter": "$[?(@.Payload.EventName == 'Core.VirtualAssets.MountPoint')]",
 			"groupBy": "$.Payload.Build_BranchName,$.Payload.Name",
 			"property": "$.Payload.TotalFileSize",
-			"function": "Count",
+			"function": "Sum",
 			"interval": "00:05:00"
 		},
 		{
@@ -559,7 +559,7 @@
 			"filter": "$[?(@.Payload.EventName == 'Core.VirtualAssets.MountPoint')]",
 			"groupBy": "$.Payload.Build_BranchName",
 			"property": "$.Payload.VirtualizedPercent",
-			"function": "Count",
+			"function": "Average",
 			"interval": "00:05:00"
 		},
 		{
