Tweaks to Content metrics and dashboard

diff --git a/Engine/Source/Programs/Horde/HordeServer/Defaults/default-metrics.telemetry.json b/Engine/Source/Programs/Horde/HordeServer/Defaults/default-metrics.telemetry.json
--- a/Engine/Source/Programs/Horde/HordeServer/Defaults/default-metrics.telemetry.json
+++ b/Engine/Source/Programs/Horde/HordeServer/Defaults/default-metrics.telemetry.json
@@ -552,7 +552,8 @@
 			"groupBy": "$.Payload.Build_BranchName,$.Payload.Name",
 			"property": "$.Payload.TotalFileSize",
 			"function": "Sum",
-			"interval": "00:05:00"
+			"interval": "00:05:00",
+			"topN": 25
 		},
 		{
 			"id": "core-content-virtualization",
@@ -560,7 +561,7 @@
 			"groupBy": "$.Payload.Build_BranchName",
 			"property": "$.Payload.VirtualizedPercent",
 			"function": "Average",
-			"interval": "00:05:00"
+			"interval": "00:05:00",
 		},
 		{
 			"id": "core-content-virtualization-mp",
@@ -568,7 +569,8 @@
 			"groupBy": "$.Payload.Build_BranchName,$.Payload.Name",
 			"property": "$.Payload.VirtualizedPercent",
 			"function": "Average",
-			"interval": "00:05:00"
+			"interval": "00:05:00",
+			"bottomN": 25
 		},
 		{
 			"id": "horde-incremental-build-jobs",
diff --git a/Engine/Source/Programs/Horde/HordeServer/Defaults/default-views.telemetry.json b/Engine/Source/Programs/Horde/HordeServer/Defaults/default-views.telemetry.json
--- a/Engine/Source/Programs/Horde/HordeServer/Defaults/default-views.telemetry.json
+++ b/Engine/Source/Programs/Horde/HordeServer/Defaults/default-views.telemetry.json
@@ -548,15 +548,6 @@
 								}
 							]
 						},
-						{
-							"name": "Virtualization",
-							"display": "Value",
-							"metrics": [
-								{
-									"id": "core-content-virtualization"
-								}
-							]
-						},
 						{
 							"name": "Content Size By Mount Point",
 							"display": "Value",
@@ -567,20 +558,20 @@
 							]
 						},
 						{
-							"name": "Virtualization By Mount Point",
+							"name": "Virtualization",
 							"display": "Value",
 							"metrics": [
 								{
-									"id": "core-content-virtualization-mp"
+									"id": "core-content-virtualization"
 								}
 							]
 						},
 						{
-							"name": "Content Event Count",
+							"name": "Virtualization By Mount Point",
 							"display": "Value",
 							"metrics": [
 								{
-									"id": "core-content-events"
+									"id": "core-content-virtualization-mp"
 								}
 							]
 						}
