Horde: Updating schema documentation

#rnx

diff --git a/Engine/Source/Programs/Horde/Docs/Deployment/ServerSettings.md b/Engine/Source/Programs/Horde/Docs/Deployment/ServerSettings.md
--- a/Engine/Source/Programs/Horde/Docs/Deployment/ServerSettings.md
+++ b/Engine/Source/Programs/Horde/Docs/Deployment/ServerSettings.md
@@ -36,6 +36,7 @@ Name | Description
 `oidcClaimEmailMapping` | `string[]`<br>List of fields in /userinfo endpoint to try map to the standard email claim (see System.Security.Claims.ClaimTypes.Email)
 `oidcClaimHordeUserMapping` | `string[]`<br>List of fields in /userinfo endpoint to try map to the Horde user claim (see HordeClaimTypes.User)
 `oidcClaimHordePerforceUserMapping` | `string[]`<br>List of fields in /userinfo endpoint to try map to the Horde Perforce user claim (see HordeClaimTypes.PerforceUser)
+`oidcApiRequestedScopes` | `string[]`<br>API scopes to request when acquiring OIDC access tokens
 `oidcAddDefaultScopesAndMappings` | `boolean`<br>Add common scopes and mappings to above OIDC config fields Provided as a workaround since .NET config will only *merge* array entries when combining multiple config sources. Due to this unwanted behavior, having hard-coded defaults makes such fields unchangeable. See [REDACTED_PROFILE_URL]/runtime/issues/36569
 `serverUrl` | `string`<br>Base URL this Horde server is accessible from For example https://horde.mystudio.com If not set, a default is used based on current hostname. In more advanced setups where a reverse proxy is present in front of Horde, this must manually be set.
 `jwtIssuer` | `string`<br>Name of the issuer in bearer tokens from the server
