Horde: Send final capabilities update before agent session termination

Previously, capabilities were only refreshed every 5 minutes. This could lead to stale disk usage values being reported if an agent terminated shortly after disk cleanup operations.

diff --git a/Engine/Source/Programs/Horde/HordeAgent.Tests/FakeHordeServer.cs b/Engine/Source/Programs/Horde/HordeAgent.Tests/FakeHordeServer.cs
--- a/Engine/Source/Programs/Horde/HordeAgent.Tests/FakeHordeServer.cs
+++ b/Engine/Source/Programs/Horde/HordeAgent.Tests/FakeHordeServer.cs
@@ -71,7 +71,7 @@ internal class FakeHordeClient(FakeHordeRpcServer server) : IHordeClient
 internal class FakeHordeRpcServer
 {
 	public readonly TaskCompletionSource<bool> CreateSessionReceived = new();
-	public readonly TaskCompletionSource<bool> UpdateSessionReceived = new();
+	public readonly Channel<RpcUpdateSessionRequest> UpdateSessionRequests = Channel.CreateUnbounded<RpcUpdateSessionRequest>();
 	
 	/// <summary>
 	/// Last status reported by the agent during a session update
@@ -147,7 +147,7 @@ internal class FakeHordeRpcServer
 		
 		async Task OnRequest(RpcUpdateSessionRequest request)
 		{
-			UpdateSessionReceived.TrySetResult(true);
+			await UpdateSessionRequests.Writer.WriteAsync(request, cancellationToken);
 			LastReportedStatus = request.Status;
 			
 			switch (request.Status)
diff --git a/Engine/Source/Programs/Horde/HordeAgent.Tests/Leases/LeaseManagerTests.cs b/Engine/Source/Programs/Horde/HordeAgent.Tests/Leases/LeaseManagerTests.cs
--- a/Engine/Source/Programs/Horde/HordeAgent.Tests/Leases/LeaseManagerTests.cs
+++ b/Engine/Source/Programs/Horde/HordeAgent.Tests/Leases/LeaseManagerTests.cs
@@ -4,6 +4,7 @@ using System;
 using System.Collections.Generic;
 using System.Diagnostics.CodeAnalysis;
 using System.IO;
+using System.Linq;
 using System.Threading;
 using System.Threading.Tasks;
 using EpicGames.Core;
@@ -24,11 +25,12 @@ using OpenTelemetry.Trace;
 
 namespace HordeAgent.Tests.Leases;
 
-public class CapabilitiesServiceStub(RpcAgentCapabilities? capabilities = null) : ICapabilitiesService
+public class FakeCapabilitiesService : ICapabilitiesService
 {
+	public RpcAgentCapabilities Capabilities { get; set; } = new ();
 	public Task<RpcAgentCapabilities> GetCapabilitiesAsync(DirectoryReference? workingDir)
 	{
-		return Task.FromResult(capabilities ?? new RpcAgentCapabilities());
+		return Task.FromResult(Capabilities);
 	}
 }
 
@@ -51,17 +53,23 @@ public class SessionStub(AgentId agentId, SessionId sessionId, DirectoryReferenc
 
 internal class TestHandlerFactory : LeaseHandlerFactory<TestTask>
 {
+	private readonly LeaseResult _result = new((byte[]?)null);
+	
+	public TestHandlerFactory(LeaseResult? leaseResult = null)
+	{
+		_result = leaseResult ?? _result;
+	}
+	
 	public override LeaseHandler<TestTask> CreateHandler(RpcLease lease)
 	{
-		return new Handler(lease);
+		return new Handler(lease, _result);
 	}
 	
-	private class Handler(RpcLease rpcLease) : LeaseHandler<TestTask>(rpcLease)
+	private class Handler(RpcLease rpcLease, LeaseResult leaseResult) : LeaseHandler<TestTask>(rpcLease)
 	{
 		protected override Task<LeaseResult> ExecuteAsync(ISession session, LeaseId leaseId, TestTask message, Tracer tracer, ILogger logger, CancellationToken cancellationToken)
 		{
 			logger.LogInformation("Executed TestTask");
-			LeaseResult leaseResult = new((byte[]?)null);
 			return Task.FromResult(leaseResult);
 		}
 	}
@@ -75,7 +83,7 @@ public sealed class LeaseManagerTests
 	
 	public LeaseManagerTests()
 	{
-		_hordeServer= new FakeHordeRpcServer(CreateConsoleLogger<FakeHordeRpcServer>());
+		_hordeServer = new FakeHordeRpcServer(CreateConsoleLogger<FakeHordeRpcServer>());
 		_leaseManager = CreateLeaseManager(_hordeServer.GetHordeClient());
 	}
 	
@@ -94,6 +102,26 @@ public sealed class LeaseManagerTests
 		Assert.AreEqual(new SessionResult(SessionOutcome.BackOff, SessionReason.Completed), await runTask);
 	}
 	
+	[TestMethod]
+	public async Task UpdateSession_SendsCapabilitiesInFinalUpdate_WhenSessionTerminates_Async()
+	{
+		using CancellationTokenSource cts = new(5000);
+		LeaseResult leaseResult = new (new SessionResult((_, _) => Task.CompletedTask));
+		RpcAgentCapabilities fooCaps = new();
+		fooCaps.Properties.Add("foo");
+		FakeCapabilitiesService capsService = new () { Capabilities = fooCaps };
+		
+		LeaseManager leaseManager = CreateLeaseManager(_hordeServer.GetHordeClient(), capsService, leaseResult);
+		_hordeServer.ScheduleTestLease();
+		SessionResult sessionResult = await leaseManager.RunAsync(cts.Token);
+		_hordeServer.UpdateSessionRequests.Writer.Complete();
+		
+		List<RpcUpdateSessionRequest> requests = await _hordeServer.UpdateSessionRequests.Reader.ReadAllAsync(cts.Token).ToListAsync(cts.Token);
+		Assert.AreEqual(3, requests.Count);
+		Assert.IsTrue(requests[2].Capabilities.Properties.Contains("foo"));
+		Assert.AreEqual(new SessionResult(SessionOutcome.RunCallback, SessionReason.Completed), sessionResult);
+	}
+	
 	[TestMethod]
 	public async Task TerminateGracefully_LeaseIsActive_Async()
 	{
@@ -130,7 +158,7 @@ public sealed class LeaseManagerTests
 	}
 	
 	[SuppressMessage("Reliability", "CA2000:Dispose objects before losing scope")]
-	private static LeaseManager CreateLeaseManager(IHordeClient hordeClient)
+	private static LeaseManager CreateLeaseManager(IHordeClient hordeClient, FakeCapabilitiesService? capsService = null, LeaseResult? leaseResult = null)
 	{
 		DirectoryReference tempWorkingDir = new(Path.Join(Path.GetTempPath(), Path.GetRandomFileName()));
 		AgentSettings settings = new() { WriteStepOutputToLogger = true };
@@ -140,10 +168,10 @@ public sealed class LeaseManagerTests
 
 		return new LeaseManager(
 			session,
-			new CapabilitiesServiceStub(),
+			capsService ?? new FakeCapabilitiesService(),
 			statusService,
 			new DefaultSystemMetrics(),
-			new List<LeaseHandlerFactory>() { new TestHandlerFactory() },
+			new List<LeaseHandlerFactory>() { new TestHandlerFactory(leaseResult) },
 			new LeaseLoggerFactory(settingsOptions, CreateConsoleLogger<LeaseLoggerFactory>()),
 			settingsOptions,
 			TracerProvider.Default.GetTracer("TestTracer"),
diff --git a/Engine/Source/Programs/Horde/HordeAgent/Leases/LeaseManager.cs b/Engine/Source/Programs/Horde/HordeAgent/Leases/LeaseManager.cs
--- a/Engine/Source/Programs/Horde/HordeAgent/Leases/LeaseManager.cs
+++ b/Engine/Source/Programs/Horde/HordeAgent/Leases/LeaseManager.cs
@@ -317,6 +317,8 @@ namespace HordeAgent.Leases
 				// Get the new agent status to be reported back to server
 				if (stopping)
 				{
+					// If stopping, ensure capabilities are updated one last time
+					_capabilities = await _capabilitiesService.GetCapabilitiesAsync(_session.WorkingDir);
 					updateSessionRequest.Status = RpcAgentStatus.Stopping;
 				}
 				else if (_unhealthy)
