Tweak to Content metrics so they aggregate once every 3 hours
Removed the total % virtualization as it is doesn't really help.

diff --git a/Engine/Source/Programs/Horde/HordeServer/Defaults/default-metrics.telemetry.json b/Engine/Source/Programs/Horde/HordeServer/Defaults/default-metrics.telemetry.json
--- a/Engine/Source/Programs/Horde/HordeServer/Defaults/default-metrics.telemetry.json
+++ b/Engine/Source/Programs/Horde/HordeServer/Defaults/default-metrics.telemetry.json
@@ -544,7 +544,7 @@
 			"groupBy": "$.Payload.Build_BranchName",
 			"property": "$.Payload.TotalFileSize",
 			"function": "Sum",
-			"interval": "00:05:00"
+			"interval": "03:00:00"
 		},
 		{
 			"id": "core-content-totalsize-mp",
@@ -552,7 +552,7 @@
 			"groupBy": "$.Payload.Build_BranchName,$.Payload.Name",
 			"property": "$.Payload.TotalFileSize",
 			"function": "Sum",
-			"interval": "00:05:00",
+			"interval": "03:00:00",
 			"topN": 25
 		},
 		{
@@ -561,7 +561,7 @@
 			"groupBy": "$.Payload.Build_BranchName",
 			"property": "$.Payload.VirtualizedPercent",
 			"function": "Average",
-			"interval": "00:05:00",
+			"interval": "03:00:00"
 		},
 		{
 			"id": "core-content-virtualization-mp",
@@ -569,7 +569,7 @@
 			"groupBy": "$.Payload.Build_BranchName,$.Payload.Name",
 			"property": "$.Payload.VirtualizedPercent",
 			"function": "Average",
-			"interval": "00:05:00",
+			"interval": "03:00:00",
 			"bottomN": 25
 		},
 		{
diff --git a/Engine/Source/Programs/Horde/HordeServer/Defaults/default-views.telemetry.json b/Engine/Source/Programs/Horde/HordeServer/Defaults/default-views.telemetry.json
--- a/Engine/Source/Programs/Horde/HordeServer/Defaults/default-views.telemetry.json
+++ b/Engine/Source/Programs/Horde/HordeServer/Defaults/default-views.telemetry.json
@@ -557,15 +557,6 @@
 								}
 							]
 						},
-						{
-							"name": "Virtualization",
-							"display": "Value",
-							"metrics": [
-								{
-									"id": "core-content-virtualization"
-								}
-							]
-						},
 						{
 							"name": "Virtualization By Mount Point",
 							"display": "Value",
