Horde: Expose lease ID as an env var for jobs

diff --git a/Engine/Source/Programs/Horde/Drivers/JobDriver/Execution/JobExecutor.cs b/Engine/Source/Programs/Horde/Drivers/JobDriver/Execution/JobExecutor.cs
--- a/Engine/Source/Programs/Horde/Drivers/JobDriver/Execution/JobExecutor.cs
+++ b/Engine/Source/Programs/Horde/Drivers/JobDriver/Execution/JobExecutor.cs
@@ -1843,6 +1843,7 @@ namespace JobDriver.Execution
 			{
 			}
 
+			newEnvVars["UE_HORDE_LEASEID"] = LeaseId.ToString();
 			newEnvVars["UE_HORDE_JOBID"] = JobId.ToString();
 			newEnvVars["UE_HORDE_BATCHID"] = BatchId.ToString();
 			newEnvVars["UE_HORDE_STEPID"] = step.StepId.ToString();
