Add button to clear search on job logs.
Added text field to existing searchState to avoid adding more state.
Fixed bug where highlight remains by forcing a re-render. Also added to condition to remove straggling ">>>" prefix from previous query.

diff --git a/Engine/Source/Programs/Horde/HordeDashboard/src/components/LogView.tsx b/Engine/Source/Programs/Horde/HordeDashboard/src/components/LogView.tsx
--- a/Engine/Source/Programs/Horde/HordeDashboard/src/components/LogView.tsx
+++ b/Engine/Source/Programs/Horde/HordeDashboard/src/components/LogView.tsx
@@ -1,6 +1,6 @@
 // Copyright Epic Games, Inc. All Rights Reserved.
 
-import { Callout, DefaultButton, DetailsList, DetailsListLayoutMode, DetailsRow, DirectionalHint, Dropdown, FocusZone, FocusZoneDirection, FontIcon, IColumn, Icon, IconButton, IContextualMenuItem, IContextualMenuProps, IDetailsListProps, ITextField, List, Modal, PrimaryButton, ProgressIndicator, ScrollToMode, Selection, SelectionMode, SelectionZone, Separator, Spinner, SpinnerSize, Stack, Text, TextField, TooltipHost } from '@fluentui/react';
+import { Callout, DefaultButton, DetailsList, DetailsListLayoutMode, DetailsRow, DirectionalHint, Dropdown, FocusZone, FocusZoneDirection, FontIcon, IColumn, Icon, IconButton, IContextualMenuItem, IContextualMenuProps, IDetailsListProps, ITextField, List, mergeStyles, Modal, PrimaryButton, ProgressIndicator, ScrollToMode, Selection, SelectionMode, SelectionZone, Separator, Spinner, SpinnerSize, Stack, Text, TextField, TooltipHost } from '@fluentui/react';
 import { action, makeObservable, observable } from 'mobx';
 import { observer } from 'mobx-react-lite';
 import moment from 'moment-timezone';
@@ -239,7 +239,7 @@ let logListKey = 0;
 
 // have to be available to a globlao keyboard handler
 let globalHandler: LogHandler | undefined;
-let globalSearchState: { search?: string, results?: number[], curRequest?: any } | undefined;
+let globalSearchState: { search?: string, results?: number[], curRequest?: any, text?: string } | undefined;
 
 const searchUp = (navigate: any) => {
 
@@ -546,7 +546,7 @@ const LogLineIndicator: React.FC<{ lineIndex: number }> = ({ lineIndex }) => {
 
    useQuery();
 
-   if (!globalHandler) {
+   if (!globalHandler || !globalSearchState?.results?.length) {
       return null;
    }
 
@@ -583,7 +583,7 @@ export const LogList: React.FC<{ logId: string }> = observer(({ logId }) => {
    const navigate = useNavigate();
    const [tsFormat, setTSFormat] = useState('relative');
    const [handler, setHandler] = useState<LogHandler | undefined>(undefined);
-   const [searchState, setSearchState] = useState<{ search?: string, results?: number[], curRequest?: any }>({});
+   const [searchState, setSearchState] = useState<{ search?: string, results?: number[], curRequest?: any, text?: string }>({});
    const [issueHistory, setIssueHistory] = useState(false);
    const [logHistory, setLogHistory] = useState(false);
    const [logArtifacts, setLogArtifacts] = useState("");
@@ -928,16 +928,27 @@ export const LogList: React.FC<{ logId: string }> = observer(({ logId }) => {
 
       handler.currentLine = undefined;
 
+      // always re-render log when querying with new value
+      if (searchState?.search !== newValue) {
+         logListKey++;
+      }
+
       if (!newValue) {
 
-         setSearchState({});
+         setSearchState({
+               search: '',
+               curRequest: undefined,
+               results: [],
+               text: ''
+            });
 
       } else if (newValue !== searchState.search) {
 
          const request = backend.searchLog(logId, newValue, 0, 65535);
-
          let newState = {
-            curRequest: request
+            search: newValue,
+            curRequest: request,
+            text: newValue
          }
 
          request.then(logResponse => {
@@ -948,15 +959,13 @@ export const LogList: React.FC<{ logId: string }> = observer(({ logId }) => {
 
                const lines = logResponse.lines;
 
-
                if (lines.length) {
 
-                  logListKey++;
-
                   setSearchState({
                      search: newValue,
                      curRequest: undefined,
-                     results: lines
+                     results: lines,
+                     text: newValue
                   });
 
 
@@ -979,7 +988,8 @@ export const LogList: React.FC<{ logId: string }> = observer(({ logId }) => {
                   setSearchState({
                      search: newValue,
                      curRequest: undefined,
-                     results: lines
+                     results: lines,
+                     text: newValue
                   });
 
                }
@@ -1284,27 +1294,41 @@ export const LogList: React.FC<{ logId: string }> = observer(({ logId }) => {
                                  {!!searchState.curRequest && <Spinner style={{ paddingRight: 8 }} />}
 
                                  <TextField
-                                    suffix={searchState?.results?.length ? `${curSearchIdx + 1}/${searchState?.results?.length}` : undefined}
                                     spellCheck={false}
                                     autoComplete="off"
-                                    deferredValidationTime={1000}
+                                    deferredValidationTime={1500}
                                     validateOnLoad={false}
                                     componentRef={searchBox}
+                                    value={searchState.text ? searchState.text : ''}
+                                    onChange={(ev, value) => { setSearchState({...searchState, text: value}); }}
+                                    
+                                    onRenderSuffix={
+                                       !!searchState.text ? () => (
+                                          <Stack verticalAlign="center" horizontalAlign="center" horizontal tokens={{ childrenGap: 2 }} styles={{ root: { padding: 0 } }}>
+                                             <IconButton iconProps={{ iconName: 'Cancel'}} styles={{ root: { padding: 12, height: 20, width: 20 }, icon: { fontSize: 15 } }}
+                                             onClick={() => { doQuery('') }}
+                                             />
+                                             <span className={mergeStyles({ paddingLeft: 1, paddingRight: 7, paddingBottom: 1})}>
+                                                {searchState?.results?.length ? `${curSearchIdx + 1}/${searchState?.results?.length}` : undefined}
+                                             </span>               
+                                          </Stack>
+                                       ) : undefined
+                                    }
 
                                     styles={{
-                                       root: { width: 280, fontSize: 12 }, fieldGroup: {
-                                          borderWidth: 1
-                                       }
+                                       root: { width: 280, fontSize: 12 }, 
+                                       fieldGroup: { borderWidth: 1 },
+                                       suffix: { padding: '0 2px', height: 30 }
                                     }}
                                     placeholder="Search"
 
                                     onKeyPress={(ev) => {
 
                                        if (ev.key === "Enter" && !searchState.curRequest && !inTimeout) {
-                                          if (searchBox.current?.value === searchState.search) {
+                                          if (searchState.text === searchState.search) {
                                              searchDown(navigate);
                                           } else {
-                                             doQuery(searchBox.current?.value ?? "");
+                                             doQuery(searchState.text ?? "");
                                           }
                                        }
 
