Horde: Add telemetry spans to IssueService

diff --git a/Engine/Source/Programs/Horde/Plugins/Build/HordeServer.Build/Issues/IssueService.cs b/Engine/Source/Programs/Horde/Plugins/Build/HordeServer.Build/Issues/IssueService.cs
--- a/Engine/Source/Programs/Horde/Plugins/Build/HordeServer.Build/Issues/IssueService.cs
+++ b/Engine/Source/Programs/Horde/Plugins/Build/HordeServer.Build/Issues/IssueService.cs
@@ -138,7 +138,7 @@ namespace HordeServer.Issues
 	}
 
 	/// <summary>
-	/// Wraps funtionality for manipulating build health issues
+	/// Wraps functionality for manipulating build health issues
 	/// </summary>
 	public sealed class IssueService : IHostedService, IAsyncDisposable
 	{
@@ -334,6 +334,7 @@ namespace HordeServer.Issues
 		/// <returns>Async task</returns>
 		async ValueTask TickAsync(CancellationToken cancellationToken)
 		{
+			using TelemetrySpan span = _tracer.StartActiveSpan($"{nameof(IssueService)}.{nameof(TickAsync)}");
 			DateTime utcNow = _clock.UtcNow;
 
 			Dictionary<StreamId, HashSet<TemplateId>> newCachedDesktopAlerts = new Dictionary<StreamId, HashSet<TemplateId>>();
@@ -353,11 +354,12 @@ namespace HordeServer.Issues
 
 			// Resolve any issues that haven't been seen in a week
 			List<IIssue> openIssues = (await _issueCollection.FindIssuesAsync(resolved: false, cancellationToken: cancellationToken)).ToList();
+			span.SetAttribute("horde.open_issues", openIssues.Count);
 			for (int idx = 0; idx < openIssues.Count; idx++)
 			{
 				IIssue openIssue = openIssues[idx];
 
-				// Do not automaically close quarantined issues
+				// Do not automatically close quarantined issues
 				if (openIssue.QuarantinedByUserId != null)
 				{
 					continue;
@@ -403,6 +405,7 @@ namespace HordeServer.Issues
 		/// <inheritdoc/>
 		public async Task<IIssueDetails?> GetCachedIssueDetailsAsync(int issueId, CancellationToken cancellationToken = default)
 		{
+			using TelemetrySpan tracerSpan = _tracer.StartActiveSpan($"{nameof(IssueService)}.{nameof(GetCachedIssueDetailsAsync)}");
 			IIssueDetails? cachedIssue = _cachedIssues.FirstOrDefault(x => x.Issue.Id == issueId);
 			if (cachedIssue == null)
 			{
@@ -418,6 +421,8 @@ namespace HordeServer.Issues
 		/// <inheritdoc/>
 		public async Task<IIssueDetails> GetIssueDetailsAsync(IIssue issue, CancellationToken cancellationToken = default)
 		{
+			using TelemetrySpan tracerSpan = _tracer.StartActiveSpan($"{nameof(IssueService)}.{nameof(GetIssueDetailsAsync)}");
+			
 			IUser? owner = issue.OwnerId.HasValue ? await _userCollection.GetCachedUserAsync(issue.OwnerId.Value, cancellationToken) : null;
 			IUser? nominatedBy = issue.NominatedById.HasValue ? await _userCollection.GetCachedUserAsync(issue.NominatedById.Value, cancellationToken) : null;
 			IUser? resolvedBy = (issue.ResolvedById.HasValue && issue.ResolvedById != IIssue.ResolvedByTimeoutId && issue.ResolvedById != IIssue.ResolvedByUnknownId) ? await _userCollection.GetCachedUserAsync(issue.ResolvedById.Value, cancellationToken) : null;
@@ -490,6 +495,8 @@ namespace HordeServer.Issues
 		/// <returns>True if the issue was updated</returns>
 		public async Task<bool> UpdateIssueAsync(int id, string? summary = null, string? description = null, bool? promoted = null, UserId? ownerId = null, UserId? nominatedById = null, bool? acknowledged = null, UserId? declinedById = null, CommitId? fixCommitId = null, bool? fixSystemic = null, UserId? resolvedById = null, List<ObjectId>? addSpanIds = null, List<ObjectId>? removeSpanIds = null, string? externalIssueKey = null, UserId? quarantinedById = null, UserId? forceClosedById = null, UserId? initiatedById = null, Uri? workflowThreadUrl = null, CancellationToken cancellationToken = default)
 		{
+			using TelemetrySpan tracerSpan = _tracer.StartActiveSpan($"{nameof(IssueService)}.{nameof(UpdateIssueAsync)}");
+			
 			IIssue? issue;
 			for (; ; )
 			{
@@ -571,6 +578,7 @@ namespace HordeServer.Issues
 
 		async Task UpdateCompleteStepsAsync(CancellationToken cancellationToken)
 		{
+			using TelemetrySpan tracerSpan = _tracer.StartActiveSpan($"{nameof(IssueService)}.{nameof(UpdateCompleteStepsAsync)}");
 			while (await _completeSteps.Reader.WaitToReadAsync(cancellationToken))
 			{
 				CompleteStep? completeStep;
@@ -713,6 +721,8 @@ namespace HordeServer.Issues
 		/// <returns>Set of new events</returns>
 		async Task<HashSet<IssueEventGroupInternal>> GetEventGroupsForStepAsync(IJob job, IJobStepBatch batch, IJobStep step, INode node, IReadOnlyNodeAnnotations annotations, WorkflowConfig? workflow, CancellationToken cancellationToken)
 		{
+			using TelemetrySpan tracerSpan = _tracer.StartActiveSpan($"{nameof(IssueService)}.{nameof(GetEventGroupsForStepAsync)}");
+			
 			// Make sure the step has a log file
 			if (step.LogId == null)
 			{
@@ -819,6 +829,7 @@ namespace HordeServer.Issues
 		/// <returns>True if the adding completed</returns>
 		async Task<bool> AddEventsToExistingSpansAsync(IJob job, IJobStepBatch batch, IJobStep step, HashSet<IssueEventGroupInternal> newEventGroups, IReadOnlyList<IIssueSpan> openSpans, HashSet<ObjectId> checkedSpanIds, IReadOnlyNodeAnnotations? annotations, bool promoteByDefault, CancellationToken cancellationToken)
 		{
+			using TelemetrySpan tracerSpan = _tracer.StartActiveSpan($"{nameof(IssueService)}.{nameof(AddEventsToExistingSpansAsync)}");
 			for (int spanIdx = 0; spanIdx < openSpans.Count; spanIdx++)
 			{
 				IIssueSpan openSpan = openSpans[spanIdx];
@@ -903,6 +914,7 @@ namespace HordeServer.Issues
 		/// <returns>True if all events were added</returns>
 		async Task<bool> AddEventsToNewSpansAsync(StreamConfig streamConfig, IJob job, IJobStepBatch batch, IJobStep step, INode node, List<IIssueSpan> openSpans, HashSet<IssueEventGroupInternal> newEventGroups, IReadOnlyNodeAnnotations? annotations, bool promoteByDefault, CancellationToken cancellationToken)
 		{
+			using TelemetrySpan tracerSpan = _tracer.StartActiveSpan($"{nameof(IssueService)}.{nameof(AddEventsToNewSpansAsync)}");
 			while (newEventGroups.Count > 0)
 			{
 				// Keep track of the event groups we merge together
@@ -1022,6 +1034,7 @@ namespace HordeServer.Issues
 
 		async Task<IIssue?> UpdateIssueDerivedDataAsync(IIssue issue, CancellationToken cancellationToken)
 		{
+			using TelemetrySpan tracerSpan = _tracer.StartActiveSpan($"{nameof(IssueService)}.{nameof(UpdateIssueDerivedDataAsync)}");
 			Dictionary<(StreamId, CommitId), bool> fixChangeCache = new Dictionary<(StreamId, CommitId), bool>();
 			for (; ; )
 			{
@@ -1204,6 +1217,9 @@ namespace HordeServer.Issues
 		/// </summary>
 		async Task<bool> HasFixFailedAsync(IEnumerable<IIssueSpan> spans, CommitId? fixCommit, DateTime resolvedAt, IReadOnlyList<NewIssueStream> streams, CancellationToken cancellationToken)
 		{
+			using TelemetrySpan span = _tracer.StartActiveSpan($"{nameof(IssueService)}.{nameof(HasFixFailedAsync)}");
+			span.SetAttribute("horde.fix_commit", fixCommit?.ToString());
+			
 			foreach (IGrouping<StreamId, IIssueSpan> group in spans.Where(x => x.NextSuccess == null).GroupBy(x => x.StreamId))
 			{
 				NewIssueStream? stream = streams.FirstOrDefault(x => x.StreamId == group.Key);
@@ -1231,6 +1247,7 @@ namespace HordeServer.Issues
 		/// </summary>
 		async ValueTask<bool> ContainsFixChangeAsync(BuildConfig buildConfig, StreamId streamId, CommitId fixChange, Dictionary<(StreamId, CommitId), bool> cachedContainsFixChange, CancellationToken cancellationToken)
 		{
+			using TelemetrySpan span = _tracer.StartActiveSpan($"{nameof(IssueService)}.{nameof(ContainsFixChangeAsync)}");
 			bool containsFixChange;
 			if (!cachedContainsFixChange.TryGetValue((streamId, fixChange), out containsFixChange))
 			{
@@ -1262,6 +1279,8 @@ namespace HordeServer.Issues
 		/// <returns>List of suspects</returns>
 		async Task<List<NewIssueSpanSuspectData>> FindSuspectsForSpanAsync(StreamConfig streamConfig, IIssueFingerprint fingerprint, CommitId lastSuccess, CommitId firstFailure, CancellationToken cancellationToken)
 		{
+			using TelemetrySpan span = _tracer.StartActiveSpan($"{nameof(IssueService)}.{nameof(FindSuspectsForSpanAsync)}");
+			
 			List<NewIssueSpanSuspectData> suspects = new List<NewIssueSpanSuspectData>();
 			if (!String.IsNullOrEmpty(fingerprint.ChangeFilter))
 			{
@@ -1308,7 +1327,7 @@ namespace HordeServer.Issues
 				int maxRank = -1;
 				foreach (ICommit commit in changes)
 				{
-					int rank = await GetChangeScoreAsync(commit, filter, suspectFiles, suspectNotes, suspectSymbols);
+					int rank = await GetChangeScoreAsync(commit, filter, suspectFiles, suspectNotes, suspectSymbols, _tracer);
 					if (rank > maxRank)
 					{
 						suspects.Clear();
@@ -1326,8 +1345,10 @@ namespace HordeServer.Issues
 			return suspects;
 		}
 
-		static async ValueTask<int> GetChangeScoreAsync(ICommit commit, FileFilter filter, HashSet<string> suspectFiles, HashSet<string> suspectNotes, HashSet<string> suspectSymbols)
+		static async ValueTask<int> GetChangeScoreAsync(ICommit commit, FileFilter filter, HashSet<string> suspectFiles, HashSet<string> suspectNotes, HashSet<string> suspectSymbols, Tracer tracer)
 		{
+			using TelemetrySpan span = tracer.StartActiveSpan($"{nameof(IssueService)}.{nameof(GetChangeScoreAsync)}");
+			
 			int rank = -1;
 
 			IReadOnlyList<string> commitFiles = await commit.GetFilesAsync(MaxFilesForSuspects, CancellationToken.None);
@@ -1372,6 +1393,7 @@ namespace HordeServer.Issues
 		/// <returns>The matching issue</returns>
 		async Task<IIssue> FindOrAddIssueForSpanAsync(IIssueSpan span, CancellationToken cancellationToken)
 		{
+			using TelemetrySpan tracerSpan = _tracer.StartActiveSpan($"{nameof(IssueService)}.{nameof(FindOrAddIssueForSpanAsync)}");
 			IReadOnlyList<IIssue> existingIssues;
 			if (span.LastSuccess == null)
 			{
@@ -1399,6 +1421,7 @@ namespace HordeServer.Issues
 		/// <returns>The matching issue</returns>
 		async Task<IIssue> FindOrAddIssueForSpanAsync(NewIssueSpanData span, CancellationToken cancellationToken)
 		{
+			using TelemetrySpan tracerSpan = _tracer.StartActiveSpan($"{nameof(IssueService)}.{nameof(FindOrAddIssueForSpanAsync)}.NewIssueSpan");
 			IReadOnlyList<IIssue> existingIssues;
 			if (span.LastSuccess == null)
 			{
@@ -1552,6 +1575,7 @@ namespace HordeServer.Issues
 		/// <returns></returns>
 		async Task<bool> TryUpdateSentinelsAsync(IReadOnlyList<IIssueSpan> spans, StreamConfig streamConfig, IJob job, IJobStepBatch batch, IJobStep step, CancellationToken cancellationToken)
 		{
+			using TelemetrySpan tracerSpan = _tracer.StartActiveSpan($"{nameof(IssueService)}.{nameof(TryUpdateSentinelsAsync)}");
 			foreach (IIssueSpan span in spans)
 			{
 				if (job.CommitId < span.FirstFailure.CommitId && (span.LastSuccess == null || job.CommitId > span.LastSuccess.CommitId))
