Improve handling/recovery of hanging P4 health checks.

- Tasks to update server status now run asynchronously, outside of the regular health tick loop. This allows healthy servers to be updated even if unhealthy servers are not returning a valid status.
- Report the pod status as unhealthy if an update task has been running > 5m.

#rnx

diff --git a/Engine/Source/Programs/Horde/HordeServer/Startup.cs b/Engine/Source/Programs/Horde/HordeServer/Startup.cs
--- a/Engine/Source/Programs/Horde/HordeServer/Startup.cs
+++ b/Engine/Source/Programs/Horde/HordeServer/Startup.cs
@@ -35,6 +35,7 @@ using HordeServer.Auditing;
 using HordeServer.Authentication;
 using HordeServer.Configuration;
 using HordeServer.Dashboard;
+using HordeServer.Perforce;
 using HordeServer.Plugins;
 using HordeServer.Server;
 using HordeServer.Server.Notices;
@@ -939,7 +940,8 @@ namespace HordeServer
 			services.AddHealthChecks()
 				.AddCheck("self", () => HealthCheckResult.Healthy(), tags: new[] { "self" })
 				.AddCheck<RedisService>("redis")
-				.AddCheck<MongoService>("mongo");
+				.AddCheck<MongoService>("mongo")
+				.AddCheck<PerforceLoadBalancer>("p4lb");
 		}
 
 		// This method gets called by the runtime. Use this method to configure the HTTP request pipeline.
diff --git a/Engine/Source/Programs/Horde/Plugins/Build/HordeServer.Build/Perforce/PerforceLoadBalancer.cs b/Engine/Source/Programs/Horde/Plugins/Build/HordeServer.Build/Perforce/PerforceLoadBalancer.cs
--- a/Engine/Source/Programs/Horde/Plugins/Build/HordeServer.Build/Perforce/PerforceLoadBalancer.cs
+++ b/Engine/Source/Programs/Horde/Plugins/Build/HordeServer.Build/Perforce/PerforceLoadBalancer.cs
@@ -1,5 +1,6 @@
 // Copyright Epic Games, Inc. All Rights Reserved.
 
+using System.Diagnostics;
 using System.Globalization;
 using System.Net;
 using System.Text;
@@ -100,7 +101,7 @@ namespace HordeServer.Perforce
 	/// <summary>
 	/// Load balancer for Perforce edge servers
 	/// </summary>
-	public sealed class PerforceLoadBalancer : IHostedService, IAsyncDisposable
+	public sealed class PerforceLoadBalancer : IHostedService, IHealthCheck, IAsyncDisposable
 	{
 		/// <summary>
 		/// Information about a resolved Perforce server
@@ -144,6 +145,8 @@ namespace HordeServer.Perforce
 			public List<PerforceServerEntry> Servers { get; set; } = new List<PerforceServerEntry>();
 		}
 
+		record class HealthTask(Task Task, Stopwatch Timer);
+
 		readonly IRedisService _redisService;
 		readonly ILeaseCollection _leaseCollection;
 		readonly SingletonDocument<PerforceServerList> _serverListSingleton;
@@ -154,6 +157,9 @@ namespace HordeServer.Perforce
 		readonly Tracer _tracer;
 		readonly ILogger _logger;
 		readonly ITicker _ticker;
+		readonly CancellationTokenSource _healthTaskCts = new CancellationTokenSource();
+		readonly SemaphoreSlim _healthTaskSema = new SemaphoreSlim(1);
+		readonly Dictionary<string, HealthTask> _healthTasks = new Dictionary<string, HealthTask>(StringComparer.OrdinalIgnoreCase);
 
 		/// <summary>
 		/// Constructor
@@ -167,6 +173,7 @@ namespace HordeServer.Perforce
 			_buildConfig = buildConfig;
 			_tracer = tracer;
 			_logger = logger;
+
 			if (mongoService.ReadOnlyMode)
 			{
 				_ticker = new NullTicker();
@@ -187,7 +194,39 @@ namespace HordeServer.Perforce
 		public Task StopAsync(CancellationToken cancellationToken) => _ticker.StopAsync();
 
 		/// <inheritdoc/>
-		public async ValueTask DisposeAsync() => await _ticker.DisposeAsync();
+		public async ValueTask DisposeAsync()
+		{
+			await _ticker.DisposeAsync();
+
+			if (_healthTasks.Count > 0)
+			{
+				await _healthTaskCts.CancelAsync();
+				foreach (HealthTask healthTask in _healthTasks.Values)
+				{
+					await healthTask.Task;
+				}
+				_healthTasks.Clear();
+			}
+
+			_healthTaskCts.Dispose();
+		}
+
+		/// <inheritdoc/>
+		public async Task<HealthCheckResult> CheckHealthAsync(HealthCheckContext context, CancellationToken cancellationToken = default)
+		{
+			using (await _healthTaskSema.WaitDisposableAsync(cancellationToken))
+			{
+				KeyValuePair<string, HealthTask> longestTask = _healthTasks.OrderByDescending(x => x.Value.Timer.Elapsed).FirstOrDefault(x => !x.Value.Task.IsCompleted);
+				if (longestTask.Key == null || longestTask.Value.Timer.Elapsed < TimeSpan.FromMinutes(5.0))
+				{
+					return HealthCheckResult.Healthy();
+				}
+				else
+				{
+					return HealthCheckResult.Unhealthy($"Health check for {longestTask.Key} has taken {(int)longestTask.Value.Timer.Elapsed.TotalSeconds}s");
+				}
+			}
+		}
 
 		/// <summary>
 		/// Get the current server list
@@ -423,6 +462,7 @@ namespace HordeServer.Perforce
 		/// <inheritdoc/>
 		async ValueTask TickInternalAsync(CancellationToken cancellationToken)
 		{
+			_logger.LogDebug("Ticking server health status");
 			using TelemetrySpan span = _tracer.StartActiveSpan($"{nameof(PerforceLoadBalancer)}.{nameof(TickInternalAsync)}");
 
 			BuildConfig buildConfig = _buildConfig.CurrentValue;
@@ -450,15 +490,36 @@ namespace HordeServer.Perforce
 			List<PerforceServerEntry> newEntries = newServers.OrderBy(x => x.Cluster).ThenBy(x => x.BaseServerAndPort).ThenBy(x => x.ServerAndPort).ToList();
 			await UpdateLeaseCountsAsync(newEntries, cancellationToken);
 			PerforceServerList list = await _serverListSingleton.UpdateAsync(list => MergeServerList(list, newEntries), cancellationToken);
-
-			// Now update the health of each entry in parallel
-			List<Task> tasks = [];
 			span.SetAttribute("numServers", list.Servers.Count);
-			foreach (PerforceServerEntry entry in list.Servers)
+
+			// Lock the task lookup and update server health
+			using (await _healthTaskSema.WaitDisposableAsync(cancellationToken))
 			{
-				tasks.Add(Task.Run(() => UpdateHealthAsync(entry, cancellationToken), cancellationToken));
+				// Remove any health check tasks that have finished
+				KeyValuePair<string, HealthTask>[] healthTasks = _healthTasks.ToArray();
+				foreach ((string serverAndPort, HealthTask healthTask) in healthTasks)
+				{
+					if (healthTask.Task.IsCompleted)
+					{
+						await healthTask.Task;
+						_healthTasks.Remove(serverAndPort);
+					}
+					else if (healthTask.Timer.Elapsed > TimeSpan.FromMinutes(1.0))
+					{
+						_logger.LogWarning("Health check for {ServerAndPort} is still running after {TimeSecs}s", serverAndPort, (int)healthTask.Timer.Elapsed.TotalSeconds);
+					}
+				}
+
+				// Now update the health of each entry in parallel
+				foreach (PerforceServerEntry entry in list.Servers)
+				{
+					if (!_healthTasks.ContainsKey(entry.ServerAndPort))
+					{
+						Task task = Task.Run(() => UpdateHealthGuardedAsync(entry, _healthTaskCts.Token), CancellationToken.None);
+						_healthTasks.Add(entry.ServerAndPort, new HealthTask(task, Stopwatch.StartNew()));
+					}
+				}
 			}
-			await Task.WhenAll(tasks);
 
 			// Update Horde's health monitor for subsystems
 			list = await _serverListSingleton.GetAsync(cancellationToken);
@@ -473,15 +534,18 @@ namespace HordeServer.Perforce
 			HealthStatus result = HealthStatus.Healthy;
 			foreach (PerforceServerEntry server in servers)
 			{
-				HealthStatus serverHealth = server.Status switch
+				if (server.Status != PerforceServerStatus.Unknown)
 				{
-					PerforceServerStatus.Unknown => HealthStatus.Degraded,
-					PerforceServerStatus.Unhealthy => HealthStatus.Unhealthy,
-					PerforceServerStatus.Degraded => HealthStatus.Degraded,
-					PerforceServerStatus.Healthy => HealthStatus.Healthy,
-					_ => throw new ArgumentOutOfRangeException($"Unknown health status: {server.Status}")
-				};
-				result = serverHealth < result ? serverHealth : result;
+					HealthStatus serverHealth = server.Status switch
+					{
+						PerforceServerStatus.Unknown => HealthStatus.Degraded,
+						PerforceServerStatus.Unhealthy => HealthStatus.Unhealthy,
+						PerforceServerStatus.Degraded => HealthStatus.Degraded,
+						PerforceServerStatus.Healthy => HealthStatus.Healthy,
+						_ => throw new ArgumentOutOfRangeException($"Unknown health status: {server.Status}")
+					};
+					result = serverHealth < result ? serverHealth : result;
+				}
 			}
 
 			string? message = result switch
@@ -632,6 +696,21 @@ namespace HordeServer.Perforce
 			}
 		}
 
+		async Task UpdateHealthGuardedAsync(PerforceServerEntry entry, CancellationToken cancellationToken)
+		{
+			try
+			{
+				await UpdateHealthAsync(entry, cancellationToken);
+			}
+			catch (OperationCanceledException)
+			{
+			}
+			catch (Exception ex)
+			{
+				_logger.LogError(ex, "Error updating health for {ServerAndPort}: {Message}", entry.ServerAndPort, ex.Message);
+			}
+		}
+
 		/// <summary>
 		/// Checks and updates the health of a Perforce server
 		/// </summary>
@@ -712,9 +791,9 @@ namespace HordeServer.Perforce
 				_logger.LogWarning("Unable to parse health check output. Data {Data} URL {HealthCheckUrl}", dataStr, healthCheckUrl);
 				return new ServerHealth(PerforceServerStatus.Unknown, "Unable to parse health check output");
 			}
-			catch (Exception e)
+			catch (Exception ex) when (ex is not OperationCanceledException)
 			{
-				_logger.LogWarning("Failed checking Perforce server health via HTTP. URL {HealthCheckUrl} Reason {Reason}", healthCheckUrl, e.Message);
+				_logger.LogWarning("Failed checking Perforce server health via HTTP. URL {HealthCheckUrl} Reason {Reason}", healthCheckUrl, ex.Message);
 				return new ServerHealth(PerforceServerStatus.Unhealthy, $"Failed to query status at {healthCheckUrl}");
 			}
 		}
@@ -736,9 +815,9 @@ namespace HordeServer.Perforce
 				await connection.GetInfoAsync(InfoOptions.ShortOutput, cancellationToken);
 				return new ServerHealth(PerforceServerStatus.Healthy, "Server responded to \"p4 info\"");
 			}
-			catch (Exception e)
+			catch (Exception ex) when (ex is not OperationCanceledException)
 			{
-				logger.LogWarning("Failed checking Perforce server health via info command. Reason {Reason}", e.Message);
+				logger.LogWarning("Failed checking Perforce server health via info command. Reason {Reason}", ex.Message);
 				return new ServerHealth(PerforceServerStatus.Unhealthy, "Server failed \"p4 info\" query");
 			}
 		}
