Horde: Fix static analysis warning.

#rnx

diff --git a/Engine/Source/Programs/Horde/Plugins/Build/HordeServer.Build/Perforce/PerforceLoadBalancer.cs b/Engine/Source/Programs/Horde/Plugins/Build/HordeServer.Build/Perforce/PerforceLoadBalancer.cs
--- a/Engine/Source/Programs/Horde/Plugins/Build/HordeServer.Build/Perforce/PerforceLoadBalancer.cs
+++ b/Engine/Source/Programs/Horde/Plugins/Build/HordeServer.Build/Perforce/PerforceLoadBalancer.cs
@@ -209,6 +209,7 @@ namespace HordeServer.Perforce
 			}
 
 			_healthTaskCts.Dispose();
+			_healthTaskSema.Dispose();
 		}
 
 		/// <inheritdoc/>
