Horde: Fix command line for running BuildGraph script in Horde docs.

#rnx

diff --git a/Engine/Source/Programs/Horde/Docs/Internals/Building.md b/Engine/Source/Programs/Horde/Docs/Internals/Building.md
--- a/Engine/Source/Programs/Horde/Docs/Internals/Building.md
+++ b/Engine/Source/Programs/Horde/Docs/Internals/Building.md
@@ -40,7 +40,7 @@ The Horde dashboard is a frontend client developed in [TypeScript](https://www.t
 Horde includes a `Dockerfile` for creating Docker images. However, its position in Unreal Engine source tree requires
 files to be staged beforehand to reduce the size of data copied into build images.
 
-A BuildGraph script to perform these operations is included under `Engine\Source\Programs\Horde\HordeBuild.xml`, which
+A BuildGraph script to perform these operations is included under `Engine\Source\Programs\Horde\BuildHorde.xml`, which
 can be run as follows:
 
-    RunUAT.bat Engine/Source/Programs/Horde/HordeBuild.xml -Target="Build HordeServer"
+    RunUAT.bat BuildGraph -Script=Engine/Source/Programs/Horde/BuildHorde.xml -Target="Build HordeServer"
