Horde: Add a "Copy Artifact ID" to artifact button

diff --git a/Engine/Source/Programs/Horde/HordeDashboard/src/components/artifacts/ArtifactButton.tsx b/Engine/Source/Programs/Horde/HordeDashboard/src/components/artifacts/ArtifactButton.tsx
--- a/Engine/Source/Programs/Horde/HordeDashboard/src/components/artifacts/ArtifactButton.tsx
+++ b/Engine/Source/Programs/Horde/HordeDashboard/src/components/artifacts/ArtifactButton.tsx
@@ -2,6 +2,7 @@ import { IContextualMenuProps, DirectionalHint, PrimaryButton } from "@fluentui/
 import backend from "horde/backend";
 import { ArtifactBrowserItem, ArtifactBrowserType, ArtifactSelection, CreateZipRequest, DownloadFormat, GetArtifactResponse } from "horde/backend/Api";
 import dashboard from "horde/backend/Dashboard";
+import { copyToClipboard } from "horde/base/utilities/clipboard";
 import { formatBytes } from "horde/base/utilities/stringUtills";
 import { useNavigate } from "react-router-dom";
 
@@ -170,6 +171,17 @@ export const ArtifactButton: React.FC<{ jobId?: string, stepId?: string, artifac
             }
          }
       )
+
+      downloadProps.items.push(
+         {
+            key: 'download_copy_to_clipboard',
+            text: 'Copy Artifact ID',
+            onClick: () => {
+               copyToClipboard(artifact.id);
+            }
+         }
+      )
+
    }
 
    return <PrimaryButton split menuProps={downloadProps} text={buttonText} styles={{ root: { minWidth: minWidth ?? 104, fontFamily: 'Horde Open Sans SemiBold !important' } }} disabled={disabled} onClick={() => {
