[UE-192322] Job dialog: improve navigation of long dropdown lists

Adds a new `LongListParameter` component used when the number of items in a non TagPicker style list exceeds a set threshold. Supports single and multi selection, as well as custom filtering

diff --git a/Engine/Source/Programs/Horde/HordeDashboard/src/components/build/NewBuildV2.tsx b/Engine/Source/Programs/Horde/HordeDashboard/src/components/build/NewBuildV2.tsx
--- a/Engine/Source/Programs/Horde/HordeDashboard/src/components/build/NewBuildV2.tsx
+++ b/Engine/Source/Programs/Horde/HordeDashboard/src/components/build/NewBuildV2.tsx
@@ -1,4 +1,4 @@
-import { Checkbox, ComboBox, ContextualMenuItemType, DefaultButton, DirectionalHint, Dropdown, DropdownMenuItemType, IComboBoxOption, IContextualMenuItem, IContextualMenuProps, IDropdownOption, Icon, IconButton, Label, MessageBar, MessageBarType, Modal, Pivot, PivotItem, PrimaryButton, ScrollablePane, ScrollbarVisibility, Spinner, SpinnerSize, Stack, TagPicker, Text, TextField, TooltipHost, ValidationState } from "@fluentui/react";
+import { Checkbox, ComboBox, ContextualMenuItemType, DefaultButton, DirectionalHint, Dropdown, DropdownMenuItemType, IComboBoxOption, IContextualMenuItem, IContextualMenuProps, IDropdownOption, ITag, Icon, IconButton, Label, MessageBar, MessageBarType, Modal, Pivot, PivotItem, PrimaryButton, ScrollablePane, ScrollbarVisibility, Spinner, SpinnerSize, Stack, TagPicker, Text, TextField, TooltipHost, ValidationState } from "@fluentui/react";
 import { useConst } from '@fluentui/react-hooks';
 import { ITextField } from "@fluentui/react/lib-commonjs/TextField";
 import Markdown from "markdown-to-jsx";
@@ -28,6 +28,8 @@ type ValidationError = {
 const parameterWidth = 667;
 const parameterGap = 6;
 
+const LONG_LIST_PARAMETER_THRESHOLD = 50; // List parameters with more than this many items will become a `LongListParamater` automatically
+
 type NewBuildMode = "Basic" | "Advanced";
 
 type OptionsConstruction = {
@@ -1280,12 +1282,94 @@ const MultiListParameter: React.FC<{ param: ListParameterData, disabled?: boolea
    />;
 })
 
+type PickerItem = {
+   itemData: ListParameterItemData;
+   key: string;
+   name: string;
+}
+const LongListParameter: React.FC<{ param: ListParameterData, disabled?: boolean, allowSelectFilter?: (item: PickerItem) => boolean, allowMultiple: boolean }> = observer(({ param, disabled, allowSelectFilter, allowMultiple }) => {
 
-const ListParameter: React.FC<{ param: ListParameterData, disabled?: boolean }> = ({ param, disabled }) => {
+   const options = BuildOptions.get();
+
+   options.subscribeToParameterChange();
+
+   const key = `parameter_key_${param.label}` + `_${options.currentRenderKey}`;
+
+   const allItems: Map<ITag['key'], PickerItem> = new Map( param.items.map( item => {
+      return [
+         item.text, 
+         {
+            itemData: item,
+            key: item.text,
+            name: item.text
+         }
+      ]
+   }) );
+
+   const selectFilter = allowSelectFilter ?? (( item: PickerItem ) => {
+      return !(item.name.startsWith("**") && item.name.endsWith( "**" ));
+   });
+   const selectedItems = allItems.values().filter(item => {
+         return options.parameters[item.itemData.id] === "true" && selectFilter( item );
+      })
+      .toArray();
+
+   // tag picker
+   return <Stack key={key}>
+      <Label>
+         {param.label}
+      </Label>
+      <TagPicker
+         disabled={options.readOnly || disabled}
+         onResolveSuggestions={(filter, selected) => {
+            return allItems.values().filter(i => {
+                  return !selected?.find(s => i.key === s.key) && i.name.toLowerCase().indexOf(filter.toLowerCase()) !== -1;
+               })
+               .toArray();
+         }}
+
+         onEmptyResolveSuggestions={(selected) => {
+            return allItems.values().filter(i => {
+                  return !selected?.find(s => i.key === s.key);
+               })
+               .toArray();
+         }}
+
+         selectedItems={selectedItems}
+
+         onChange={(item_tags?) => {
+            
+            if ( item_tags ) {
+
+               // Get the items for all selected tags for filtering 
+               const selected_items = item_tags.map( tag => allItems.get( tag.key ) ).filter( item => item !== undefined );
+
+               let filtered_items = selected_items.filter( item => selectFilter( item ) );
+               if( !allowMultiple ) {
+                  // Keep the last item which passed the filter list as it's most likely to be the newest (need more selection metadata to perform accurately)
+                  filtered_items = [ filtered_items[ filtered_items.length - 1 ] ];
+               }
+
+               // Update all tags
+               for( const [key, item] of allItems ) {
+                  const was_selected    = filtered_items.find( selected_item => key === selected_item?.key ) !== undefined;
+                  const can_be_selected = selectFilter( item );
+                  
+                  options.onListItemChanged( item.itemData.id, was_selected && can_be_selected );
+               }
+            }
+            options.onListItemsUpdated();
+         }}
+      />
+   </Stack>;
+})
 
+const ListParameter: React.FC<{ param: ListParameterData, disabled?: boolean }> = ({ param, disabled }) => {
 
    if (param.style === ListParameterStyle.TagPicker) {
       return <TagPickerParameter param={param} disabled={disabled} />
+   } else if (param.items.length > LONG_LIST_PARAMETER_THRESHOLD) {
+      return <LongListParameter param={param} disabled={disabled} allowMultiple={param.style === ListParameterStyle.MultiList}/>
    } else if (param.style === ListParameterStyle.MultiList) {
       return <MultiListParameter param={param} disabled={disabled} />
    } else if (param.style === ListParameterStyle.List) {
