Horde Log search improvements.
Added support for populating search field with highlighted text from log, either by hitting Ctrl+F or left-clicking in the field.
Ctrl+F will also now select text in the search field, if there is any, allowing users to replace it by continuing to type.
Now distinguishing between no results and undefined, in order to display "0/0" and clear button when there are no matches.

diff --git a/Engine/Source/Programs/Horde/HordeDashboard/src/components/LogView.tsx b/Engine/Source/Programs/Horde/HordeDashboard/src/components/LogView.tsx
--- a/Engine/Source/Programs/Horde/HordeDashboard/src/components/LogView.tsx
+++ b/Engine/Source/Programs/Horde/HordeDashboard/src/components/LogView.tsx
@@ -606,6 +606,15 @@ export const LogList: React.FC<{ logId: string }> = observer(({ logId }) => {
 
    let inTimeout = false;
 
+   function searchSelectedText() {
+      const selection = window.getSelection();
+      const selected_text = selection ? selection.toString() : '';
+      
+      if (selected_text) {
+         setSearchState({...searchState, text: selected_text})
+      }
+   }
+
    useEffect(() => {
 
       const handler = (e: KeyboardEvent) => {
@@ -623,6 +632,13 @@ export const LogList: React.FC<{ logId: string }> = observer(({ logId }) => {
 
          if ((e.ctrlKey || e.metaKey) && e.keyCode === 70) {
             e.preventDefault();
+
+            searchSelectedText();
+            
+            if (searchState.text !== '') {
+               searchBox.current?.setSelectionRange(0, searchBox.current.value!.length)
+            }
+
             searchBox.current?.focus();
          }
       }
@@ -636,6 +652,13 @@ export const LogList: React.FC<{ logId: string }> = observer(({ logId }) => {
 
    }, []);
 
+   const handleSearchBoxClick = (e: React.MouseEvent<HTMLElement, MouseEvent>) => {
+      if (e.button !== 0) return; // only capture left clicks
+
+      searchSelectedText();
+   };
+   
+
    if (logError) {
       return <Stack horizontalAlign='center' style={{ paddingTop: 24 }} >
          <Text variant='mediumPlus'>{`Unable to load log data - ${logError}`}</Text>
@@ -938,7 +961,7 @@ export const LogList: React.FC<{ logId: string }> = observer(({ logId }) => {
          setSearchState({
                search: '',
                curRequest: undefined,
-               results: [],
+               results: undefined,
                text: ''
             });
 
@@ -1301,6 +1324,7 @@ export const LogList: React.FC<{ logId: string }> = observer(({ logId }) => {
                                     componentRef={searchBox}
                                     value={searchState.text ? searchState.text : ''}
                                     onChange={(ev, value) => { setSearchState({...searchState, text: value}); }}
+                                    onMouseDown={handleSearchBoxClick}
                                     
                                     onRenderSuffix={
                                        !!searchState.text ? () => (
@@ -1309,7 +1333,7 @@ export const LogList: React.FC<{ logId: string }> = observer(({ logId }) => {
                                              onClick={() => { doQuery('') }}
                                              />
                                              <span className={mergeStyles({ paddingLeft: 1, paddingRight: 7, paddingBottom: 1})}>
-                                                {searchState?.results?.length ? `${curSearchIdx + 1}/${searchState?.results?.length}` : undefined}
+                                                {searchState?.results ? searchState?.results?.length ? `${curSearchIdx + 1}/${searchState?.results?.length}` : '0/0' : undefined}
                                              </span>               
                                           </Stack>
                                        ) : undefined
