Horde: Add a warning for slow metric evaluation.

#rnx

diff --git a/Engine/Source/Programs/Horde/Plugins/Analytics/HordeServer.Analytics/Telemetry/Metrics/MetricCollection.cs b/Engine/Source/Programs/Horde/Plugins/Analytics/HordeServer.Analytics/Telemetry/Metrics/MetricCollection.cs
--- a/Engine/Source/Programs/Horde/Plugins/Analytics/HordeServer.Analytics/Telemetry/Metrics/MetricCollection.cs
+++ b/Engine/Source/Programs/Horde/Plugins/Analytics/HordeServer.Analytics/Telemetry/Metrics/MetricCollection.cs
@@ -1,5 +1,6 @@
 // Copyright Epic Games, Inc. All Rights Reserved.
 
+using System.Diagnostics;
 using System.Text.Json.Nodes;
 using EpicGames.Core;
 using EpicGames.Horde.Telemetry;
@@ -124,9 +125,28 @@ namespace HordeServer.Telemetry.Metrics
 			TelemetryStoreConfig? telemetryStoreConfig;
 			if (_analyticsConfig.CurrentValue.TryGetTelemetryStore(storeId, out telemetryStoreConfig))
 			{
+				Stopwatch timer = Stopwatch.StartNew();
+				Stopwatch metricTimer = new Stopwatch();
+
+				MetricId slowestMetricId = default;
+				TimeSpan slowestMetricTime = TimeSpan.Zero;
+
 				foreach (MetricConfig metric in telemetryStoreConfig.Metrics)
 				{
+					metricTimer.Restart();
 					AddEvent(storeId, metric, node, array);
+					metricTimer.Stop();
+
+					if (metricTimer.Elapsed > slowestMetricTime)
+					{
+						slowestMetricId = metric.Id;
+						slowestMetricTime = metricTimer.Elapsed;
+					}
+				}
+
+				if (timer.Elapsed.TotalMilliseconds > 50)
+				{
+					_logger.LogWarning("Evaluating metrics for {StoreId} took {TimeMs}ms (slowest: {MetricId} @ {MetricTimeMs}ms)", storeId, (long)timer.ElapsedMilliseconds, slowestMetricId, (long)slowestMetricTime.TotalMilliseconds);
 				}
 			}
 
