Horde: Do not re-queue telemetry data that fails to flush.

#rnx

diff --git a/Engine/Source/Programs/Horde/Plugins/Analytics/HordeServer.Analytics/Telemetry/Metrics/MetricCollection.cs b/Engine/Source/Programs/Horde/Plugins/Analytics/HordeServer.Analytics/Telemetry/Metrics/MetricCollection.cs
--- a/Engine/Source/Programs/Horde/Plugins/Analytics/HordeServer.Analytics/Telemetry/Metrics/MetricCollection.cs
+++ b/Engine/Source/Programs/Horde/Plugins/Analytics/HordeServer.Analytics/Telemetry/Metrics/MetricCollection.cs
@@ -294,7 +294,14 @@ namespace HordeServer.Telemetry.Metrics
 				newDataTask = _newDataEvent.Task;
 				flushTask = _flushEvent.Task;
 
-				await FlushAsync(cancellationToken);
+				try
+				{
+					await FlushAsync(cancellationToken);
+				}
+				catch (Exception ex)
+				{
+					_logger.LogError(ex, "Error flushing telemetry data: {Message}", ex.Message);
+				}
 			}
 		}
 
@@ -309,27 +316,13 @@ namespace HordeServer.Telemetry.Metrics
 			}
 
 			// Add the samples to the database
-			try
+			foreach ((SampleKey sampleKey, List<double> sampleValues) in samples)
 			{
-				foreach ((SampleKey sampleKey, List<double> sampleValues) in samples)
+				MetricConfig? metricConfig;
+				if (_analyticsConfig.CurrentValue.TryGetTelemetryStore(sampleKey.Store, out TelemetryStoreConfig? telemetryStoreConfig) && telemetryStoreConfig.TryGetMetric(sampleKey.Metric, out metricConfig))
 				{
-					MetricConfig? metricConfig;
-					if (_analyticsConfig.CurrentValue.TryGetTelemetryStore(sampleKey.Store, out TelemetryStoreConfig? telemetryStoreConfig) && telemetryStoreConfig.TryGetMetric(sampleKey.Metric, out metricConfig))
-					{
-						await CombineValuesAsync(sampleKey.Store, metricConfig, sampleKey.Group, sampleKey.Time, sampleValues, cancellationToken);
-					}
-				}
-			}
-			catch
-			{
-				lock (_lockObject)
-				{
-					foreach ((SampleKey key, List<double> values) in samples)
-					{
-						QueueSampleValues(key, values);
-					}
+					await CombineValuesAsync(sampleKey.Store, metricConfig, sampleKey.Group, sampleKey.Time, sampleValues, cancellationToken);
 				}
-				throw;
 			}
 		}
 
