Horde: Add warnings for waits over lock times.

#rnx

diff --git a/Engine/Source/Programs/Horde/Plugins/Analytics/HordeServer.Analytics/Telemetry/Metrics/MetricCollection.cs b/Engine/Source/Programs/Horde/Plugins/Analytics/HordeServer.Analytics/Telemetry/Metrics/MetricCollection.cs
--- a/Engine/Source/Programs/Horde/Plugins/Analytics/HordeServer.Analytics/Telemetry/Metrics/MetricCollection.cs
+++ b/Engine/Source/Programs/Horde/Plugins/Analytics/HordeServer.Analytics/Telemetry/Metrics/MetricCollection.cs
@@ -244,9 +244,24 @@ namespace HordeServer.Telemetry.Metrics
 				DateTime sampleTime = new DateTime(utcNow.Ticks - (utcNow.Ticks % metric.Interval.Ticks), DateTimeKind.Utc);
 
 				SampleKey key = new SampleKey(telemetryStoreId, metric.Id, group.ToString(), sampleTime);
+
+				Stopwatch timer = Stopwatch.StartNew();
 				lock (_lockObject)
 				{
+					timer.Stop();
+					if (timer.ElapsedMilliseconds > 100)
+					{
+						_logger.LogWarning("Waited {TimerMs}ms for lock over telemetry sample collection", timer.ElapsedMilliseconds);
+					}
+
+					timer.Restart();
 					QueueSampleValues(key, values);
+
+					timer.Stop();
+					if (timer.ElapsedMilliseconds > 100)
+					{
+						_logger.LogWarning("Queuing {NumSamples} for {StoreId}:{MetricId} took {TimerMs}", values.Count, telemetryStoreId, metric.Id, timer.ElapsedMilliseconds);
+					}
 				}
 
 				_newDataEvent.Set();
