Horde: Include event data in slow metric warnings.

#rnx

diff --git a/Engine/Source/Programs/Horde/Plugins/Analytics/HordeServer.Analytics/Telemetry/Metrics/MetricCollection.cs b/Engine/Source/Programs/Horde/Plugins/Analytics/HordeServer.Analytics/Telemetry/Metrics/MetricCollection.cs
--- a/Engine/Source/Programs/Horde/Plugins/Analytics/HordeServer.Analytics/Telemetry/Metrics/MetricCollection.cs
+++ b/Engine/Source/Programs/Horde/Plugins/Analytics/HordeServer.Analytics/Telemetry/Metrics/MetricCollection.cs
@@ -146,7 +146,15 @@ namespace HordeServer.Telemetry.Metrics
 
 				if (timer.Elapsed.TotalMilliseconds > 50)
 				{
-					_logger.LogWarning("Evaluating metrics for {StoreId} took {TimeMs}ms (slowest: {MetricId} @ {MetricTimeMs}ms)", storeId, (long)timer.ElapsedMilliseconds, slowestMetricId, (long)slowestMetricTime.TotalMilliseconds);
+					string text = node.ToJsonString();
+
+					const int MaxLength = 4096;
+					if (text.Length > MaxLength)
+					{
+						text = text.Substring(0, MaxLength) + $"... ({text.Length})";
+					}
+
+					_logger.LogWarning("Evaluating metrics for {StoreId} took {TimeMs}ms (slowest: {MetricId} @ {MetricTimeMs}ms): {Text}", storeId, (long)timer.ElapsedMilliseconds, slowestMetricId, (long)slowestMetricTime.TotalMilliseconds, text);
 				}
 			}
 
