Horde: Show warnings for individual json path queries taking longer than expected, or returning a large number of results.

#rnx

diff --git a/Engine/Source/Programs/Horde/Plugins/Analytics/HordeServer.Analytics/Telemetry/Metrics/MetricCollection.cs b/Engine/Source/Programs/Horde/Plugins/Analytics/HordeServer.Analytics/Telemetry/Metrics/MetricCollection.cs
--- a/Engine/Source/Programs/Horde/Plugins/Analytics/HordeServer.Analytics/Telemetry/Metrics/MetricCollection.cs
+++ b/Engine/Source/Programs/Horde/Plugins/Analytics/HordeServer.Analytics/Telemetry/Metrics/MetricCollection.cs
@@ -161,11 +161,29 @@ namespace HordeServer.Telemetry.Metrics
 			array.Remove(node);
 		}
 
+		PathResult Evaluate(JsonPath path, JsonNode node)
+		{
+			Stopwatch timer = Stopwatch.StartNew();
+			PathResult result = path.Evaluate(node);
+			timer.Stop();
+
+			if (timer.ElapsedMilliseconds > 50)
+			{
+				_logger.LogWarning("Querying {Path} against {Node} took {TimeMs}ms", path.ToString(), node.ToJsonString(), timer.ElapsedMilliseconds);
+			}
+			if (result.Matches != null && result.Matches.Count > 100)
+			{
+				_logger.LogWarning("Querying {Path} against {Node} returned {NumMatches}", path.ToString(), node.ToJsonString(), result.Matches.Count);
+			}
+
+			return result;
+		}
+
 		void AddEvent(TelemetryStoreId telemetryStoreId, MetricConfig metric, JsonNode node, JsonArray array)
 		{
 			if (metric.Filter != null)
 			{
-				PathResult filterResult = metric.Filter.Evaluate(array);
+				PathResult filterResult = Evaluate(metric.Filter, array);
 				if (filterResult.Error != null)
 				{
 					_logger.LogWarning("Error evaluating filter for metric {MetricId}: {Message}", metric.Id, filterResult.Error);
@@ -181,7 +199,7 @@ namespace HordeServer.Telemetry.Metrics
 
 			if (metric.Property != null)
 			{
-				PathResult result = metric.Property.Evaluate(node);
+				PathResult result = Evaluate(metric.Property, node);
 				if (result.Error != null)
 				{
 					_logger.LogWarning("Error evaluating filter for metric {MetricId}: {Message}", metric.Id, result.Error);
@@ -231,7 +249,7 @@ namespace HordeServer.Telemetry.Metrics
 				List<string> groupKeys = new List<string>();
 				foreach (JsonPath groupByPath in metric.GroupByPaths)
 				{
-					PathResult groupResult = groupByPath.Evaluate(node);
+					PathResult groupResult = Evaluate(groupByPath, node);
 
 					string groupKey;
 					if (groupResult.Error != null || groupResult.Matches == null || groupResult.Matches.Count == 0)
