Horde: Improve preflight autosubmit messaging
#skipci

diff --git a/Engine/Source/Programs/Horde/HordeDashboard/src/components/AutoSubmit.tsx b/Engine/Source/Programs/Horde/HordeDashboard/src/components/AutoSubmit.tsx
--- a/Engine/Source/Programs/Horde/HordeDashboard/src/components/AutoSubmit.tsx
+++ b/Engine/Source/Programs/Horde/HordeDashboard/src/components/AutoSubmit.tsx
@@ -35,7 +35,7 @@ export const AutosubmitInfo: React.FC<{ jobDetails: JobDetailsV2 }> = observer((
    if (job.state !== JobState.Complete) {
 
       if (job.preflightChange) {
-         return <Stack style={{ paddingTop: 24 }}><Checkbox label="Automatically submit preflight on success"
+         return <Stack style={{ paddingTop: 24 }}><Checkbox label={`Automatically submit changelist ${job.preflightChange} upon preflight success`}
             checked={state.autoSubmit}
             disabled={state.inflight || (job.startedByUserInfo?.id !== dashboard.userId) }
             onChange={(ev, checked) => {
diff --git a/Engine/Source/Programs/Horde/HordeDashboard/src/components/build/NewBuildV2.tsx b/Engine/Source/Programs/Horde/HordeDashboard/src/components/build/NewBuildV2.tsx
--- a/Engine/Source/Programs/Horde/HordeDashboard/src/components/build/NewBuildV2.tsx
+++ b/Engine/Source/Programs/Horde/HordeDashboard/src/components/build/NewBuildV2.tsx
@@ -1483,7 +1483,7 @@ const BuildParametersPanel: React.FC = observer(() => {
                   }
 
                })}
-               {!!options.preflightChange && <Stack style={{ paddingTop: 8 }}><Checkbox label="Automatically submit preflight on success"
+               {!!options.preflightChange && <Stack style={{ paddingTop: 8 }}><Checkbox label={`Automatically submit changelist ${options.preflightChange} upon preflight success`}
                   checked={options.autoSubmit}
                   onChange={(ev, checked) => {
                      options.onAutoSubmitChanged(checked ? true : false);
diff --git a/Engine/Source/Programs/Horde/HordeDashboard/src/components/jobDetailsV2/JobDetailPreflight.tsx b/Engine/Source/Programs/Horde/HordeDashboard/src/components/jobDetailsV2/JobDetailPreflight.tsx
--- a/Engine/Source/Programs/Horde/HordeDashboard/src/components/jobDetailsV2/JobDetailPreflight.tsx
+++ b/Engine/Source/Programs/Horde/HordeDashboard/src/components/jobDetailsV2/JobDetailPreflight.tsx
@@ -69,7 +69,7 @@ const AutosubmitInfo: React.FC<{ jobDetails: JobDetailsV2 }> = observer(({ jobDe
       if (jobData.preflightChange) {
          return <Stack tokens={{ childrenGap: 8 }} style={{ paddingBottom: 12, paddingTop: 18 }}>
             <Stack style={{ paddingTop: 8 }}>
-               <Checkbox label="Automatically submit preflight on success"
+               <Checkbox label={`Automatically submit changelist ${jobData.preflightChange} upon preflight success`}
                   checked={state.autoSubmit}
                   disabled={state.inflight || (jobData.startedByUserInfo?.id !== dashboard.userId)}
                   onChange={(ev, checked) => {
