Horde: AutomationHub now loads all streams by default when none are selected
#skipci

diff --git a/Engine/Source/Programs/Horde/HordeDashboard/src/backend/AutomationTestData.ts b/Engine/Source/Programs/Horde/HordeDashboard/src/backend/AutomationTestData.ts
--- a/Engine/Source/Programs/Horde/HordeDashboard/src/backend/AutomationTestData.ts
+++ b/Engine/Source/Programs/Horde/HordeDashboard/src/backend/AutomationTestData.ts
@@ -45,7 +45,7 @@ export class TestDataHandler {
     constructor(search: URLSearchParams) {
         TestDataHandler.instance = this;
         makeObservable(this);
-        this.state = this.stateFromSearch(search);        
+        this.state = this.stateFromSearch(search);
     }
 
     setAutomation(automation: string, initStreams: boolean = false) {
@@ -70,7 +70,7 @@ export class TestDataHandler {
             if (tests?.length) {
                 this.state.tests = tests.map(t => t.name);
             }
-            
+
             const suites = this.streamSuites;
             if (suites?.length) {
                 this.state.suites = suites.map(t => t.name);
@@ -380,7 +380,7 @@ export class TestDataHandler {
         }
     }
 
-    addStream(streamId: StreamId) {
+    addStream(streamId: StreamId, queryQuery: boolean = true) {
 
         const state = this.state;
 
@@ -396,7 +396,7 @@ export class TestDataHandler {
 
         state.streams.push(streamId);
 
-        if (this.updateSearch()) {
+        if (this.updateSearch(queryQuery)) {
             this.setUpdated();
         }
     }
@@ -649,7 +649,7 @@ export class TestDataHandler {
     }
 
 
-    updateSearch(): boolean {
+    updateSearch(queueQuery: boolean = true): boolean {
 
         const state = this.state;
 
@@ -723,10 +723,15 @@ export class TestDataHandler {
                 clearTimeout(this.queryTimeoutId);
             }
 
-            this.queryTimeoutId = setTimeout(() => {
-                this.queryTimeoutId = undefined;
-                this.query();
-            }, queryRefreshMS);
+            this.queryTimeoutId = undefined;
+
+            if (queueQuery) {
+
+                this.queryTimeoutId = setTimeout(() => {
+                    this.queryTimeoutId = undefined;
+                    this.query();
+                }, queryRefreshMS);
+            }
 
             this.search = search;
             return true;
@@ -1313,15 +1318,15 @@ export class TestDataHandler {
         return this.streamStatus.get(streamId);
     }
 
-    getStatusTests(streamIdIn? : string): GetTestResponse[] {
+    getStatusTests(streamIdIn?: string): GetTestResponse[] {
 
-        const testIds = new Set<TestId>();        
+        const testIds = new Set<TestId>();
 
         for (let [streamId, streamMap] of this.streamStatus) {
             if (streamIdIn && streamIdIn !== streamId) {
                 continue;
             }
-            
+
             for (let id of streamMap.keys()) {
                 const test = this.testMap.get(id);
                 if (test) {
@@ -1329,19 +1334,19 @@ export class TestDataHandler {
                 }
             }
         }
-    
-        return Array.from(testIds).map(id => this.testMap.get(id)!).sort((a, b) => (a.displayName ?? a.name).localeCompare((b.displayName ?? b.name)));        
+
+        return Array.from(testIds).map(id => this.testMap.get(id)!).sort((a, b) => (a.displayName ?? a.name).localeCompare((b.displayName ?? b.name)));
     }
 
-    getStatusSuites(streamIdIn? : string): GetTestSuiteResponse[] {
+    getStatusSuites(streamIdIn?: string): GetTestSuiteResponse[] {
 
-        const suiteIds = new Set<TestId>();        
+        const suiteIds = new Set<TestId>();
 
         for (let [streamId, streamMap] of this.streamStatus) {
             if (streamIdIn && streamIdIn !== streamId) {
                 continue;
             }
-            
+
             for (let id of streamMap.keys()) {
                 const suite = this.suiteMap.get(id);
                 if (suite) {
@@ -1349,11 +1354,11 @@ export class TestDataHandler {
                 }
             }
         }
-    
-        return Array.from(suiteIds).map(id => this.suiteMap.get(id)!).sort((a, b) =>a.name.localeCompare(b.name));        
+
+        return Array.from(suiteIds).map(id => this.suiteMap.get(id)!).sort((a, b) => a.name.localeCompare(b.name));
     }
 
-    getStatus(id: TestId | SuiteId): Map<StreamId, Map<MetaId, TestStatus>> | undefined  {
+    getStatus(id: TestId | SuiteId): Map<StreamId, Map<MetaId, TestStatus>> | undefined {
 
         if (!this.testMap.get(id) && !this.suiteMap.get(id)) {
             return undefined;
diff --git a/Engine/Source/Programs/Horde/HordeDashboard/src/components/AutomationView.tsx b/Engine/Source/Programs/Horde/HordeDashboard/src/components/AutomationView.tsx
--- a/Engine/Source/Programs/Horde/HordeDashboard/src/components/AutomationView.tsx
+++ b/Engine/Source/Programs/Horde/HordeDashboard/src/components/AutomationView.tsx
@@ -39,13 +39,20 @@ const StreamChooser: React.FC<{ handler: TestDataHandler }> = observer(({ handle
       return null;
    }
 
-   const cstreams: Set<string> = new Set(handler.state.streams ?? []);
-   const streams = handler.getAutomationStreams(automation).sort((a, b) => a.localeCompare(b)).map(a => projectStore.streamById(a)).filter(a => !!a) as StreamData[];
+   let cstreams: Set<string> = new Set(handler.state.streams ?? []);
+   const streams = handler.getAutomationStreams(automation).sort((a, b) => a.localeCompare(b)).map(a => projectStore.streamById(a)).filter(a => !!a) as StreamData[];   
 
    if (!streams.length) {
       return null;
    }
 
+   if (!cstreams.size)
+   {
+      streams.forEach(s => handler.addStream(s.id, false));
+      cstreams = new Set(handler.state.streams ?? [])
+      handler.query();
+   }
+
    const options: IComboBoxOption[] = [];
    streams.forEach(stream => {
       options.push({ key: stream.id, text: stream.fullname ?? "???" });
