Horde: Prevent external telemetry events sent to Horde from being forwarded to the data router. The engine will send this data by itself.

#rnx

diff --git a/Engine/Source/Programs/Horde/HordeServer.Shared/Telemetry/ITelemetryWriter.cs b/Engine/Source/Programs/Horde/HordeServer.Shared/Telemetry/ITelemetryWriter.cs
--- a/Engine/Source/Programs/Horde/HordeServer.Shared/Telemetry/ITelemetryWriter.cs
+++ b/Engine/Source/Programs/Horde/HordeServer.Shared/Telemetry/ITelemetryWriter.cs
@@ -37,7 +37,13 @@ namespace HordeServer.Telemetry
 	/// <param name="AppVersion">Version number of the application</param>
 	/// <param name="AppEnvironment">Name of the environment that the sending application is running in</param>
 	/// <param name="SessionId">Unique identifier for the current session</param>
-	public record class TelemetryRecordMeta(string? AppId = null, string? AppVersion = null, string? AppEnvironment = null, string? SessionId = null);
+	public record class TelemetryRecordMeta(string? AppId = null, string? AppVersion = null, string? AppEnvironment = null, string? SessionId = null)
+	{
+		/// <summary>
+		/// App id for events originating from Horde itself
+		/// </summary>
+		public const string HordeAppId = "Horde";
+	}
 
 	/// <summary>
 	/// Inactive implementation of <see cref="ITelemetryWriter"/>
diff --git a/Engine/Source/Programs/Horde/Plugins/Analytics/HordeServer.Analytics/Telemetry/Sinks/EpicTelemetrySink.cs b/Engine/Source/Programs/Horde/Plugins/Analytics/HordeServer.Analytics/Telemetry/Sinks/EpicTelemetrySink.cs
--- a/Engine/Source/Programs/Horde/Plugins/Analytics/HordeServer.Analytics/Telemetry/Sinks/EpicTelemetrySink.cs
+++ b/Engine/Source/Programs/Horde/Plugins/Analytics/HordeServer.Analytics/Telemetry/Sinks/EpicTelemetrySink.cs
@@ -242,6 +242,13 @@ namespace HordeServer.Telemetry.Sinks
 				return;
 			}
 
+			// Only forward internal Horde events to the data router; other apps (like the engine) will write to the data
+			// router by themselves, and we want need duplicates.
+			if (!String.Equals(telemetryEvent.RecordMeta.AppId, TelemetryRecordMeta.HordeAppId, StringComparison.OrdinalIgnoreCase))
+			{
+				return;
+			}
+
 			TelemetryRecordMeta recordMeta = telemetryEvent.RecordMeta;
 			if (recordMeta.AppId == null || recordMeta.AppVersion == null || recordMeta.AppEnvironment == null)
 			{
diff --git a/Engine/Source/Programs/Horde/Plugins/Analytics/HordeServer.Analytics/Telemetry/TelemetryManager.cs b/Engine/Source/Programs/Horde/Plugins/Analytics/HordeServer.Analytics/Telemetry/TelemetryManager.cs
--- a/Engine/Source/Programs/Horde/Plugins/Analytics/HordeServer.Analytics/Telemetry/TelemetryManager.cs
+++ b/Engine/Source/Programs/Horde/Plugins/Analytics/HordeServer.Analytics/Telemetry/TelemetryManager.cs
@@ -35,7 +35,7 @@ namespace HordeServer.Telemetry
 			_tracer = tracer;
 			_logger = loggerFactory.CreateLogger<TelemetryManager>();
 			_ticker = clock.AddTicker<EpicTelemetrySink>(TimeSpan.FromSeconds(30.0), FlushAsync, _logger);
-			_serverEventMetadata = new TelemetryRecordMeta("Horde", serverInfo.Version.ToString(), serverInfo.Environment, serverInfo.SessionId);
+			_serverEventMetadata = new TelemetryRecordMeta(TelemetryRecordMeta.HordeAppId, serverInfo.Version.ToString(), serverInfo.Environment, serverInfo.SessionId);
 
 			TelemetrySinkConfig telemetrySinks = serverSettings.Value.Sinks;
 			if (telemetrySinks.Epic?.Enabled ?? false)
