Horde: Process metrics with a single background task, and print diagnostics for slow metrics accumulated over 1 minute intervals.

#rnx

diff --git a/Engine/Source/Programs/Horde/Plugins/Analytics/HordeServer.Analytics/Telemetry/Metrics/MetricCollection.cs b/Engine/Source/Programs/Horde/Plugins/Analytics/HordeServer.Analytics/Telemetry/Metrics/MetricCollection.cs
--- a/Engine/Source/Programs/Horde/Plugins/Analytics/HordeServer.Analytics/Telemetry/Metrics/MetricCollection.cs
+++ b/Engine/Source/Programs/Horde/Plugins/Analytics/HordeServer.Analytics/Telemetry/Metrics/MetricCollection.cs
@@ -2,6 +2,7 @@
 
 using System.Diagnostics;
 using System.Text.Json.Nodes;
+using System.Threading.Channels;
 using EpicGames.Core;
 using EpicGames.Horde.Telemetry;
 using EpicGames.Horde.Telemetry.Metrics;
@@ -45,6 +46,7 @@ namespace HordeServer.Telemetry.Metrics
 		/// </summary>
 		public int BottomN { get; set; } = 0;
 	}
+
 	class MetricCollection : IMetricCollection, IHostedService, IAsyncDisposable
 	{
 		class MetricDocument : IMetric
@@ -75,15 +77,25 @@ namespace HordeServer.Telemetry.Metrics
 
 		record class SampleKey(TelemetryStoreId Store, MetricId Metric, string Group, DateTime Time);
 
+		record struct QualifiedMetricId(TelemetryStoreId StoreId, MetricId MetricId)
+		{
+			public override string ToString()
+				=> $"{StoreId}:{MetricId}";
+		}
+
 		readonly object _lockObject = new object();
 		readonly IMongoCollection<MetricDocument> _metrics;
 		readonly AsyncEvent _newDataEvent = new AsyncEvent();
 		readonly AsyncEvent _flushEvent = new AsyncEvent();
+		readonly BackgroundTask _processTask;
 		readonly BackgroundTask _flushTask;
 		readonly IClock _clock;
 		readonly IOptionsMonitor<AnalyticsConfig> _analyticsConfig;
 		readonly ILogger _logger;
 
+		int _numDroppedItems;
+		Channel<(TelemetryStoreId, JsonNode)> _events;
+		Dictionary<QualifiedMetricId, TimeSpan> _slowMetrics = new Dictionary<QualifiedMetricId, TimeSpan>();
 		Dictionary<SampleKey, List<double>> _queuedSamples = new Dictionary<SampleKey, List<double>>();
 
 		public MetricCollection(IMongoService mongoService, IClock clock, IOptionsMonitor<AnalyticsConfig> analyticsConfig, ILogger<MetricCollection> logger)
@@ -92,7 +104,10 @@ namespace HordeServer.Telemetry.Metrics
 			indexes.Add(MongoIndex.Create<MetricDocument>(keys => keys.Ascending(x => x.TelemetryStoreId).Descending(x => x.Time).Ascending(x => x.MetricId).Ascending(x => x.Group)));
 			_metrics = mongoService.GetCollection<MetricDocument>("Metrics", indexes);
 
-			_flushTask = new BackgroundTask(BackgroundTickAsync);
+			_events = Channel.CreateBounded<(TelemetryStoreId, JsonNode)>(new BoundedChannelOptions(10000));
+
+			_processTask = new BackgroundTask(BackgroundProcessAsync);
+			_flushTask = new BackgroundTask(BackgroundFlushAsync);
 			_clock = clock;
 			_analyticsConfig = analyticsConfig;
 			_logger = logger;
@@ -101,6 +116,7 @@ namespace HordeServer.Telemetry.Metrics
 		/// <inheritdoc/>
 		public Task StartAsync(CancellationToken cancellationToken)
 		{
+			_processTask.Start();
 			_flushTask.Start();
 			return Task.CompletedTask;
 		}
@@ -109,81 +125,71 @@ namespace HordeServer.Telemetry.Metrics
 		public async Task StopAsync(CancellationToken cancellationToken)
 		{
 			await _flushTask.StopAsync(cancellationToken);
+			await _processTask.StopAsync();
 		}
 
 		/// <inheritdoc/>
 		public async ValueTask DisposeAsync()
 		{
 			await _flushTask.DisposeAsync();
+			await _processTask.DisposeAsync();
 		}
 
 		/// <inheritdoc/>
 		public void AddEvent(TelemetryStoreId storeId, JsonNode node)
 		{
-			JsonArray array = new JsonArray { node };
+			if (!_events.Writer.TryWrite((storeId, node)))
+			{
+				Interlocked.Increment(ref _numDroppedItems);
+			}
+		}
+
+		async Task BackgroundProcessAsync(CancellationToken cancellationToken)
+		{
+			await foreach ((TelemetryStoreId storeId, JsonNode node) in _events.Reader.ReadAllAsync(cancellationToken))
+			{
+				ProcessEvent(storeId, node);
+			}
+		}
 
+		void ProcessEvent(TelemetryStoreId storeId, JsonNode node)
+		{
 			TelemetryStoreConfig? telemetryStoreConfig;
 			if (_analyticsConfig.CurrentValue.TryGetTelemetryStore(storeId, out telemetryStoreConfig))
 			{
-				Stopwatch timer = Stopwatch.StartNew();
-				Stopwatch metricTimer = new Stopwatch();
+				JsonArray array = new JsonArray { node };
 
-				MetricId slowestMetricId = default;
-				TimeSpan slowestMetricTime = TimeSpan.Zero;
+				Stopwatch timer = Stopwatch.StartNew();
+				TimeSpan slowEventTime = TimeSpan.FromMilliseconds(10);
 
 				foreach (MetricConfig metric in telemetryStoreConfig.Metrics)
 				{
-					metricTimer.Restart();
-					AddEvent(storeId, metric, node, array);
-					metricTimer.Stop();
+					TimeSpan startTime = timer.Elapsed;
+					ProcessEventInternal(storeId, metric, node, array);
 
-					if (metricTimer.Elapsed > slowestMetricTime)
+					TimeSpan elapsedTime = timer.Elapsed - startTime;
+					if (elapsedTime > slowEventTime)
 					{
-						slowestMetricId = metric.Id;
-						slowestMetricTime = metricTimer.Elapsed;
-					}
-				}
-
-				if (timer.Elapsed.TotalMilliseconds > 50)
-				{
-					string text = node.ToJsonString();
+						lock (_lockObject)
+						{
+							QualifiedMetricId qualifiedMetricId = new QualifiedMetricId(storeId, metric.Id);
 
-					const int MaxLength = 4096;
-					if (text.Length > MaxLength)
-					{
-						text = text.Substring(0, MaxLength) + $"... ({text.Length})";
+							TimeSpan currentTime;
+							_slowMetrics.TryGetValue(qualifiedMetricId, out currentTime);
+							_slowMetrics[qualifiedMetricId] = currentTime + elapsedTime;
+						}
 					}
-
-					_logger.LogWarning("Evaluating metrics for {StoreId} took {TimeMs}ms (slowest: {MetricId} @ {MetricTimeMs}ms): {Text}", storeId, (long)timer.ElapsedMilliseconds, slowestMetricId, (long)slowestMetricTime.TotalMilliseconds, text);
 				}
-			}
-
-			array.Remove(node);
-		}
-
-		PathResult Evaluate(JsonPath path, JsonNode node)
-		{
-			Stopwatch timer = Stopwatch.StartNew();
-			PathResult result = path.Evaluate(node);
-			timer.Stop();
 
-			if (timer.ElapsedMilliseconds > 50)
-			{
-				_logger.LogWarning("Querying {Path} against {Node} took {TimeMs}ms", path.ToString(), node.ToJsonString(), timer.ElapsedMilliseconds);
-			}
-			if (result.Matches != null && result.Matches.Count > 100)
-			{
-				_logger.LogWarning("Querying {Path} against {Node} returned {NumMatches}", path.ToString(), node.ToJsonString(), result.Matches.Count);
+				array.Remove(node);
 			}
-
-			return result;
 		}
 
-		void AddEvent(TelemetryStoreId telemetryStoreId, MetricConfig metric, JsonNode node, JsonArray array)
+		void ProcessEventInternal(TelemetryStoreId telemetryStoreId, MetricConfig metric, JsonNode node, JsonArray array)
 		{
 			if (metric.Filter != null)
 			{
-				PathResult filterResult = Evaluate(metric.Filter, array);
+				PathResult filterResult = metric.Filter.Evaluate(array);
 				if (filterResult.Error != null)
 				{
 					_logger.LogWarning("Error evaluating filter for metric {MetricId}: {Message}", metric.Id, filterResult.Error);
@@ -199,7 +205,7 @@ namespace HordeServer.Telemetry.Metrics
 
 			if (metric.Property != null)
 			{
-				PathResult result = Evaluate(metric.Property, node);
+				PathResult result = metric.Property.Evaluate(node);
 				if (result.Error != null)
 				{
 					_logger.LogWarning("Error evaluating filter for metric {MetricId}: {Message}", metric.Id, result.Error);
@@ -249,7 +255,7 @@ namespace HordeServer.Telemetry.Metrics
 				List<string> groupKeys = new List<string>();
 				foreach (JsonPath groupByPath in metric.GroupByPaths)
 				{
-					PathResult groupResult = Evaluate(groupByPath, node);
+					PathResult groupResult = groupByPath.Evaluate(node);
 
 					string groupKey;
 					if (groupResult.Error != null || groupResult.Matches == null || groupResult.Matches.Count == 0)
@@ -271,23 +277,9 @@ namespace HordeServer.Telemetry.Metrics
 
 				SampleKey key = new SampleKey(telemetryStoreId, metric.Id, group.ToString(), sampleTime);
 
-				Stopwatch timer = Stopwatch.StartNew();
 				lock (_lockObject)
 				{
-					timer.Stop();
-					if (timer.ElapsedMilliseconds > 100)
-					{
-						_logger.LogWarning("Waited {TimerMs}ms for lock over telemetry sample collection", timer.ElapsedMilliseconds);
-					}
-
-					timer.Restart();
 					QueueSampleValues(key, values);
-
-					timer.Stop();
-					if (timer.ElapsedMilliseconds > 100)
-					{
-						_logger.LogWarning("Queuing {NumSamples} for {StoreId}:{MetricId} took {TimerMs}", values.Count, telemetryStoreId, metric.Id, timer.ElapsedMilliseconds);
-					}
 				}
 
 				_newDataEvent.Set();
@@ -322,11 +314,27 @@ namespace HordeServer.Telemetry.Metrics
 			queuedValues.AddRange(values);
 		}
 
-		async Task BackgroundTickAsync(CancellationToken cancellationToken)
+		/// <inheritdoc/>
+		public async Task FlushAsync(CancellationToken cancellationToken)
+		{
+			// Handle all the events currently buffered for processing
+			int count = _events.Reader.Count;
+			while (count-- > 0 && _events.Reader.TryRead(out (TelemetryStoreId StoreId, JsonNode Node) item))
+			{
+				ProcessEvent(item.StoreId, item.Node);
+			}
+
+			// Write all the buffered samples
+			await WriteQueuedSamplesAsync(cancellationToken);
+		}
+
+		async Task BackgroundFlushAsync(CancellationToken cancellationToken)
 		{
 			Task newDataTask = _newDataEvent.Task;
 			Task flushTask = _flushEvent.Task;
 
+			Stopwatch timer = Stopwatch.StartNew();
+
 			while (!cancellationToken.IsCancellationRequested)
 			{
 				await newDataTask.WaitAsync(cancellationToken);
@@ -337,16 +345,35 @@ namespace HordeServer.Telemetry.Metrics
 
 				try
 				{
-					await FlushAsync(cancellationToken);
+					await WriteQueuedSamplesAsync(cancellationToken);
 				}
 				catch (Exception ex)
 				{
 					_logger.LogError(ex, "Error flushing telemetry data: {Message}", ex.Message);
 				}
+
+				TimeSpan slowMetricReportTime = TimeSpan.FromMinutes(1.0);
+				if (timer.Elapsed > slowMetricReportTime)
+				{
+					int numDroppedItems = Interlocked.Exchange(ref _numDroppedItems, 0);
+					if (numDroppedItems > 0)
+					{
+						lock (_lockObject)
+						{
+							_logger.LogWarning("{NumEvents} telemetry events were dropped due to queue being full.", numDroppedItems);
+							foreach ((QualifiedMetricId qualifiedMetricId, TimeSpan time) in _slowMetrics)
+							{
+								_logger.LogWarning("Metric {StoreId}:{MetricId} took {TimeSecs:n1}s ({Pct:n1}%) of analysis time in the last minute", qualifiedMetricId.StoreId, qualifiedMetricId.MetricId, time.TotalSeconds, time.TotalSeconds / slowMetricReportTime.TotalSeconds); 
+							}
+							_slowMetrics.Clear();
+						}
+					}
+					timer.Restart();
+				}
 			}
 		}
 
-		public async Task FlushAsync(CancellationToken cancellationToken)
+		async Task WriteQueuedSamplesAsync(CancellationToken cancellationToken)
 		{
 			// Copy the current sample buffer
 			Dictionary<SampleKey, List<double>> samples;
