Horde: Fix static analysis warnings.

#rnx

diff --git a/Engine/Source/Programs/Horde/Plugins/Analytics/HordeServer.Analytics/Telemetry/Metrics/MetricCollection.cs b/Engine/Source/Programs/Horde/Plugins/Analytics/HordeServer.Analytics/Telemetry/Metrics/MetricCollection.cs
--- a/Engine/Source/Programs/Horde/Plugins/Analytics/HordeServer.Analytics/Telemetry/Metrics/MetricCollection.cs
+++ b/Engine/Source/Programs/Horde/Plugins/Analytics/HordeServer.Analytics/Telemetry/Metrics/MetricCollection.cs
@@ -94,8 +94,8 @@ namespace HordeServer.Telemetry.Metrics
 		readonly ILogger _logger;
 
 		int _numDroppedItems;
-		Channel<(TelemetryStoreId, JsonNode)> _events;
-		Dictionary<QualifiedMetricId, TimeSpan> _slowMetrics = new Dictionary<QualifiedMetricId, TimeSpan>();
+		readonly Channel<(TelemetryStoreId, JsonNode)> _events;
+		readonly Dictionary<QualifiedMetricId, TimeSpan> _slowMetrics = new Dictionary<QualifiedMetricId, TimeSpan>();
 		Dictionary<SampleKey, List<double>> _queuedSamples = new Dictionary<SampleKey, List<double>>();
 
 		public MetricCollection(IMongoService mongoService, IClock clock, IOptionsMonitor<AnalyticsConfig> analyticsConfig, ILogger<MetricCollection> logger)
@@ -125,7 +125,7 @@ namespace HordeServer.Telemetry.Metrics
 		public async Task StopAsync(CancellationToken cancellationToken)
 		{
 			await _flushTask.StopAsync(cancellationToken);
-			await _processTask.StopAsync();
+			await _processTask.StopAsync(cancellationToken);
 		}
 
 		/// <inheritdoc/>
