Horde: Set the "background": true flag when creating new indexes, and replace the deletion of unused indexes with a warning message. We don't want old pods to delete indexes created by new pods.

diff --git a/Engine/Source/Programs/Horde/HordeServer/Server/MongoService.cs b/Engine/Source/Programs/Horde/HordeServer/Server/MongoService.cs
--- a/Engine/Source/Programs/Horde/HordeServer/Server/MongoService.cs
+++ b/Engine/Source/Programs/Horde/HordeServer/Server/MongoService.cs
@@ -696,15 +696,7 @@ namespace HordeServer.Server
 					// Drop any indexes that are no longer needed
 					foreach (string removeIndexName in removeIndexNames)
 					{
-						if (ReadOnlyMode)
-						{
-							_logger.LogWarning("Would drop unused index {CollectionName}.{IndexName} - skipping due to read-only setting.", collectionName, removeIndexName);
-						}
-						else
-						{
-							_logger.LogInformation("Dropping unused index {IndexName}", removeIndexName);
-							await collection.Indexes.DropOneAsync(removeIndexName, cancellationToken);
-						}
+						_logger.LogWarning("Collection {CollectionName} index {IndexName} is no longer required and should be removed. This may adversely affect performance.", collectionName, removeIndexName);
 					}
 
 					// Create all the new indexes
@@ -722,6 +714,7 @@ namespace HordeServer.Server
 							options.Name = createIndex.Name;
 							options.Unique = createIndex.Unique;
 							options.Sparse = createIndex.Sparse;
+							options.Background = true;
 
 							CreateIndexModel<T> model = new CreateIndexModel<T>(createIndex.Keys, options);
 
