Fix issue with inconsistent reservation time + remove unnecessary observer wrapper

#rnx

diff --git a/Engine/Source/Programs/Horde/HordeDashboard/src/backend/DeviceTelemetryHandler.ts b/Engine/Source/Programs/Horde/HordeDashboard/src/backend/DeviceTelemetryHandler.ts
--- a/Engine/Source/Programs/Horde/HordeDashboard/src/backend/DeviceTelemetryHandler.ts
+++ b/Engine/Source/Programs/Horde/HordeDashboard/src/backend/DeviceTelemetryHandler.ts
@@ -187,7 +187,8 @@ export class PoolTelemetryHandler {
             if (!platformData || !platformData.reserved) {
                return;
             }
-            tdata.reservations += Object.keys(platformData.reserved).length;
+            // platform 'reserved' property is a map of streams, each item contains an array of the reserved devices
+            tdata.reservations += Object.values(platformData.reserved).reduce((acc, item) => acc + item.length, 0);
          });
          timeSeries.push(tdata);
       });
diff --git a/Engine/Source/Programs/Horde/HordeDashboard/src/components/devices/DevicePoolTelemetry.tsx b/Engine/Source/Programs/Horde/HordeDashboard/src/components/devices/DevicePoolTelemetry.tsx
--- a/Engine/Source/Programs/Horde/HordeDashboard/src/components/devices/DevicePoolTelemetry.tsx
+++ b/Engine/Source/Programs/Horde/HordeDashboard/src/components/devices/DevicePoolTelemetry.tsx
@@ -1,5 +1,4 @@
 import { DefaultButton, DetailsList, DetailsListLayoutMode, Dropdown, FocusZone, FocusZoneDirection, IColumn, IconButton, IDropdownOption, Modal, ScrollablePane, ScrollbarVisibility, SelectionMode, Spinner, SpinnerSize, Stack, Text } from "@fluentui/react";
-import { observer } from "mobx-react-lite";
 import { useState, useRef, useEffect, useCallback, useMemo } from "react";
 import { Link } from "react-router-dom";
 import { deviceHealthRateBar } from "./DeviceHealthWidgets"
@@ -21,7 +20,7 @@ const scrollableOverflowOverrideClass = mergeStyles({
    overflow: "visible !important" as any
 });
 
-const HealthGraph: React.FC<{ handler: PoolTelemetryHandler, platform?: string, onChange: (options: object) => void }> = observer(({ handler, platform, onChange }) => {
+const HealthGraph: React.FC<{ handler: PoolTelemetryHandler, platform?: string, onChange: (options: object) => void }> = ({ handler, platform, onChange }) => {
 
    const selectedPlatformItem = Array.from(handler.platforms.values()).find((p) => p.id === platform);
    if (!selectedPlatformItem) {
@@ -150,9 +149,9 @@ const HealthGraph: React.FC<{ handler: PoolTelemetryHandler, platform?: string,
             </FocusZone>
          </Stack>
       </Stack>
-});
+}
 
-const DistributionGraph: React.FC<{ handler: PoolTelemetryHandler, platform?: string, onChange: (options: object) => void }> = observer(({ handler, platform, onChange }) => {
+const DistributionGraph: React.FC<{ handler: PoolTelemetryHandler, platform?: string, onChange: (options: object) => void }> = ({ handler, platform, onChange }) => {
    const [container, setContainer] = useState<HTMLDivElement | null>(null);
    const graph = useRef<PoolDistributionGraph>(new PoolDistributionGraph(handler));
 
@@ -216,9 +215,9 @@ const DistributionGraph: React.FC<{ handler: PoolTelemetryHandler, platform?: st
             </Stack.Item>
          </Stack>
       </Stack>
-});
+}
 
-export const DevicePoolGraph: React.FC<{ handler: PoolTelemetryHandler, pool?: string, platform?: string, onChange: (options: object) => void }> = observer(({ handler, pool, platform, onChange }) => {
+export const DevicePoolGraph: React.FC<{ handler: PoolTelemetryHandler, pool?: string, platform?: string, onChange: (options: object) => void }> = ({ handler, pool, platform, onChange }) => {
 
    const [container, setContainer] = useState<HTMLDivElement | null>(null);
    const graph = useRef<PoolTelemetryGraph>(new PoolTelemetryGraph(handler));
@@ -259,7 +258,7 @@ export const DevicePoolGraph: React.FC<{ handler: PoolTelemetryHandler, pool?: s
             <div id="pool_graph_container" className="horde-no-darktheme" style={{ userSelect: "none" }} ref={(ref: HTMLDivElement) => setContainer(ref)} onMouseEnter={() => { }} onMouseLeave={() => { }} />
          </Stack>
       </Stack>
-});
+}
 
 export const TelemetryDropdown: React.FC<{ handler: PoolTelemetryHandler, selectedPool?: string, selectedPlatform?: string, isPoolDisabled?: boolean, isPlatformDisabled?: boolean, onChange: (options: object) => void }> = ({ handler, selectedPool, selectedPlatform, isPoolDisabled, isPlatformDisabled, onChange }) => {
    const pools = Array.from(handler.pools.values()).sort((a, b) => a.name.localeCompare(b.name));
@@ -395,8 +394,5 @@ export const DevicePoolTelemetryModal: React.FC<{ poolsIn?: Array<string>, platf
          </Stack>
       </Modal>
    </Stack>
-
 }
 
-
-
