Horde: Updating release notes
#rnx
#skipci

diff --git a/Engine/Source/Programs/Horde/Docs/ReleaseNotes.md b/Engine/Source/Programs/Horde/Docs/ReleaseNotes.md
--- a/Engine/Source/Programs/Horde/Docs/ReleaseNotes.md
+++ b/Engine/Source/Programs/Horde/Docs/ReleaseNotes.md
@@ -1,5 +1,23 @@
 # Release Notes
 
+## 2025-01-14
+
+* Set the "background": true flag when creating new indexes, and replace the deletion of unused indexes with a warning message. We don't want old pods to delete indexes created by new pods. (39182560)
+* AutomationHub now loads all streams by default when none are selected (39089285)
+* Improve preflight autosubmit messaging (39088492)
+* Horde Log search improvements (39061326)
+* [UE-192322] Job dialog: improve navigation of long dropdown lists
+* Add a "Copy Artifact ID" to artifact button (39054641)
+* Few small changes to handling of chunked data nodes.
+  - ChunkedDataNodeRef now implements IHashedNodeRef, and extension methods now allow extracting chunked data to streams or files.
+  - Adding a FileEntry no longer requires an explicit length parameter (this is already contained in the payload data).
+  - Rename FileEntry.StreamHash to FileEntry.Hash (39004911)
+* Add telemetry spans to IssueService (38934459)
+* Include OidcToken app in Horde's .sln file (38920316)
+* Include lease ID in returned properties for allocated compute tasks (38917576)
+* Expose lease ID as an env var for jobs (38906012)
+* Send final capabilities update before agent session termination (38886583)
+
 ## 2024-12-17
 
 * Pass through OIDC scopes for API authorization.  Allows requesting custom delegated scopes for tools using OidcToken, like UGS interacting with Horde.
