Bring Pool load and device health to the main device page
- poll telemetry data only every 5 minutes
- add device group graph telemetry
- add health badge next to device avaliability

#rnx

diff --git a/Engine/Source/Programs/Horde/HordeDashboard/src/backend/DeviceHandler.ts b/Engine/Source/Programs/Horde/HordeDashboard/src/backend/DeviceHandler.ts
new file mode 100644
--- /dev/null
+++ b/Engine/Source/Programs/Horde/HordeDashboard/src/backend/DeviceHandler.ts
@@ -0,0 +1,312 @@
+import moment from "moment";
+import backend from "../backend";
+import { GetDevicePlatformResponse, GetDevicePoolResponse, GetDeviceReservationResponse, GetDevicePoolTelemetryResponse, GetDevicePlatformTelemetryResponse, GetDeviceResponse, GetUserResponse, DevicePoolTelemetryQuery } from "../backend/Api";
+import dashboard from "../backend/Dashboard";
+import { PollBase } from "../backend/PollBase";
+
+// in sort order
+export enum DeviceStatus {
+   Available = 0,
+   Reserved = 1,
+   Maintenance = 2,
+   Problem = 3,
+   Disabled = 4
+}
+
+export const DeviceStatusLabels = new Map<DeviceStatus, string>([
+   [DeviceStatus.Available, "Available"],
+   [DeviceStatus.Problem, "Problem"],
+   [DeviceStatus.Reserved, "Reserved"],
+   [DeviceStatus.Disabled, "Disabled"],
+   [DeviceStatus.Maintenance, "Maintenance"]
+]);
+
+export enum DeviceHealthStatus {
+   Unknown = 0,
+   Healthy = 1,
+   Unreliable = 2,
+   Problematic = 3
+}
+
+export const DeviceHealthLabels: Map<DeviceHealthStatus, string> = new Map<DeviceHealthStatus, string>([
+   [DeviceHealthStatus.Healthy, 'Healthy'],
+   [DeviceHealthStatus.Unreliable, 'Unreliable'],
+   [DeviceHealthStatus.Problematic, 'Problematic'],
+]);
+
+type DeviceTelemetry = { reservations: number, available: number, problems: number, disabled: number, deviceId: string };
+type DeviceUsageTimeSerie = { data: DeviceTelemetry[], date: Date };
+export type FilteredDeviceTelemetry = { reservations: number, available: number, problems: number, disabled: number, date: Date };
+type DeviceHealth = { reservations: number, problems: number, isAvailable: boolean };
+
+export class DeviceHandler extends PollBase {
+
+   constructor(pollTime = 5000) {
+      super(pollTime);
+   }
+
+   clear() {
+      super.stop();
+      this.loaded = false;
+      this.devices = [];
+      this.reservations = [];
+      this.pools.clear();
+      this.platforms.clear();
+      this.poolTelemetry = [];
+      this.poolTelemetryLastUpdate = undefined;
+      this.deviceHealth.clear();
+   }
+
+   async poll(): Promise<void> {
+
+      try {
+
+         this.devices = await backend.getDevices();
+         this.reservations = await backend.getDeviceReservations();
+
+         // Get telemetry from the last 3 days using a different frequency than the poll frequency
+         if (!this.poolTelemetryLastUpdate || Date.now() - this.poolTelemetryLastUpdate > this.poolTelemetryUpdateFrequency)
+         {
+            this.poolTelemetryLastUpdate = Date.now();
+            const date = moment().subtract(this.telemetryHoursRange, 'hours');
+            const deviceQuery: DevicePoolTelemetryQuery = { minCreateTime: date.toISOString(), count: 65536 };
+            this.processTelemetry(await backend.getDevicePoolTelemetry(deviceQuery));
+         }
+
+         // filter out utilization to entries which contain job id (legacy reservations don't)
+         this.devices.forEach(d => d.utilization = d.utilization?.filter(u => !!u.jobId));
+
+         const userIds = new Set<string>();
+         this.devices.forEach(d => {
+            if (d.checkedOutByUserId && !this.users.has(d.checkedOutByUserId)) {
+               userIds.add(d.checkedOutByUserId);
+            }
+         });
+
+         if (!this.users.get(dashboard.userId)) {
+            userIds.add(dashboard.userId);
+         }
+
+         if (userIds.size) {
+            const users = await backend.getUsers({ ids: Array.from(userIds) });
+            users.forEach(u => this.users.set(u.id, u));
+         }
+
+         if (!this.devices.length || this.devices.find(d => !this.platforms.has(d.platformId))) {
+            // we need to get platforms
+            const platforms = await backend.getDevicePlatforms();
+            this.platforms.clear();
+            platforms.forEach(p => {
+               this.platforms.set(p.id, p);
+            });
+
+            const pools = await backend.getDevicePools();
+            this.pools.clear();
+            pools.forEach(p => {
+               this.pools.set(p.id, p);
+            });
+
+         }
+         
+         this.loaded = true;         
+         this.setUpdated();
+
+      } catch (err) {
+
+      }
+
+   }
+
+   getUserDeviceCheckouts(userId: string): GetDeviceResponse[] {
+      return this.devices.filter(d => d.checkedOutByUserId === userId);
+   }
+
+   getReservation(device: GetDeviceResponse): GetDeviceReservationResponse | undefined {
+      return this.reservations.find(r => !!r.devices.find(d => d === device.id));
+   }
+
+   getDeviceWriteAccess(device: GetDeviceResponse): boolean {
+
+      const pool = this.pools.get(device.poolId);
+
+      if (!pool) {
+         return false;
+      }
+
+      return pool.writeAccess;
+
+   }
+
+   getDeviceStatus(device: GetDeviceResponse): DeviceStatus {
+
+      if (this.getReservation(device)) {
+         return DeviceStatus.Reserved;
+      }
+
+      if (!device.enabled) {
+         return DeviceStatus.Disabled;
+      }
+
+      if (device.maintenanceTime) {
+         return DeviceStatus.Maintenance;
+      }
+
+      if (device.problemTime) {
+
+         const end = moment.utc();
+         const d = moment.duration(end.diff(moment(device.problemTime)));
+
+         // note this must match reservation selection for problem time in backend
+         if (d.asMinutes() < dashboard.deviceProblemCooldownMinutes) {
+            return DeviceStatus.Problem;
+         }
+      }
+
+      return DeviceStatus.Available;
+   }
+
+   getDevices(platformId?: string) {
+      return this.devices.filter(d => !platformId || d.platformId === platformId);
+   }
+
+   processTelemetry(telemetryData: GetDevicePoolTelemetryResponse[]) {
+      const timeSeries: DeviceUsageTimeSerie[] = [];
+      const healthAggregates: Map<string, DeviceHealth> = new Map();
+      const getHealthDevice = (id) => {
+         let deviceHealth = healthAggregates.get(id);
+         if (!deviceHealth) {
+            deviceHealth = {
+               reservations: 0, problems: 0,
+               isAvailable: true // contiguous reservation flag
+            };
+            healthAggregates.set(id, deviceHealth);
+         }
+         return deviceHealth!;
+      }
+
+      const cleanedData = telemetryData.map(t => {
+         const date = new Date(t.createTimeUtc);
+         return { date: date, pools: t.telemetry }
+      }).sort((a, b) => a.date.getTime() - b.date.getTime());
+
+      cleanedData.forEach(d => {
+         const tserie: DeviceUsageTimeSerie = { data: [], date: d.date };
+         timeSeries.push(tserie);
+         Object.values(d.pools).forEach(telemetry => {
+            telemetry.forEach(platformData => {
+               const devices: Map<string, DeviceTelemetry> = new Map();
+               const getDevice = (id) => {
+                  let tdata = devices.get(id);
+                  if (!tdata) {
+                     tdata = {
+                        reservations: 0,
+                        available: 0,
+                        problems: 0,
+                        disabled: 0,
+                        deviceId: id
+                     };
+                     tserie.data.push(tdata);
+                  }
+                  return tdata!;
+               }
+               // update device health aggregates
+               Object.values(platformData.reserved ?? 0).forEach(stream => {
+                  stream.forEach(d => {
+                     const deviceHealth = getHealthDevice(d.deviceId);
+                     if (deviceHealth.isAvailable) {
+                        deviceHealth.reservations += 1;
+                        deviceHealth.isAvailable = false;
+                     }
+                     const device = getDevice(d.deviceId);
+                     device.reservations += 1;
+                  });
+               });
+               if (platformData.problem) {
+                  platformData.problem.forEach(deviceId => {
+                     const deviceHealth = getHealthDevice(deviceId);
+                     deviceHealth.problems += 1;
+                     deviceHealth.isAvailable = true; // reset contiguous reservation flag
+                     const device = getDevice(deviceId);
+                     device.problems += 1;
+                  });
+               }
+               if (platformData.available) {
+                  platformData.available.forEach(deviceId => {
+                     const deviceHealth = getHealthDevice(deviceId);
+                     deviceHealth.isAvailable = true;
+                     const device = getDevice(deviceId);
+                     device.available += 1;
+                  });
+               }
+            });
+         });
+      });
+
+      this.deviceHealth = new Map();
+      healthAggregates.forEach((value, key, map) => {
+         this.deviceHealth.set(key, this.computeDeviceHealth(value));
+      });
+
+      this.poolTelemetry = timeSeries;
+   }
+
+   getDevicesTelemetry(devices?: string[]) {
+      if (!devices) {
+         return undefined;
+      }
+
+      return this.poolTelemetry.map(telemetry => {
+         const tdata: FilteredDeviceTelemetry = {
+            reservations: 0,
+            available: 0,
+            problems: 0,
+            disabled: 0,
+            date: telemetry.date
+         };
+
+         return telemetry.data.filter(
+               t => devices.includes(t.deviceId)
+            ).reduce((acc, t) => {
+               acc.reservations += t.reservations;
+               acc.available += t.available;
+               acc.problems += t.problems;
+               acc.disabled += t.disabled;
+               return acc;
+            }, tdata);
+      });
+   }
+
+   computeDeviceHealth(health: DeviceHealth) {
+      if (health.reservations === 0 && health.problems === 0 ) {
+         return DeviceHealthStatus.Unknown;
+      }
+      const issueRate = (health.problems / (health.reservations + health.problems)) * 100;
+      if (issueRate > 60) return DeviceHealthStatus.Problematic;
+      if (issueRate > 40) return DeviceHealthStatus.Unreliable;
+      return DeviceHealthStatus.Healthy;
+   }
+
+   getDeviceHealth(deviceId: string) {
+      return this.deviceHealth.get(deviceId) ?? DeviceHealthStatus.Unknown;
+   }
+
+   loaded = false;
+
+   private devices: GetDeviceResponse[] = [];
+
+   platforms: Map<string, GetDevicePlatformResponse> = new Map();
+
+   pools: Map<string, GetDevicePoolResponse> = new Map();
+
+   users: Map<string, GetUserResponse> = new Map();
+
+   private reservations: GetDeviceReservationResponse[] = [];
+
+   private poolTelemetry: DeviceUsageTimeSerie[] = [];
+   poolTelemetryLastUpdate?: number;
+   private poolTelemetryUpdateFrequency = 5 * 60 * 1000; // 5 minutes
+   private deviceHealth: Map<string, DeviceHealthStatus> = new Map();
+
+   telemetryHoursRange = 72;
+
+}
diff --git a/Engine/Source/Programs/Horde/HordeDashboard/src/components/devices/DeviceEditor.tsx b/Engine/Source/Programs/Horde/HordeDashboard/src/components/devices/DeviceEditor.tsx
--- a/Engine/Source/Programs/Horde/HordeDashboard/src/components/devices/DeviceEditor.tsx
+++ b/Engine/Source/Programs/Horde/HordeDashboard/src/components/devices/DeviceEditor.tsx
@@ -1,11 +1,9 @@
 import { SpinnerSize, DialogType, IDropdownOption, MessageBarType, Text, Checkbox, DefaultButton, Dialog, DialogFooter, Dropdown, MessageBar, Modal, PrimaryButton, Spinner, Stack, TextField } from "@fluentui/react";
-import moment from "moment";
 import React, { useState } from "react";
 import backend from "../../backend";
-import { GetDevicePlatformResponse, GetDevicePoolResponse, GetDeviceReservationResponse, GetDeviceResponse, GetUserResponse } from "../../backend/Api";
-import dashboard from "../../backend/Dashboard";
-import { PollBase } from "../../backend/PollBase";
+import { GetDeviceResponse } from "../../backend/Api";
 import { getHordeStyling } from "../../styles/Styles";
+import { DeviceHandler, DeviceStatus } from "../../backend/DeviceHandler";
 
 type DeviceEditData = {
    // if defined, existing device
@@ -20,147 +18,6 @@ type DeviceEditData = {
    notes?: string;
 };
 
-// in sort order
-export enum DeviceStatus {
-   Available = 0,
-   Reserved = 1,
-   Maintenance = 2,
-   Problem = 3,
-   Disabled = 4
-}
-
-export class DeviceHandler extends PollBase {
-
-   constructor(pollTime = 5000) {
-      super(pollTime);
-   }
-
-   clear() {
-      super.stop();
-      this.loaded = false;
-      this.devices = [];
-      this.reservations = [];
-      this.pools.clear();
-      this.platforms.clear();
-   }
-
-   async poll(): Promise<void> {
-
-      try {
-
-         this.devices = await backend.getDevices();
-         this.reservations = await backend.getDeviceReservations();
-
-         // filter out utilization to entries which contain job id (legacy reservations don't)
-         this.devices.forEach(d => d.utilization = d.utilization?.filter(u => !!u.jobId));
-
-         const userIds = new Set<string>();
-         this.devices.forEach(d => {
-            if (d.checkedOutByUserId && !this.users.has(d.checkedOutByUserId)) {
-               userIds.add(d.checkedOutByUserId);
-            }
-         });
-
-         if (!this.users.get(dashboard.userId)) {
-            userIds.add(dashboard.userId);
-         }
-
-         if (userIds.size) {
-            const users = await backend.getUsers({ ids: Array.from(userIds) });
-            users.forEach(u => this.users.set(u.id, u));
-         }
-
-         if (!this.devices.length || this.devices.find(d => !this.platforms.has(d.platformId))) {
-            // we need to get platforms
-            const platforms = await backend.getDevicePlatforms();
-            this.platforms.clear();
-            platforms.forEach(p => {
-               this.platforms.set(p.id, p);
-            });
-
-            const pools = await backend.getDevicePools();
-            this.pools.clear();
-            pools.forEach(p => {
-               this.pools.set(p.id, p);
-            });
-
-         }
-         
-         this.loaded = true;         
-         this.setUpdated();
-
-      } catch (err) {
-
-      }
-
-   }
-
-   getUserDeviceCheckouts(userId: string): GetDeviceResponse[] {
-      return this.devices.filter(d => d.checkedOutByUserId === userId);
-   }
-
-   getReservation(device: GetDeviceResponse): GetDeviceReservationResponse | undefined {
-      return this.reservations.find(r => !!r.devices.find(d => d === device.id));
-   }
-
-   getDeviceWriteAccess(device: GetDeviceResponse): boolean {
-
-      const pool = this.pools.get(device.poolId);
-
-      if (!pool) {
-         return false;
-      }
-
-      return pool.writeAccess;
-
-   }
-
-   getDeviceStatus(device: GetDeviceResponse): DeviceStatus {
-
-      if (this.getReservation(device)) {
-         return DeviceStatus.Reserved;
-      }
-
-      if (!device.enabled) {
-         return DeviceStatus.Disabled;
-      }
-
-      if (device.maintenanceTime) {
-         return DeviceStatus.Maintenance;
-      }
-
-      if (device.problemTime) {
-
-         const end = moment.utc();
-         const d = moment.duration(end.diff(moment(device.problemTime)));
-
-         // note this must match reservation selection for problem time in backend
-         if (d.asMinutes() < dashboard.deviceProblemCooldownMinutes) {
-            return DeviceStatus.Problem;
-         }
-      }
-
-      return DeviceStatus.Available;
-   }
-
-   getDevices(platformId?: string) {
-      return this.devices.filter(d => !platformId || d.platformId === platformId);
-   }
-
-   loaded = false;
-
-   private devices: GetDeviceResponse[] = [];
-
-   platforms: Map<string, GetDevicePlatformResponse> = new Map();
-
-   pools: Map<string, GetDevicePoolResponse> = new Map();
-
-   users: Map<string, GetUserResponse> = new Map();
-
-   private reservations: GetDeviceReservationResponse[] = [];
-}
-
-
 export const DeviceEditor: React.FC<{ handler: DeviceHandler, deviceIn?: GetDeviceResponse | undefined, editNote?: boolean, onClose: () => void }> = ({ handler, deviceIn, editNote, onClose }) => {
 
    const [state, setState] = useState<{ device?: DeviceEditData, title?: string }>({});
diff --git a/Engine/Source/Programs/Horde/HordeDashboard/src/components/devices/DeviceHealthWidgets.tsx b/Engine/Source/Programs/Horde/HordeDashboard/src/components/devices/DeviceHealthWidgets.tsx
--- a/Engine/Source/Programs/Horde/HordeDashboard/src/components/devices/DeviceHealthWidgets.tsx
+++ b/Engine/Source/Programs/Horde/HordeDashboard/src/components/devices/DeviceHealthWidgets.tsx
@@ -1,8 +1,20 @@
-import { Stack, Text } from "@fluentui/react";
+import { Stack, Label, Text, TooltipHost, DirectionalHint, IGroup } from "@fluentui/react";
+import { useState, useRef } from "react";
 import dashboard, { StatusColor } from "../../backend/Dashboard";
+import { getHordeStyling } from "../../styles/Styles";
+import { DeviceHandler, DeviceHealthStatus, DeviceHealthLabels, FilteredDeviceTelemetry } from "../../backend/DeviceHandler";
+import { getHumanTime } from "../../base/utilities/timeUtils";
+
+import * as d3 from "d3";
+
+// Handle bad "@types/d3" types, fix if addressed upstream
+const _d3 = d3 as any;
+type SelectionType = d3.Selection<SVGGElement, unknown, null, undefined>;
+type ScalarTime = d3.ScaleTime<number, number>
 
 export const deviceHealthRateBar = (rate: number, width: number, height: number, text: string): JSX.Element => {
-   const colors = dashboard.getStatusColors();
+   const dashboardColors = dashboard.getStatusColors();
+
    return (
       <Stack verticalFill={true} verticalAlign="center" horizontal>
          <div style={{ width: width, height: height, verticalAlign: 'middle', marginRight: '3px' }}>
@@ -13,8 +25,8 @@ export const deviceHealthRateBar = (rate: number, width: number, height: number,
                   </clipPath>
                </defs>
                <g clipPath="url(#cutout)">
-                  <rect x={0} y={0} width={width} height={height} fill={colors.get(StatusColor.Success)} />
-                  <rect x={0} y={0} width={Math.ceil(width * rate / 100)} height={height} fill={colors.get(StatusColor.Failure)} />
+                  <rect x={0} y={0} width={width} height={height} fill={dashboardColors.get(StatusColor.Success)} />
+                  <rect x={0} y={0} width={Math.ceil(width * rate / 100)} height={height} fill={dashboardColors.get(StatusColor.Failure)} />
                </g>
             </svg>
          </div>
@@ -22,3 +34,192 @@ export const deviceHealthRateBar = (rate: number, width: number, height: number,
       </Stack>
    );
 }
+
+export const deviceHealthBadge = (health: DeviceHealthStatus, hourRange: number): JSX.Element => {
+   const dashboardColors = dashboard.getStatusColors();
+   const healthColors: Map<DeviceHealthStatus, string> = new Map<DeviceHealthStatus, string>([
+      [DeviceHealthStatus.Healthy, dashboardColors.get(StatusColor.Success)!],
+      [DeviceHealthStatus.Unreliable, dashboardColors.get(StatusColor.Warnings)!],
+      [DeviceHealthStatus.Problematic, dashboardColors.get(StatusColor.Failure)!],
+   ]);
+   const { modeColors } = getHordeStyling();
+
+   return <TooltipHost content={`${DeviceHealthLabels.get(health)} device < ${hourRange}h`} delay={0} closeDelay={250}
+         styles={{root: {display: 'inline-block'}}} directionalHint={DirectionalHint.bottomCenter}>
+            <Label style={{ padding: '3px 7px 3px 7px', fontSize: 9, fontFamily: "Horde Open Sans SemiBold", color: "#ffffff", backgroundColor: healthColors.get(health), cursor: 'inherit'}}>Health</Label>
+      </TooltipHost>
+}
+
+export class PoolTimelineGraph {
+   constructor(handler: DeviceHandler) {
+      this.handler = handler;
+      this.margin = { top: 0, right: 0, bottom: 0, left: 0 };
+   }
+
+   set(devices?: string[]) {
+      this.data = this.handler.getDevicesTelemetry(devices);
+      this.devices = devices;
+      this.lastUpdate = this.handler.poolTelemetryLastUpdate;
+   }
+
+   draw(container: HTMLDivElement, width: number, height: number, devices?: string[]) {
+
+      if (!!devices && !!this.devices && (devices.length !== this.devices.length || !devices.every(d => this.devices!.includes(d)))) {
+         this.forceRender = true;
+      } else if (this.lastUpdate != this.handler.poolTelemetryLastUpdate) {
+         this.forceRender = true;
+      }
+
+      if (this.hasRendered && !this.forceRender) {
+         return;
+      }
+      
+      this.container = container;
+      this.width = width;
+      this.height = height;
+      const margin = this.margin;
+      this.clear();
+
+      this.set(devices);
+
+      const data = this.data;
+
+      if (!data) {
+         return;
+      }
+
+      this.hasRendered = true;
+      this.forceRender = false;
+
+      const { modeColors } = getHordeStyling();
+      const dashboardColors = dashboard.getStatusColors();
+
+      const maxReservations = data.reduce((a, d) => Math.max(a, d.reservations + d.available), 1);
+
+      const svg = d3.select(container)
+         .append("svg")
+         .attr("width", width)
+         .attr("height", height)
+         .attr("viewBox", [0, 0, width, height] as any)
+         .style("background-color", modeColors.crumbs);
+
+      // plots
+      const xPlot = d3.scaleTime()
+         .domain(d3.extent(data, (d => d.date)) as any)
+         .range([this.margin.left, this.width - this.margin.right]);
+
+      const yPlot = d3.scaleLinear()
+         .domain([0, maxReservations])
+         .range([height - margin.bottom, margin.top]);
+
+
+      // bottom axis
+      const xAxis = svg.append("g")
+         .call(this.xAxisMethod as any, xPlot)
+         .style("pointer-events", "none");
+
+      // reservation graph
+      const area = d3.area()
+         .x((d: any) => xPlot(d.date))
+         .y0(yPlot(0))
+         .y1((d: any) => yPlot(d.reservations))
+         .curve(d3.curveStep);
+
+      const areaColor = dashboard.darktheme ?
+         (d3.color(modeColors.text) as any)?.darker(0.2)
+         : (d3.color(modeColors.text) as any)?.brighter(0.2);
+
+      svg.append("path")
+         .datum(data)
+         .attr("fill", areaColor)
+         .attr("stroke", modeColors.text).attr("stroke-width", 1.5)
+         .attr("opacity", 0.4)
+         .attr("d", area as any)
+         .style("pointer-events", "none");
+
+      // problems graph
+      const areaProblem = d3.area()
+         .x((d: any) => xPlot(d.date))
+         .y0(yPlot(0))
+         .y1((d: any) => yPlot(d.problems))
+         .curve(d3.curveStep);
+
+      const colorProblem = dashboardColors.get(StatusColor.Failure)!;
+
+      svg.append("path")
+         .datum(data)
+         .attr("fill", colorProblem)
+         .attr("d", areaProblem as any)
+         .style("pointer-events", "none");
+   }
+
+   private clear() {
+      if (!!this.container) {
+         d3.select(this.container).selectAll("*").remove();
+      }
+   }
+
+   private get xAxisMethod() {
+      return (g: SelectionType, x: ScalarTime) => {
+         const tickSet: Set<string> = new Set();
+
+         g.attr("transform", `translate(0,${this.height - this.margin.bottom})`)
+            .style("font-family", "Horde Open Sans Regular")
+            .style("font-size", "7px")
+            .attr("opacity", dashboard.darktheme? 0.4 : 0.5)
+            .call(d3.axisTop(x)
+               .ticks(8)
+               .tickFormat(d => {
+                  const v = (d as Date).toDateString();
+                  if (!tickSet.has(v)) {
+                     tickSet.add(v)
+                     return getHumanTime(d as Date)
+                  }
+
+                  return (d as Date).toLocaleString(undefined, { hour: 'numeric', hour12: true })
+               })
+               .tickSizeOuter(0)
+               .tickPadding(10));
+      }
+   }
+
+   handler: DeviceHandler;
+
+   container?: HTMLDivElement;
+   private margin: { top: number, right: number, bottom: number, left: number }
+   private width = 0;
+   private height = 0;
+
+   private data?: FilteredDeviceTelemetry[];
+   private devices?: string[];
+   private lastUpdate?: number = 0;
+
+   private hasRendered = false;
+   forceRender = false;
+}
+
+export const DeviceGroupHeader: React.FC<{ group: IGroup, handler: DeviceHandler, devices?: string[] }> = ({ group, handler, devices }) => {
+
+   const [container, setContainer] = useState<HTMLDivElement | null>(null);
+   const graph = useRef<PoolTimelineGraph>(new PoolTimelineGraph(handler));
+
+   if (!!container && !!devices) {
+      graph.current.draw(container, 300, 38, devices);
+   }
+
+   const backgroundColor = (dashboard.darktheme? d3.gray(10) : d3.gray(80)).toString();
+
+   const title = `${group.name} (${group.count})`;
+
+   return <TooltipHost content={`Device usage < ${handler.telemetryHoursRange}h`} delay={0} closeDelay={250}
+         styles={{root: {display: 'inline-block'}}} directionalHint={DirectionalHint.rightCenter}> 
+      <Stack style={{ position: "relative" }} horizontal verticalAlign="center">
+         <Text style={{width: 120, fontWeight: 'bold'}}>{title}</Text>
+         <div id={`pool_graph_container_${group.key}`}
+            className="horde-no-darktheme"
+            style={{ userSelect: "none", height: 38, backgroundColor: backgroundColor }}
+            ref={(ref: HTMLDivElement) => setContainer(ref)}
+            onMouseEnter={() => { }} onMouseLeave={() => { }} />
+      </Stack>
+   </TooltipHost>
+}
diff --git a/Engine/Source/Programs/Horde/HordeDashboard/src/components/devices/DeviceInfoView.tsx b/Engine/Source/Programs/Horde/HordeDashboard/src/components/devices/DeviceInfoView.tsx
--- a/Engine/Source/Programs/Horde/HordeDashboard/src/components/devices/DeviceInfoView.tsx
+++ b/Engine/Source/Programs/Horde/HordeDashboard/src/components/devices/DeviceInfoView.tsx
@@ -1,5 +1,4 @@
 import { DefaultButton, DetailsList, DetailsRow, IColumn, IconButton, IDetailsListProps, Label, Modal, PrimaryButton, SelectionMode, Spinner, SpinnerSize, Stack, Text, TextField } from "@fluentui/react";
-import { observer } from "mobx-react-lite";
 import moment from "moment";
 import { useState } from "react";
 import { Link, useNavigate } from "react-router-dom";
@@ -9,13 +8,13 @@ import dashboard, { StatusColor } from "../../backend/Dashboard";
 import { projectStore } from "../../backend/ProjectStore";
 import { displayTimeZone } from "../../base/utilities/timeUtils";
 import { getHordeStyling } from "../../styles/Styles";
-import { DeviceHandler } from "./DeviceEditor";
+import { DeviceHandler, DeviceStatusLabels } from "../../backend/DeviceHandler";
 import { deviceHealthRateBar } from "./DeviceHealthWidgets"
 import { DeviceStatusIcon } from "../StatusIcon";
 
 const streamIdToFullname = new Map<string, string>();
 
-export const DeviceInfoModal: React.FC<{ handler: DeviceHandler, deviceIn?: GetDeviceResponse | undefined, onEdit: (device: GetDeviceResponse) => void, onClose: () => void }> = observer(({ handler, deviceIn, onEdit, onClose }) => {
+export const DeviceInfoModal: React.FC<{ handler: DeviceHandler, deviceIn?: GetDeviceResponse | undefined, onEdit: (device: GetDeviceResponse) => void, onClose: () => void }> = ({ handler, deviceIn, onEdit, onClose }) => {
 
    type JobItem = {
       job: JobData,
@@ -235,6 +234,7 @@ export const DeviceInfoModal: React.FC<{ handler: DeviceHandler, deviceIn?: GetD
    }
 
    const failureRate = jobState.telemetryItems.length == 0? 0 : Math.floor(jobState.telemetryItems.filter(item => !!item.job.problemTimeUtc).length / jobState.telemetryItems.length * 100);
+   const status = handler.getDeviceStatus(deviceIn);
 
    return (<Modal isOpen={true} styles={{ main: { padding: 8, width: 1140, height: '800px', backgroundColor: '#FFFFFF' } }} className={hordeClasses.modal} onDismiss={() => { onClose() }}>
       <Stack styles={{ root: { paddingTop: 8, paddingLeft: 24, paddingRight: 12, paddingBottom: 8 } }}>
@@ -257,20 +257,23 @@ export const DeviceInfoModal: React.FC<{ handler: DeviceHandler, deviceIn?: GetD
                   <Stack tokens={{ childrenGap: 12 }}>
 
                      <Stack>
-                        <TextField label="Pool" value={deviceIn.poolId?.toUpperCase()} disabled={true} />
+                        <TextField label="Pool" value={deviceIn.poolId?.toUpperCase()} disabled={true} readOnly={true} />
                      </Stack>
 
                      <Stack>
-                        <TextField label="Address" value={deviceIn.address} disabled={true} />
+                        <TextField label="Model" value={deviceIn.modelId ? deviceIn.modelId : "Base"} disabled={true} readOnly={true} />
                      </Stack>
 
                      <Stack>
-                        <TextField label="Model" value={deviceIn.modelId ? deviceIn.modelId : "Base"} disabled={true} />
+                        <TextField label="Address" value={deviceIn.address} disabled={true} readOnly={true} />
                      </Stack>
 
+                     <Stack>
+                        <TextField label="Status" value={DeviceStatusLabels.get(status)} disabled={true} readOnly={true} />
+                     </Stack>
 
                      <Stack>
-                        <TextField label="Last Modified By" value={jobState.modifiedByUser ? jobState.modifiedByUser : ""} disabled={true} />
+                        <TextField label="Last Modified By" value={jobState.modifiedByUser ? jobState.modifiedByUser : ""} disabled={true} readOnly={true} />
                      </Stack>
 
                      <Stack horizontal style={{ paddingTop: 24 }}>
@@ -322,6 +325,4 @@ export const DeviceInfoModal: React.FC<{ handler: DeviceHandler, deviceIn?: GetD
       </Stack>
 
    </Modal >);
-});
-
-
+}
diff --git a/Engine/Source/Programs/Horde/HordeDashboard/src/components/devices/DeviceView.tsx b/Engine/Source/Programs/Horde/HordeDashboard/src/components/devices/DeviceView.tsx
--- a/Engine/Source/Programs/Horde/HordeDashboard/src/components/devices/DeviceView.tsx
+++ b/Engine/Source/Programs/Horde/HordeDashboard/src/components/devices/DeviceView.tsx
@@ -1,4 +1,4 @@
-import { DefaultButton, DetailsList, DetailsListLayoutMode, Dropdown, FocusZone, FocusZoneDirection, IColumn, Icon, IconButton, IGroup, Label, Link as FluentLink, mergeStyleSets, Modal, Pivot, PivotItem, PrimaryButton, ScrollablePane, ScrollbarVisibility, SelectionMode, Spinner, SpinnerSize, Stack, Text, TextField, themeRulesStandardCreator, TooltipHost, DirectionalHint, } from "@fluentui/react";
+import { DefaultButton, DetailsList, DetailsListLayoutMode, Dropdown, FocusZone, FocusZoneDirection, IColumn, Icon, IconButton, IGroup, GroupHeader, Label, Link as FluentLink, mergeStyleSets, Modal, Pivot, PivotItem, PrimaryButton, ScrollablePane, ScrollbarVisibility, SelectionMode, Spinner, SpinnerSize, Stack, Text, TextField, IDetailsGroupDividerProps, IGroupHeaderProps } from "@fluentui/react";
 import { observer } from "mobx-react-lite";
 import { observable, action, makeObservable } from "mobx";
 import moment from "moment";
@@ -10,8 +10,10 @@ import dashboard from "../../backend/Dashboard";
 import { useWindowSize } from "../../base/utilities/hooks";
 import { getNiceTime } from "../../base/utilities/timeUtils";
 import { Breadcrumbs } from "../Breadcrumbs";
-import { DeviceEditor, DeviceHandler, DeviceStatus } from "./DeviceEditor";
+import { DeviceHandler, DeviceStatus, DeviceStatusLabels, DeviceHealthStatus } from "../../backend/DeviceHandler";
+import { DeviceEditor } from "./DeviceEditor";
 import { DeviceInfoModal } from "./DeviceInfoView";
+import { deviceHealthBadge, DeviceGroupHeader } from "./DeviceHealthWidgets";
 import { DevicePoolTelemetryModal } from "./DevicePoolTelemetry";
 import { TopNav } from "../TopNav";
 import { getHordeStyling } from "../../styles/Styles";
@@ -197,14 +199,6 @@ class LocalState {
    private initialized = false;
 }
 
-export const StatusNames = new Map<DeviceStatus, string>([
-   [DeviceStatus.Available, "Available"],
-   [DeviceStatus.Problem, "Problem"],
-   [DeviceStatus.Reserved, "Reserved"],
-   [DeviceStatus.Disabled, "Disabled"],
-   [DeviceStatus.Maintenance, "Maintenance"]
-]);
-
 const dropDownStyle: any = () => {
 
    return {
@@ -289,7 +283,7 @@ const DevicePanel: React.FC = observer(() => {
    const { hordeClasses } = getHordeStyling();
    const hordeTheme = getHordeTheme();
 
-   if (handler.updated) { }
+   handler.subscribe();
 
    let message = "";
    let linkText = "device service configuration.";
@@ -387,7 +381,7 @@ const DevicePanel: React.FC = observer(() => {
 
       let color = status === DeviceStatus.Disabled ? "#777777 !important" : "#ffffff";
       let backgroundColor = StatusColors.get(status)!
-      let text = StatusNames.get(status)!;
+      let text = DeviceStatusLabels.get(status)!;
 
 
       if (device.checkedOutByUserId) {
@@ -398,7 +392,7 @@ const DevicePanel: React.FC = observer(() => {
       }
 
       // @todo: support deleting reservations
-      return <Stack verticalFill={true} verticalAlign="center">
+      return <Stack verticalFill={true} verticalAlign="center" horizontal>
          <PrimaryButton styles={{ root: { border: "0px", width: 84, height: 24, fontSize: 12, fontFamily: "Horde Open Sans SemiBold !important", selectors: { ".ms-Button-label": { color: color } }, backgroundColor: backgroundColor } }} text={text}
             onClick={(ev) => {
                ev.stopPropagation(); ev.preventDefault();
@@ -512,6 +506,10 @@ const DevicePanel: React.FC = observer(() => {
       // Automation Kit
 
       const r = handler.getReservation(device);
+      let health = automationTab? handler.getDeviceHealth(device.id) : DeviceHealthStatus.Unknown;
+      if (health === DeviceHealthStatus.Healthy && handler.getDeviceStatus(device) === DeviceStatus.Problem) {
+         health = DeviceHealthStatus.Unreliable;
+      }
 
       let url = "";
 
@@ -524,8 +522,6 @@ const DevicePanel: React.FC = observer(() => {
       return <Stack verticalFill={true} verticalAlign="center">
          <Stack horizontal tokens={{ childrenGap: 12 }}>
             {automationTab && <Stack verticalFill={true} verticalAlign="center">
-               <TooltipHost content="History" delay={0} closeDelay={250}
-                  styles={{root: {display: 'inline-block'}}} directionalHint={DirectionalHint.leftCenter}>
                   <div style={{ cursor: "pointer" }} onClick={(ev) => {
                      ev.stopPropagation(); ev.preventDefault();
                      setEditState({
@@ -533,12 +529,11 @@ const DevicePanel: React.FC = observer(() => {
                      })
                      localState.setHistoryItem(item.device.id);
                   }}>
-                     <Stack horizontal tokens={{ childrenGap: 18 }}>
+                     <Stack horizontal tokens={{ childrenGap: 8 }}>
                         <Icon style={{ padding: 2 }} iconName="History"/>
+                        {!!health && deviceHealthBadge(health, handler.telemetryHoursRange)}
                      </Stack>
                   </div>
-               </TooltipHost>
-
             </Stack>}
             {!!r && <Stack>
                {!!url && <Link to={url}><Text variant="small">Job Details</Text></Link>}
@@ -717,19 +712,35 @@ const DevicePanel: React.FC = observer(() => {
       else if (d.platformId !== curPlatform) {
          curPlatform = d.platformId;
          const key = `group_${d.platformId}`;
+         const groupDevices = devices.filter(cd => cd.device.platformId === d.platformId).map(d => d.device.id);
          newGroups.push({
             key: key,
             name: platforms.get(d.platformId)!.name,
+            data: {devices: groupDevices},
             startIndex: index,
-            count: devices.filter(cd => cd.device.platformId === d.platformId).length,
+            count: groupDevices.length,
             level: 0,
-            isCollapsed: groups.find(g => g.key === key)?.isCollapsed
+            // Collapse if it was collapsed by the user, or filtering is not an exact match, or there is more than one device listed
+            isCollapsed: (groups.find(g => g.key === key)?.isCollapsed ?? !exactMatch) && devices.length !== 1
          });
       }
    })
 
    groups = newGroups;
 
+   const groupTitleRenderer = (props: IGroupHeaderProps): JSX.Element => {
+      return <DeviceGroupHeader
+         group={props.group!}
+         handler={handler}
+         devices={props.group!.data!.devices}
+      />
+   }
+
+   const groupHeaderRenderer = (props: IDetailsGroupDividerProps): JSX.Element => {
+      const onRenderName = automationTab ? groupTitleRenderer : undefined;
+      return <GroupHeader {...props} onRenderName={onRenderName} />
+   }
+
    const pivotItems: JSX.Element[] = [];
 
    pivotItems.push(<PivotItem headerText="Shared" itemKey={pivotKeyShared} key={pivotKeyShared} />);
@@ -908,6 +919,7 @@ const DevicePanel: React.FC = observer(() => {
                         }}
                         items={devices}
                         groups={groups}
+                        groupProps={{onRenderHeader: groupHeaderRenderer}}
                         columns={columns}
                         selectionMode={SelectionMode.none}
                         layoutMode={DetailsListLayoutMode.justified}
