Horde: Fix filter for agent changes in ReleaseNotes tool.

#rnx

diff --git a/Engine/Source/Programs/Horde/Tools/ReleaseNotes/Program.cs b/Engine/Source/Programs/Horde/Tools/ReleaseNotes/Program.cs
--- a/Engine/Source/Programs/Horde/Tools/ReleaseNotes/Program.cs
+++ b/Engine/Source/Programs/Horde/Tools/ReleaseNotes/Program.cs
@@ -29,7 +29,8 @@ namespace CreateReleaseNotes
 		{
 			"...",
 			"-/Engine/Source/Programs/Horde/HordeDashboard/...",
-			"-/Engine/Source/Programs/Horde/HordeServer/...",
+			"-/Engine/Source/Programs/Horde/HordeServer*/...",
+			"-/Engine/Source/Programs/Horde/Plugins/*/HordeServer*/...",
 		};
 
 		static readonly string[] s_dashboardFilter =
