Horde: Use bulk operations when resolving a large number of imports from newly created blobs.

#rnx

diff --git a/Engine/Source/Programs/Horde/Plugins/Storage/HordeServer.Storage/Storage/StorageService.cs b/Engine/Source/Programs/Horde/Plugins/Storage/HordeServer.Storage/Storage/StorageService.cs
--- a/Engine/Source/Programs/Horde/Plugins/Storage/HordeServer.Storage/Storage/StorageService.cs
+++ b/Engine/Source/Programs/Horde/Plugins/Storage/HordeServer.Storage/Storage/StorageService.cs
@@ -786,15 +786,7 @@ namespace HordeServer.Storage
 		/// <inheritdoc/>
 		internal async Task<BlobInfo> AddBlobAsync(NamespaceId namespaceId, BlobLocator locator, IReadOnlyCollection<BlobLocator> imports, List<AliasInfo>? exports = null, CancellationToken cancellationToken = default)
 		{
-			List<ObjectId> importIds = new List<ObjectId>();
-			if (imports.Count > 0)
-			{
-				foreach (BlobLocator import in imports)
-				{
-					ObjectId importId = await FindOrAddShadowBlobAsync(namespaceId, import, cancellationToken);
-					importIds.Add(importId);
-				}
-			}
+			List<ObjectId> importIds = await FindOrAddShadowBlobsAsync(namespaceId, imports, cancellationToken);
 
 			FilterDefinition<BlobInfo> insertFilter =
 				Builders<BlobInfo>.Filter.Eq(x => x.NamespaceId, namespaceId) &
@@ -818,17 +810,64 @@ namespace HordeServer.Storage
 			return blobInfo;
 		}
 
-		async Task<ObjectId> FindOrAddShadowBlobAsync(NamespaceId namespaceId, BlobLocator locator, CancellationToken cancellationToken = default)
+		async Task<List<ObjectId>> FindOrAddShadowBlobsAsync(NamespaceId namespaceId, IReadOnlyCollection<BlobLocator> imports, CancellationToken cancellationToken = default)
 		{
-			FilterDefinition<BlobInfo> filter =
-				Builders<BlobInfo>.Filter.Eq(x => x.NamespaceId, namespaceId) &
-				Builders<BlobInfo>.Filter.Eq(x => x.Path, locator.ToString());
+			using TelemetrySpan span = _tracer.StartActiveSpan($"{nameof(StorageService)}.{nameof(FindOrAddShadowBlobsAsync)}");
+
+			List<ObjectId> importIds = new List<ObjectId>(imports.Count);
+			if(imports.Count > 0)
+			{
+				// Find all the unique import paths
+				HashSet<string> remaining = new HashSet<string>(imports.Count, StringComparer.Ordinal);
+				foreach (BlobLocator import in imports)
+				{
+					remaining.Add(import.ToString());
+				}
 
-			UpdateDefinition<BlobInfo> update =
-				Builders<BlobInfo>.Update.SetOnInsert(x => x.Shadow, true);
+				// Get the corresponding import ids
+				Dictionary<BlobLocator, ObjectId> locatorToId = new Dictionary<BlobLocator, ObjectId>(imports.Count);
+				for (; ;)
+				{
+					// Find all the existing blobs, in batches of 100
+					foreach (IReadOnlyList<string> batch in remaining.Batch(100))
+					{
+						FilterDefinition<BlobInfo> filter =
+							Builders<BlobInfo>.Filter.Eq(x => x.NamespaceId, namespaceId) &
+							Builders<BlobInfo>.Filter.In(x => x.Path, batch);
 
-			BlobInfo result = await _blobCollection.FindOneAndUpdateAsync(filter, update, new FindOneAndUpdateOptions<BlobInfo, BlobInfo> { IsUpsert = true, ReturnDocument = ReturnDocument.After }, cancellationToken);
-			return result.Id;
+						ProjectionDefinition<BlobInfo> projection =
+							Builders<BlobInfo>.Projection.Include(x => x.Id).Include(x => x.Path);
+
+						await foreach (BlobInfo blobInfo in _blobCollection.Find(filter).Project<BlobInfo>(projection).ToAsyncEnumerable(cancellationToken))
+						{
+							remaining.Remove(blobInfo.Path);
+							locatorToId.Add(new BlobLocator(blobInfo.Path), blobInfo.Id);
+						}
+					}
+
+					// Break out if there's nothing to upsert
+					if (remaining.Count == 0)
+					{
+						break;
+					}
+
+					// If there's anything left, try to insert shadow blobs for them. We'll do a separate query to figure out which succeeded.
+					List<InsertOneModel<BlobInfo>> inserts = new List<InsertOneModel<BlobInfo>>();
+					foreach (string import in remaining)
+					{
+						BlobInfo info = new BlobInfo(ObjectId.GenerateNewId(), namespaceId, new BlobLocator(import)) { Shadow = true };
+						inserts.Add(new InsertOneModel<BlobInfo>(info));
+					}
+					await _blobCollection.BulkWriteAsync(inserts, new BulkWriteOptions { IsOrdered = false }, cancellationToken);
+				}
+
+				// Find all the matching import ids in order
+				foreach (BlobLocator import in imports)
+				{
+					importIds.Add(locatorToId[import]);
+				}
+			}
+			return importIds;
 		}
 
 		internal async Task<BlobInfo?> GetBlobAsync(ObjectId id, CancellationToken cancellationToken = default)
