Fixing small issues where a warning during the SetupBuild step would be treated as an error in the notification flow and changing the log to be within a SectionBlock

diff --git a/Engine/Source/Programs/Horde/Plugins/Experimental/HordeServer.Experimental/Notifications/Sinks/ExperimentalSlackNotificationSink.cs b/Engine/Source/Programs/Horde/Plugins/Experimental/HordeServer.Experimental/Notifications/Sinks/ExperimentalSlackNotificationSink.cs
--- a/Engine/Source/Programs/Horde/Plugins/Experimental/HordeServer.Experimental/Notifications/Sinks/ExperimentalSlackNotificationSink.cs
+++ b/Engine/Source/Programs/Horde/Plugins/Experimental/HordeServer.Experimental/Notifications/Sinks/ExperimentalSlackNotificationSink.cs
@@ -358,7 +358,8 @@ namespace HordeServer.Notifications.Sinks
 			if (step.Name == IJob.SetupNodeName)
 			{
 				// Check if the Setup step failed to notify that the job can't be executed
-				if (step.Outcome != JobStepOutcome.Success)
+				bool wasSetupSuccess = (step.Outcome == JobStepOutcome.Success || step.Outcome == JobStepOutcome.Warnings);
+				if (!wasSetupSuccess)
 				{
 					await SendSetupFailedSlackMessageAsync(job, step, branch, jobStepEventData, configs, cancellationToken);
 					return;
@@ -1049,9 +1050,8 @@ namespace HordeServer.Notifications.Sinks
 
 			foreach (LogData log in logData.Take(MaxEventMessages))
 			{
-				ContextBlock logContext = new ContextBlock();
-				logContext.Elements.Add(new TextObject($"```{log.Message}```"));
-				slackBlockContainer.Blocks.Add(logContext);
+				SectionBlock logSection = new SectionBlock(new TextObject($"```{log.Message}```"));
+				slackBlockContainer.Blocks.Add(logSection);
 			}
 
 			if (logId is not null)
