Horde: Re-add check for EC2 IMDS in agent to prevent unnecessary exceptions being logged

diff --git a/Engine/Source/Programs/Horde/HordeAgent/Services/AwsInstanceLifecycleService.cs b/Engine/Source/Programs/Horde/HordeAgent/Services/AwsInstanceLifecycleService.cs
--- a/Engine/Source/Programs/Horde/HordeAgent/Services/AwsInstanceLifecycleService.cs
+++ b/Engine/Source/Programs/Horde/HordeAgent/Services/AwsInstanceLifecycleService.cs
@@ -169,9 +169,29 @@ class AwsInstanceLifecycleService : BackgroundService
 			return Ec2InstanceState.Unknown;
 		}
 	}
+	
+	private async Task<bool> IsImdsAvailableAsync(CancellationToken cancellationToken)
+	{
+		try
+		{
+			HttpResponseMessage res = await _httpClient.GetAsync(new Uri(BaseUri + "/"), cancellationToken);
+			return res.StatusCode == HttpStatusCode.OK;
+		}
+		catch (Exception)
+		{
+			// Timed out or other error. Can safely assume the metadata server is not available.
+			return false;
+		}
+	}
 
 	internal async Task MonitorInstanceLifecycleAsync(CancellationToken cancellationToken)
 	{
+		if (!await IsImdsAvailableAsync(cancellationToken))
+		{
+			_logger.LogInformation("AWS EC2 metadata server (IMDS) not available. Will not monitor EC2 lifecycle state");
+			return;
+		}
+		
 		_logger.LogInformation("Monitoring EC2 instance lifecycle state...");
 		while (!cancellationToken.IsCancellationRequested)
 		{
