Horde: Fix test for AwsInstanceLifecycleService

diff --git a/Engine/Source/Programs/Horde/HordeAgent.Tests/Services/AwsInstanceLifecycleServiceTests.cs b/Engine/Source/Programs/Horde/HordeAgent.Tests/Services/AwsInstanceLifecycleServiceTests.cs
--- a/Engine/Source/Programs/Horde/HordeAgent.Tests/Services/AwsInstanceLifecycleServiceTests.cs
+++ b/Engine/Source/Programs/Horde/HordeAgent.Tests/Services/AwsInstanceLifecycleServiceTests.cs
@@ -50,6 +50,7 @@ public class FakeAwsImds : HttpMessageHandler
 	public FakeAwsImds(IClock clock)
 	{
 		_clock = clock;
+		_responses["/"] = req => Res(HttpStatusCode.OK, "Metadata server");
 		_responses[AwsInstanceLifecycleService.ImdsTokenPath] = req =>
 		{
 			if (req.Method != HttpMethod.Put)
