Horde: Create a new IHordeClient for each job step, which writes log output to the step rather than the batch.

#rnx

diff --git a/Engine/Source/Programs/Horde/Drivers/JobDriver/Commands/Execution/ExecuteJobCommand.cs b/Engine/Source/Programs/Horde/Drivers/JobDriver/Commands/Execution/ExecuteJobCommand.cs
--- a/Engine/Source/Programs/Horde/Drivers/JobDriver/Commands/Execution/ExecuteJobCommand.cs
+++ b/Engine/Source/Programs/Horde/Drivers/JobDriver/Commands/Execution/ExecuteJobCommand.cs
@@ -47,9 +47,7 @@ namespace JobDriver.Commands.Execution
 
 			ExecuteJobTask executeTask = ExecuteJobTask.Parser.ParseFrom(Convert.FromBase64String(Task));
 
-			HordeClient hordeClient = _hordeClientFactory.Create(accessToken: executeTask.Token);
-
-			await JobExecutorHelpers.ExecuteAsync(hordeClient, WorkingDir, LeaseId, executeTask, _jobExecutorFactories, _driverSettings.Value, logger, CancellationToken.None);
+			await JobExecutorHelpers.ExecuteAsync(_hordeClientFactory, WorkingDir, LeaseId, executeTask, _jobExecutorFactories, _driverSettings.Value, logger, CancellationToken.None);
 			return 0;
 		}
 	}
diff --git a/Engine/Source/Programs/Horde/Drivers/JobDriver/Execution/JobExecutor.cs b/Engine/Source/Programs/Horde/Drivers/JobDriver/Execution/JobExecutor.cs
--- a/Engine/Source/Programs/Horde/Drivers/JobDriver/Execution/JobExecutor.cs
+++ b/Engine/Source/Programs/Horde/Drivers/JobDriver/Execution/JobExecutor.cs
@@ -51,7 +51,8 @@ namespace JobDriver.Execution
 
 	class JobExecutorOptions
 	{
-		public IHordeClient HordeClient { get; }
+		public HordeClientFactory HordeClientFactory { get; }
+		public string AccessToken { get; }
 		public DirectoryReference WorkingDir { get; }
 		public IReadOnlyList<ProcessToTerminate>? ProcessesToTerminate { get; }
 		public JobId JobId { get; }
@@ -59,9 +60,10 @@ namespace JobDriver.Execution
 		public LeaseId LeaseId { get; }
 		public RpcJobOptions JobOptions { get; }
 
-		public JobExecutorOptions(IHordeClient hordeClient, DirectoryReference workingDir, IReadOnlyList<ProcessToTerminate>? processesToTerminate, JobId jobId, JobStepBatchId batchId, LeaseId leaseId, RpcJobOptions jobOptions)
+		public JobExecutorOptions(HordeClientFactory hordeClientFactory, string accessToken, DirectoryReference workingDir, IReadOnlyList<ProcessToTerminate>? processesToTerminate, JobId jobId, JobStepBatchId batchId, LeaseId leaseId, RpcJobOptions jobOptions)
 		{
-			HordeClient = hordeClient;
+			HordeClientFactory = hordeClientFactory;
+			AccessToken = accessToken;
 			WorkingDir = workingDir;
 			ProcessesToTerminate = processesToTerminate;
 			JobId = jobId;
@@ -216,7 +218,9 @@ namespace JobDriver.Execution
 		/// </summary>
 		protected ILogger Logger { get; }
 
-		protected IHordeClient HordeClient { get; }
+		readonly HordeClientFactory _hordeClientFactory;
+		readonly string _accessToken;
+
 		protected DirectoryReference WorkingDir { get; }
 		protected JobId JobId { get; }
 		protected JobStepBatchId BatchId { get; }
@@ -250,7 +254,9 @@ namespace JobDriver.Execution
 
 		protected JobExecutor(JobExecutorOptions options, Tracer tracer, ILogger logger)
 		{
-			HordeClient = options.HordeClient;
+			_hordeClientFactory = options.HordeClientFactory;
+			_accessToken = options.AccessToken;
+
 			WorkingDir = options.WorkingDir;
 			JobId = options.JobId;
 			BatchId = options.BatchId;
@@ -285,8 +291,10 @@ namespace JobDriver.Execution
 			span.SetAttribute("horde.job.executor", JobOptions.Executor);
 			span.SetAttribute("horde.lease.id", LeaseId.ToString());
 			JobOptions.DecorateSpan(span);
-			
-			JobRpc.JobRpcClient jobRpc = await HordeClient.CreateGrpcClientAsync<JobRpc.JobRpcClient>(cancellationToken);
+
+			await using HordeClient hordeClient = _hordeClientFactory.Create(accessToken: _accessToken);
+
+			JobRpc.JobRpcClient jobRpc = await hordeClient.CreateGrpcClientAsync<JobRpc.JobRpcClient>(cancellationToken);
 
 			// Create a storage client for this session
 			logger.LogInformation("Executing jobId {JobId}, batchId {BatchId}, leaseId {LeaseId} with executor {Name}", JobId, BatchId, LeaseId, JobOptions.Executor);
@@ -295,13 +303,13 @@ namespace JobDriver.Execution
 			RpcBeginBatchResponse batch = await jobRpc.BeginBatchAsync(new RpcBeginBatchRequest(JobId, BatchId, LeaseId), cancellationToken: cancellationToken);
 			try
 			{
-				await ExecuteBatchAsync(batch, logger, NullLogger.Instance, cancellationToken);
+				await ExecuteBatchAsync(hordeClient, batch, logger, NullLogger.Instance, cancellationToken);
 			}
 			catch (Exception ex)
 			{
 				if (cancellationToken.IsCancellationRequested && IsCancellationException(ex))
 				{
-					if (!HordeClient.HasValidAccessToken())
+					if (!hordeClient.HasValidAccessToken())
 					{
 						logger.LogError(ex, "Connection to the server was lost; step aborted.");
 					}
@@ -328,9 +336,9 @@ namespace JobDriver.Execution
 		/// <summary>
 		/// Executes a batch
 		/// </summary>
-		async Task ExecuteBatchAsync(RpcBeginBatchResponse batch, ILogger logger, ILogger localLogger, CancellationToken cancellationToken)
+		async Task ExecuteBatchAsync(IHordeClient hordeClient, RpcBeginBatchResponse batch, ILogger logger, ILogger localLogger, CancellationToken cancellationToken)
 		{
-			JobRpc.JobRpcClient rpcClient = await HordeClient.CreateGrpcClientAsync<JobRpc.JobRpcClient>(cancellationToken);
+			JobRpc.JobRpcClient rpcClient = await hordeClient.CreateGrpcClientAsync<JobRpc.JobRpcClient>(cancellationToken);
 
 			await TerminateProcessHelper.TerminateProcessesAsync(TerminateCondition.BeforeBatch, WorkingDir, ProcessesToTerminate, logger, cancellationToken);
 
@@ -338,7 +346,7 @@ namespace JobDriver.Execution
 			logger.LogInformation("Initializing executor...");
 			using (logger.BeginIndentScope("  "))
 			{
-				await InitializeAsync(batch, logger, cancellationToken);
+				await InitializeAsync(hordeClient, batch, logger, cancellationToken);
 			}
 
 			try
@@ -411,7 +419,7 @@ namespace JobDriver.Execution
 					{
 						// Start writing to the log file
 #pragma warning disable CA2000 // Dispose objects before losing scope
-						await using (JobStepLogger stepLogger = new JobStepLogger(HordeClient, step.LogId, localLogger, JobId, BatchId, step.StepId, step.Warnings, LogLevel.Debug, logger))
+						await using (JobStepLogger stepLogger = new JobStepLogger(hordeClient, step.LogId, localLogger, JobId, BatchId, step.StepId, step.Warnings, LogLevel.Debug, logger))
 						{
 							// Execute the task
 							using CancellationTokenSource stepPollCancelSource = new CancellationTokenSource();
@@ -459,7 +467,7 @@ namespace JobDriver.Execution
 			{
 				if (cancellationToken.IsCancellationRequested && IsCancellationException(ex))
 				{
-					if (!HordeClient.HasValidAccessToken())
+					if (!hordeClient.HasValidAccessToken())
 					{
 						logger.LogError(ex, "Exception while executing batch: {Ex}", ex);
 					}
@@ -488,6 +496,23 @@ namespace JobDriver.Execution
 			}
 		}
 
+		class LoggerFactoryWrapper : ILoggerFactory
+		{
+			readonly ILogger _logger;
+
+			public LoggerFactoryWrapper(ILogger logger)
+				=> _logger = logger;
+
+			public void AddProvider(ILoggerProvider provider)
+			{ }
+
+			public ILogger CreateLogger(string categoryName)
+				=> _logger;
+
+			public void Dispose()
+			{ }
+		}
+
 		/// <summary>
 		/// Executes a step
 		/// </summary>
@@ -501,7 +526,9 @@ namespace JobDriver.Execution
 			using CancellationTokenSource combined = CancellationTokenSource.CreateLinkedTokenSource(cancellationToken, stepCancellationToken);
 			try
 			{
-				JobStepOutcome stepOutcome = await RunAsync(step, stepLogger, combined.Token);
+				using LoggerFactoryWrapper loggerFactory = new LoggerFactoryWrapper(stepLogger);
+				await using HordeClient stepHordeClient = _hordeClientFactory.Create(accessToken: _accessToken, loggerFactory: loggerFactory);
+				JobStepOutcome stepOutcome = await RunAsync(stepHordeClient, step, stepLogger, combined.Token);
 				return (stepOutcome, JobStepState.Completed);
 			}
 			catch (Exception ex)
@@ -571,16 +598,16 @@ namespace JobDriver.Execution
 			}
 		}
 
-		public virtual async Task InitializeAsync(RpcBeginBatchResponse batch, ILogger logger, CancellationToken cancellationToken)
+		public virtual async Task InitializeAsync(IHordeClient hordeClient, RpcBeginBatchResponse batch, ILogger logger, CancellationToken cancellationToken)
 		{
 			using TelemetrySpan span = Tracer.StartActiveSpan($"{nameof(JobExecutor)}.{nameof(InitializeAsync)}");
 			
 			Batch = batch;
-			JobRpc = await HordeClient.CreateGrpcClientAsync<JobRpc.JobRpcClient>(cancellationToken);
+			JobRpc = await hordeClient.CreateGrpcClientAsync<JobRpc.JobRpcClient>(cancellationToken);
 
-			_envVars[HordeHttpClient.HordeUrlEnvVarName] = HordeClient.ServerUrl.ToString();
+			_envVars[HordeHttpClient.HordeUrlEnvVarName] = hordeClient.ServerUrl.ToString();
 
-			string? accessToken = await HordeClient.GetAccessTokenAsync(false, cancellationToken);
+			string? accessToken = await hordeClient.GetAccessTokenAsync(false, cancellationToken);
 			if (accessToken != null)
 			{
 				_envVars[HordeHttpClient.HordeTokenEnvVarName] = accessToken;
@@ -672,11 +699,11 @@ namespace JobDriver.Execution
 			}
 		}
 
-		public virtual async Task<JobStepOutcome> RunAsync(JobStepInfo step, ILogger logger, CancellationToken cancellationToken)
+		public virtual async Task<JobStepOutcome> RunAsync(IHordeClient hordeClient, JobStepInfo step, ILogger logger, CancellationToken cancellationToken)
 		{
 			if (step.Name == "Setup Build")
 			{
-				if (await SetupAsync(step, logger, cancellationToken))
+				if (await SetupAsync(hordeClient, step, logger, cancellationToken))
 				{
 					return JobStepOutcome.Success;
 				}
@@ -687,7 +714,7 @@ namespace JobDriver.Execution
 			}
 			else
 			{
-				if (await ExecuteAsync(step, logger, cancellationToken))
+				if (await ExecuteAsync(hordeClient, step, logger, cancellationToken))
 				{
 					return JobStepOutcome.Success;
 				}
@@ -752,19 +779,19 @@ namespace JobDriver.Execution
 			}
 		}
 
-		protected abstract Task<bool> SetupAsync(JobStepInfo step, ILogger logger, CancellationToken cancellationToken);
+		protected abstract Task<bool> SetupAsync(IHordeClient hordeClient, JobStepInfo step, ILogger logger, CancellationToken cancellationToken);
 
-		protected abstract Task<bool> ExecuteAsync(JobStepInfo step, ILogger logger, CancellationToken cancellationToken);
+		protected abstract Task<bool> ExecuteAsync(IHordeClient hordeClient, JobStepInfo step, ILogger logger, CancellationToken cancellationToken);
 
 		const string SetupStepName = "setup";
 		const string BuildGraphTempStorageDir = "BuildGraph";
 
-		IStorageNamespace GetStorageNamespace(NamespaceId namespaceId, string? token)
+		static IStorageNamespace GetStorageNamespace(IHordeClient hordeClient, NamespaceId namespaceId, string? token)
 		{
-			return HordeClient.GetStorageNamespace(namespaceId, token);
+			return hordeClient.GetStorageNamespace(namespaceId, token);
 		}
 
-		protected virtual async Task<bool> SetupAsync(JobStepInfo step, DirectoryReference workspaceDir, bool? useP4, ILogger logger, CancellationToken cancellationToken)
+		protected virtual async Task<bool> SetupAsync(IHordeClient hordeClient, JobStepInfo step, DirectoryReference workspaceDir, bool? useP4, ILogger logger, CancellationToken cancellationToken)
 		{
 			FileReference definitionFile = FileReference.Combine(workspaceDir, $"{EnginePath}/Saved/Horde/Exported.json");
 
@@ -801,7 +828,7 @@ namespace JobDriver.Execution
 				arguments.AppendArgument("-Schema=", preprocessedSchemaFile.FullName);
 			}
 
-			int result = await ExecuteAutomationToolAsync(step, workspaceDir, arguments.ToString(), useP4, logger, cancellationToken);
+			int result = await ExecuteAutomationToolAsync(hordeClient, step, workspaceDir, arguments.ToString(), useP4, logger, cancellationToken);
 			if (result != 0)
 			{
 				return false;
@@ -839,7 +866,7 @@ namespace JobDriver.Execution
 				DecorateSpanWithArtifact(span, Prefix, artifactRequest, artifact);
 
 				// Write the data
-				IStorageNamespace storage = GetStorageNamespace(new NamespaceId(artifact.NamespaceId), artifact.Token);
+				IStorageNamespace storage = GetStorageNamespace(hordeClient, new NamespaceId(artifact.NamespaceId), artifact.Token);
 
 				Stopwatch timer = Stopwatch.StartNew();
 
@@ -1070,7 +1097,7 @@ namespace JobDriver.Execution
 			}
 		}
 
-		protected async Task<bool> ExecuteAsync(JobStepInfo step, DirectoryReference workspaceDir, bool? useP4, ILogger logger, CancellationToken cancellationToken)
+		protected async Task<bool> ExecuteAsync(IHordeClient hordeClient, JobStepInfo step, DirectoryReference workspaceDir, bool? useP4, ILogger logger, CancellationToken cancellationToken)
 		{
 			StringBuilder arguments = new StringBuilder("BuildGraph");
 			if (_preprocessScript)
@@ -1096,7 +1123,7 @@ namespace JobDriver.Execution
 
 				logger.LogInformation("Reading preprocessed script from {NamespaceId}:{RefName}", namespaceId, refName);
 
-				IStorageNamespace storage = GetStorageNamespace(namespaceId, artifact.Token);
+				IStorageNamespace storage = GetStorageNamespace(hordeClient, namespaceId, artifact.Token);
 
 				DirectoryNode node = await storage.ReadRefTargetAsync<DirectoryNode>(refName, cancellationToken: cancellationToken);
 				DirectoryNode? buildGraphDir = await node.TryOpenDirectoryAsync(BuildGraphTempStorageDir, cancellationToken: cancellationToken);
@@ -1134,11 +1161,11 @@ namespace JobDriver.Execution
 				}
 			}
 
-			bool result = await ExecuteWithTempStorageAsync(step, workspaceDir, arguments.ToString(), useP4, logger, cancellationToken);
+			bool result = await ExecuteWithTempStorageAsync(hordeClient, step, workspaceDir, arguments.ToString(), useP4, logger, cancellationToken);
 			return result;
 		}
 		
-		protected async Task CreateArtifactAsync(JobStepId stepId, ArtifactName name, ArtifactType type, DirectoryReference baseDir, IReadOnlyList<FileInfo> files, ILogger logger, CancellationToken cancellationToken)
+		protected async Task CreateArtifactAsync(IHordeClient hordeClient, JobStepId stepId, ArtifactName name, ArtifactType type, DirectoryReference baseDir, IReadOnlyList<FileInfo> files, ILogger logger, CancellationToken cancellationToken)
 		{
 			try
 			{
@@ -1152,9 +1179,9 @@ namespace JobDriver.Execution
 
 				RpcCreateJobArtifactResponseV2 artifact = await JobRpc.CreateArtifactV2Async(artifactRequest, cancellationToken: cancellationToken);
 				ArtifactId artifactId = ArtifactId.Parse(artifact.Id);
-				logger.LogInformation("Creating artifact {ArtifactId} '{ArtifactName}' ({ArtifactType}) from {NumFiles:n0} files ({TotalSize:n1}mb). Namespace {NamespaceId}, ref {RefName} ({Link})", artifactId, name, type, files.Count, totalSize / (1024.0 * 1024.0), artifact.NamespaceId, artifact.RefName, $"{HordeClient.ServerUrl}/api/v1/storage/{artifact.NamespaceId}/refs/{artifact.RefName}");
+				logger.LogInformation("Creating artifact {ArtifactId} '{ArtifactName}' ({ArtifactType}) from {NumFiles:n0} files ({TotalSize:n1}mb). Namespace {NamespaceId}, ref {RefName} ({Link})", artifactId, name, type, files.Count, totalSize / (1024.0 * 1024.0), artifact.NamespaceId, artifact.RefName, $"{hordeClient.ServerUrl}/api/v1/storage/{artifact.NamespaceId}/refs/{artifact.RefName}");
 
-				IStorageNamespace storage = GetStorageNamespace(new NamespaceId(artifact.NamespaceId), artifact.Token);
+				IStorageNamespace storage = GetStorageNamespace(hordeClient, new NamespaceId(artifact.NamespaceId), artifact.Token);
 				Stopwatch timer = Stopwatch.StartNew();
 
 				IHashedBlobRef<DirectoryNode> rootRef;
@@ -1183,7 +1210,7 @@ namespace JobDriver.Execution
 			}
 		}
 
-		private async Task<bool> ExecuteWithTempStorageAsync(JobStepInfo step, DirectoryReference workspaceDir, string arguments, bool? useP4, ILogger logger, CancellationToken cancellationToken)
+		private async Task<bool> ExecuteWithTempStorageAsync(IHordeClient hordeClient, JobStepInfo step, DirectoryReference workspaceDir, string arguments, bool? useP4, ILogger logger, CancellationToken cancellationToken)
 		{
 			DirectoryReference manifestDir = DirectoryReference.Combine(workspaceDir, $"{EnginePath}/Saved/BuildGraph");
 
@@ -1204,7 +1231,7 @@ namespace JobDriver.Execution
 				string nodeName = input.Substring(0, slashIdx);
 				string tagName = input.Substring(slashIdx + 1);
 
-				TempStorageTagManifest fileList = await TempStorage.RetrieveTagAsync(HordeClient, JobId, step.StepId, nodeName, tagName, manifestDir, logger, cancellationToken);
+				TempStorageTagManifest fileList = await TempStorage.RetrieveTagAsync(hordeClient, JobId, step.StepId, nodeName, tagName, manifestDir, logger, cancellationToken);
 				tagNameToFileSet[tagName] = fileList.ToFileSet(workspaceDir);
 				inputStorageBlocks.UnionWith(fileList.Blocks);
 			}
@@ -1221,7 +1248,7 @@ namespace JobDriver.Execution
 				span.SetAttribute(Prefix + "blocks", inputStorageBlocks.Count);
 				foreach (TempStorageBlockRef inputStorageBlock in inputStorageBlocks)
 				{
-					TempStorageBlockManifest manifest = await TempStorage.RetrieveBlockAsync(HordeClient, JobId, step.StepId, inputStorageBlock.NodeName, inputStorageBlock.OutputName, workspaceDir, manifestDir, logger, cancellationToken);
+					TempStorageBlockManifest manifest = await TempStorage.RetrieveBlockAsync(hordeClient, JobId, step.StepId, inputStorageBlock.NodeName, inputStorageBlock.OutputName, workspaceDir, manifestDir, logger, cancellationToken);
 					inputManifests[inputStorageBlock] = manifest;
 				}
 				span.SetAttribute(Prefix + "size", inputManifests.Sum(x => x.Value.GetTotalSize()));
@@ -1247,7 +1274,7 @@ namespace JobDriver.Execution
 			}
 
 			// Run UAT
-			if (await ExecuteAutomationToolAsync(step, workspaceDir, arguments, useP4, logger, cancellationToken) != 0)
+			if (await ExecuteAutomationToolAsync(hordeClient, step, workspaceDir, arguments, useP4, logger, cancellationToken) != 0)
 			{
 				return false;
 			}
@@ -1365,10 +1392,10 @@ namespace JobDriver.Execution
 
 				RpcCreateJobArtifactResponseV2 artifact = await JobRpc.CreateArtifactV2Async(artifactRequest, cancellationToken: cancellationToken);
 				ArtifactId artifactId = ArtifactId.Parse(artifact.Id);
-				logger.LogInformation("Creating artifact {ArtifactId} '{ArtifactName}' ({ArtifactType}) from {NumFiles:n0} files ({TotalSize:n1}mb). Namespace: {NamespaceId}, ref {RefName} ({RefUrl})", artifactId, artifactRequest.Name, ArtifactType.StepOutput, outputFileToBlockName.Count, totalSize / (1024.0 * 1024.0), artifact.NamespaceId, artifact.RefName, $"{HordeClient.ServerUrl.ToString().TrimEnd('/')}/api/v1/storage/{artifact.NamespaceId}/refs/{artifact.RefName}");
+				logger.LogInformation("Creating artifact {ArtifactId} '{ArtifactName}' ({ArtifactType}) from {NumFiles:n0} files ({TotalSize:n1}mb). Namespace: {NamespaceId}, ref {RefName} ({RefUrl})", artifactId, artifactRequest.Name, ArtifactType.StepOutput, outputFileToBlockName.Count, totalSize / (1024.0 * 1024.0), artifact.NamespaceId, artifact.RefName, $"{hordeClient.ServerUrl.ToString().TrimEnd('/')}/api/v1/storage/{artifact.NamespaceId}/refs/{artifact.RefName}");
 				DecorateSpanWithArtifact(span, Prefix, artifactRequest, artifact);
 
-				IStorageNamespace storage = GetStorageNamespace(new NamespaceId(artifact.NamespaceId), artifact.Token);
+				IStorageNamespace storage = GetStorageNamespace(hordeClient, new NamespaceId(artifact.NamespaceId), artifact.Token);
 
 				// Upload the data
 				Stopwatch timer = Stopwatch.StartNew();
@@ -1454,11 +1481,11 @@ namespace JobDriver.Execution
 
 					RpcCreateJobArtifactResponseV2 artifact = await JobRpc.CreateArtifactV2Async(artifactRequest, cancellationToken: cancellationToken);
 					ArtifactId artifactId = ArtifactId.Parse(artifact.Id);
-					logger.LogInformation("Creating artifact {ArtifactId} '{ArtifactName}' ({ArtifactType}) from {NumFiles:n0} files ({TotalSize:n1}mb). Namespace {NamespaceId}, ref {RefName} ({RefUrl})", artifactId, artifactRequest.Name, ArtifactType.StepOutput, fileList.Count, totalSize / (1024.0 * 1024.0), artifact.NamespaceId, artifact.RefName, $"{HordeClient.ServerUrl.ToString().TrimEnd('/')}/api/v1/storage/{artifact.NamespaceId}/refs/{artifact.RefName}");
+					logger.LogInformation("Creating artifact {ArtifactId} '{ArtifactName}' ({ArtifactType}) from {NumFiles:n0} files ({TotalSize:n1}mb). Namespace {NamespaceId}, ref {RefName} ({RefUrl})", artifactId, artifactRequest.Name, ArtifactType.StepOutput, fileList.Count, totalSize / (1024.0 * 1024.0), artifact.NamespaceId, artifact.RefName, $"{hordeClient.ServerUrl.ToString().TrimEnd('/')}/api/v1/storage/{artifact.NamespaceId}/refs/{artifact.RefName}");
 
 					DecorateSpanWithArtifact(span, Prefix, artifactRequest, artifact);
 
-					IStorageNamespace storage = GetStorageNamespace(new NamespaceId(artifact.NamespaceId), artifact.Token);
+					IStorageNamespace storage = GetStorageNamespace(hordeClient, new NamespaceId(artifact.NamespaceId), artifact.Token);
 
 					// Upload the data
 					Stopwatch timer = Stopwatch.StartNew();
@@ -1492,7 +1519,7 @@ namespace JobDriver.Execution
 			return true;
 		}
 
-		protected async Task<int> ExecuteAutomationToolAsync(JobStepInfo step, DirectoryReference workspaceDir, string? arguments, bool? useP4, ILogger logger, CancellationToken cancellationToken)
+		protected async Task<int> ExecuteAutomationToolAsync(IHordeClient hordeClient, JobStepInfo step, DirectoryReference workspaceDir, string? arguments, bool? useP4, ILogger logger, CancellationToken cancellationToken)
 		{
 			int result;
 			using TelemetrySpan span = Tracer.StartActiveSpan("BuildGraph");
@@ -1513,7 +1540,7 @@ namespace JobDriver.Execution
 
 			if (RuntimeInformation.IsOSPlatform(OSPlatform.Windows))
 			{
-				result = await ExecuteCommandAsync(step, workspaceDir, Environment.GetEnvironmentVariable("COMSPEC") ?? "cmd.exe", $"/C \"\"{workspaceDir}\\{EnginePath}\\Build\\BatchFiles\\RunUAT.bat\" {arguments}\"", span, logger, cancellationToken);
+				result = await ExecuteCommandAsync(hordeClient, step, workspaceDir, Environment.GetEnvironmentVariable("COMSPEC") ?? "cmd.exe", $"/C \"\"{workspaceDir}\\{EnginePath}\\Build\\BatchFiles\\RunUAT.bat\" {arguments}\"", span, logger, cancellationToken);
 			}
 			else if (RuntimeInformation.IsOSPlatform(OSPlatform.Linux))
 			{
@@ -1524,11 +1551,11 @@ namespace JobDriver.Execution
 					args = $"\"{workspaceDir}/{EnginePath}/Build/BatchFiles/RunWineUAT.sh\" {arguments}";
 				}
 
-				result = await ExecuteCommandAsync(step, workspaceDir, "/bin/bash", args, span, logger, cancellationToken);
+				result = await ExecuteCommandAsync(hordeClient, step, workspaceDir, "/bin/bash", args, span, logger, cancellationToken);
 			}
 			else if (RuntimeInformation.IsOSPlatform(OSPlatform.OSX))
 			{
-				result = await ExecuteCommandAsync(step, workspaceDir, "/bin/sh", $"\"{workspaceDir}/{EnginePath}/Build/BatchFiles/RunUAT.sh\" {arguments}", span, logger, cancellationToken);
+				result = await ExecuteCommandAsync(hordeClient, step, workspaceDir, "/bin/sh", $"\"{workspaceDir}/{EnginePath}/Build/BatchFiles/RunUAT.sh\" {arguments}", span, logger, cancellationToken);
 			}
 			else
 			{
@@ -1793,7 +1820,7 @@ namespace JobDriver.Execution
 			return name;
 		}
 
-		async Task<int> ExecuteCommandAsync(JobStepInfo step, DirectoryReference workspaceDir, string fileName, string arguments, TelemetrySpan buildGraphSpan, ILogger jobLogger, CancellationToken cancellationToken)
+		async Task<int> ExecuteCommandAsync(IHordeClient hordeClient, JobStepInfo step, DirectoryReference workspaceDir, string fileName, string arguments, TelemetrySpan buildGraphSpan, ILogger jobLogger, CancellationToken cancellationToken)
 		{
 			// Method for expanding environment variable properties related to this step
 			Dictionary<string, string> properties = new Dictionary<string, string>(StringComparer.OrdinalIgnoreCase);
@@ -1939,7 +1966,7 @@ namespace JobDriver.Execution
 					jobLogger.LogInformation(ex, "Unable to parse/upload telemetry files: {Message}", ex.Message);
 				}
 
-				await CreateArtifactAsync(step.StepId, TempStorage.GetArtifactNameForNode(step.Name), ArtifactType.StepTrace, workspaceDir, telemetryFiles, jobLogger, cancellationToken);
+				await CreateArtifactAsync(hordeClient, step.StepId, TempStorage.GetArtifactNameForNode(step.Name), ArtifactType.StepTrace, workspaceDir, telemetryFiles, jobLogger, cancellationToken);
 
 				foreach (FileInfo telemetryFile in telemetryFiles)
 				{
@@ -1982,14 +2009,14 @@ namespace JobDriver.Execution
 				jobLogger.LogInformation("Found {NumResults} test results", combinedTestData.Count);
 				await UploadTestDataAsync(step.StepId, combinedTestData);
 
-				await CreateArtifactAsync(step.StepId, TempStorage.GetArtifactNameForNode(step.Name), ArtifactType.StepTestData, workspaceDir, testDataFiles, jobLogger, cancellationToken);
+				await CreateArtifactAsync(hordeClient, step.StepId, TempStorage.GetArtifactNameForNode(step.Name), ArtifactType.StepTestData, workspaceDir, testDataFiles, jobLogger, cancellationToken);
 			}
 
 			DirectoryInfo logDirInfo = logDir.ToDirectoryInfo();
 			if (logDirInfo.Exists)
 			{
 				FileInfo[] artifactFiles = logDirInfo.GetFiles("*", SearchOption.AllDirectories);
-				await CreateArtifactAsync(step.StepId, TempStorage.GetArtifactNameForNode(step.Name), ArtifactType.StepSaved, workspaceDir, artifactFiles, jobLogger, cancellationToken);
+				await CreateArtifactAsync(hordeClient, step.StepId, TempStorage.GetArtifactNameForNode(step.Name), ArtifactType.StepSaved, workspaceDir, artifactFiles, jobLogger, cancellationToken);
 
 				foreach (FileReference reportFile in artifactFiles.Select(x => new FileReference(x)).Where(x => x.HasExtension(".report.json")))
 				{
@@ -2273,7 +2300,7 @@ namespace JobDriver.Execution
 
 	static class JobExecutorHelpers
 	{
-		public static async Task ExecuteAsync(IHordeClient hordeClient, DirectoryReference workingDir, LeaseId leaseId, ExecuteJobTask executeTask, IEnumerable<IJobExecutorFactory> executorFactories, DriverSettings driverSettings, ILogger logger, CancellationToken cancellationToken)
+		public static async Task ExecuteAsync(HordeClientFactory hordeClientFactory, DirectoryReference workingDir, LeaseId leaseId, ExecuteJobTask executeTask, IEnumerable<IJobExecutorFactory> executorFactories, DriverSettings driverSettings, ILogger logger, CancellationToken cancellationToken)
 		{
 			// Create an executor for this job
 			string executorName = String.IsNullOrEmpty(executeTask.JobOptions.Executor) ? driverSettings.Executor : executeTask.JobOptions.Executor;
@@ -2287,7 +2314,7 @@ namespace JobDriver.Execution
 			JobId jobId = JobId.Parse(executeTask.JobId);
 			JobStepBatchId batchId = JobStepBatchId.Parse(executeTask.BatchId);
 
-			JobExecutorOptions options = new JobExecutorOptions(hordeClient, workingDir, driverSettings.ProcessesToTerminate, jobId, batchId, leaseId, executeTask.JobOptions);
+			JobExecutorOptions options = new JobExecutorOptions(hordeClientFactory, executeTask.Token, workingDir, driverSettings.ProcessesToTerminate, jobId, batchId, leaseId, executeTask.JobOptions);
 
 			using JobExecutor executor = await executorFactory.CreateExecutorAsync(executeTask.Workspace, executeTask.AutoSdkWorkspace, options, cancellationToken);
 			await executor.ExecuteAsync(logger, cancellationToken);
diff --git a/Engine/Source/Programs/Horde/Drivers/JobDriver/Execution/LocalExecutor.cs b/Engine/Source/Programs/Horde/Drivers/JobDriver/Execution/LocalExecutor.cs
--- a/Engine/Source/Programs/Horde/Drivers/JobDriver/Execution/LocalExecutor.cs
+++ b/Engine/Source/Programs/Horde/Drivers/JobDriver/Execution/LocalExecutor.cs
@@ -1,6 +1,7 @@
 // Copyright Epic Games, Inc. All Rights Reserved.
 
 using EpicGames.Core;
+using EpicGames.Horde;
 using HordeCommon.Rpc.Messages;
 using Microsoft.Extensions.Logging;
 using Microsoft.Extensions.Options;
@@ -63,16 +64,16 @@ namespace JobDriver.Execution
 			throw new Exception($"Unable to find workspace root directory (looking for '{HordeSlnRelativePath}' in a parent directory of '{executableFileDir}'");
 		}
 
-		protected override Task<bool> SetupAsync(JobStepInfo step, ILogger logger, CancellationToken cancellationToken)
+		protected override Task<bool> SetupAsync(IHordeClient hordeClient, JobStepInfo step, ILogger logger, CancellationToken cancellationToken)
 		{
-			return SetupAsync(step, _localWorkspaceDir, null, logger, cancellationToken);
+			return SetupAsync(hordeClient, step, _localWorkspaceDir, null, logger, cancellationToken);
 		}
 
-		protected override Task<bool> ExecuteAsync(JobStepInfo step, ILogger logger, CancellationToken cancellationToken)
+		protected override Task<bool> ExecuteAsync(IHordeClient hordeClient, JobStepInfo step, ILogger logger, CancellationToken cancellationToken)
 		{
 			if (_settings.RunSteps)
 			{
-				return ExecuteAsync(step, _localWorkspaceDir, null, logger, cancellationToken);
+				return ExecuteAsync(hordeClient, step, _localWorkspaceDir, null, logger, cancellationToken);
 			}
 			else
 			{
diff --git a/Engine/Source/Programs/Horde/Drivers/JobDriver/Execution/PerforceExecutor.cs b/Engine/Source/Programs/Horde/Drivers/JobDriver/Execution/PerforceExecutor.cs
--- a/Engine/Source/Programs/Horde/Drivers/JobDriver/Execution/PerforceExecutor.cs
+++ b/Engine/Source/Programs/Horde/Drivers/JobDriver/Execution/PerforceExecutor.cs
@@ -10,6 +10,7 @@ using Horde.Common.Rpc;
 using HordeCommon.Rpc.Messages;
 using Microsoft.Extensions.Logging;
 using OpenTelemetry.Trace;
+using EpicGames.Horde;
 
 namespace JobDriver.Execution
 {
@@ -47,10 +48,10 @@ namespace JobDriver.Execution
 			_workspace = null!;
 		}
 
-		public override async Task InitializeAsync(RpcBeginBatchResponse batch, ILogger logger, CancellationToken cancellationToken)
+		public override async Task InitializeAsync(IHordeClient hordeClient, RpcBeginBatchResponse batch, ILogger logger, CancellationToken cancellationToken)
 		{
 			using TelemetrySpan span = Tracer.StartActiveSpan($"{nameof(PerforceExecutor)}.{nameof(InitializeAsync)}");
-			await base.InitializeAsync(batch, logger, cancellationToken);
+			await base.InitializeAsync(hordeClient, batch, logger, cancellationToken);
 
 			// Setup and sync the AutoSDK workspace
 			if (_autoSdkWorkspaceInfo != null)
@@ -221,18 +222,18 @@ namespace JobDriver.Execution
 			return perforceLogger;
 		}
 
-		protected override async Task<bool> SetupAsync(JobStepInfo step, ILogger logger, CancellationToken cancellationToken)
+		protected override async Task<bool> SetupAsync(IHordeClient hordeClient, JobStepInfo step, ILogger logger, CancellationToken cancellationToken)
 		{
 			PerforceMetadataLogger perforceLogger = CreatePerforceLogger(logger, Batch.Change, _workspace, _autoSdkWorkspace);
 			bool useP4 = WorkspaceInfo.ShouldUseHaveTable(_workspaceInfo.Method);
-			return await SetupAsync(step, _workspace.WorkspaceDir, useP4, perforceLogger, cancellationToken);
+			return await SetupAsync(hordeClient, step, _workspace.WorkspaceDir, useP4, perforceLogger, cancellationToken);
 		}
 
-		protected override async Task<bool> ExecuteAsync(JobStepInfo step, ILogger logger, CancellationToken cancellationToken)
+		protected override async Task<bool> ExecuteAsync(IHordeClient hordeClient, JobStepInfo step, ILogger logger, CancellationToken cancellationToken)
 		{
 			PerforceMetadataLogger perforceLogger = CreatePerforceLogger(logger, Batch.Change, _workspace, _autoSdkWorkspace);
 			bool useP4 = WorkspaceInfo.ShouldUseHaveTable(_workspaceInfo.Method);
-			return await ExecuteAsync(step, _workspace.WorkspaceDir, useP4, perforceLogger, cancellationToken);
+			return await ExecuteAsync(hordeClient, step, _workspace.WorkspaceDir, useP4, perforceLogger, cancellationToken);
 		}
 
 		public override async Task FinalizeAsync(ILogger logger, CancellationToken cancellationToken)
diff --git a/Engine/Source/Programs/Horde/Drivers/JobDriver/Execution/TestExecutor.cs b/Engine/Source/Programs/Horde/Drivers/JobDriver/Execution/TestExecutor.cs
--- a/Engine/Source/Programs/Horde/Drivers/JobDriver/Execution/TestExecutor.cs
+++ b/Engine/Source/Programs/Horde/Drivers/JobDriver/Execution/TestExecutor.cs
@@ -9,6 +9,7 @@ using Horde.Common.Rpc;
 using HordeCommon.Rpc.Messages;
 using Microsoft.Extensions.Logging;
 using OpenTelemetry.Trace;
+using EpicGames.Horde;
 
 namespace JobDriver.Execution
 {
@@ -44,9 +45,9 @@ namespace JobDriver.Execution
 			return result;
 		}
 
-		public override async Task InitializeAsync(RpcBeginBatchResponse batch, ILogger logger, CancellationToken cancellationToken)
+		public override async Task InitializeAsync(IHordeClient hordeClient, RpcBeginBatchResponse batch, ILogger logger, CancellationToken cancellationToken)
 		{
-			await base.InitializeAsync(batch, logger, cancellationToken);
+			await base.InitializeAsync(hordeClient, batch, logger, cancellationToken);
 
 			logger.LogInformation("Initializing");
 
@@ -68,7 +69,7 @@ namespace JobDriver.Execution
 			}
 		}
 
-		protected override async Task<bool> SetupAsync(JobStepInfo step, ILogger logger, CancellationToken cancellationToken)
+		protected override async Task<bool> SetupAsync(IHordeClient hordeClient, JobStepInfo step, ILogger logger, CancellationToken cancellationToken)
 		{
 			logger.LogInformation("**** BEGIN JOB SETUP ****");
 
@@ -170,7 +171,7 @@ namespace JobDriver.Execution
 			return nameToDependencyNames;
 		}
 
-		protected override async Task<bool> ExecuteAsync(JobStepInfo step, ILogger logger, CancellationToken cancellationToken)
+		protected override async Task<bool> ExecuteAsync(IHordeClient hordeClient, JobStepInfo step, ILogger logger, CancellationToken cancellationToken)
 		{
 			logger.LogInformation("**** BEGIN NODE {StepName} ****", step.Name);
 
diff --git a/Engine/Source/Programs/Horde/Drivers/JobDriver/Execution/WorkspaceExecutor.cs b/Engine/Source/Programs/Horde/Drivers/JobDriver/Execution/WorkspaceExecutor.cs
--- a/Engine/Source/Programs/Horde/Drivers/JobDriver/Execution/WorkspaceExecutor.cs
+++ b/Engine/Source/Programs/Horde/Drivers/JobDriver/Execution/WorkspaceExecutor.cs
@@ -7,6 +7,7 @@ using Horde.Common.Rpc;
 using HordeCommon.Rpc.Messages;
 using Microsoft.Extensions.Logging;
 using OpenTelemetry.Trace;
+using EpicGames.Horde;
 
 namespace JobDriver.Execution
 {
@@ -36,10 +37,10 @@ namespace JobDriver.Execution
 			base.Dispose(disposing);
 		}
 
-		public override async Task InitializeAsync(RpcBeginBatchResponse batch, ILogger logger, CancellationToken cancellationToken)
+		public override async Task InitializeAsync(IHordeClient hordeClient, RpcBeginBatchResponse batch, ILogger logger, CancellationToken cancellationToken)
 		{
 			using TelemetrySpan span = Tracer.StartActiveSpan($"{nameof(WorkspaceExecutor)}.{nameof(InitializeAsync)}");
-			await base.InitializeAsync(batch, logger, cancellationToken);
+			await base.InitializeAsync(hordeClient, batch, logger, cancellationToken);
 
 			if (Batch.Change == 0)
 			{
@@ -92,19 +93,19 @@ namespace JobDriver.Execution
 		}
 
 		/// <inheritdoc/>
-		protected override async Task<bool> SetupAsync(JobStepInfo step, ILogger logger, CancellationToken cancellationToken)
+		protected override async Task<bool> SetupAsync(IHordeClient hordeClient, JobStepInfo step, ILogger logger, CancellationToken cancellationToken)
 		{
 			// Loop back to JobExecutor's SetupAsync again, but with workspace and shared storage dir set
 			DirectoryReference workspaceDir = _workspace.DirectoryPath;
-			return await SetupAsync(step, workspaceDir, _workspace.IsPerforceWorkspace, GetLogger(logger), cancellationToken);
+			return await SetupAsync(hordeClient, step, workspaceDir, _workspace.IsPerforceWorkspace, GetLogger(logger), cancellationToken);
 		}
 
 		/// <inheritdoc/>
-		protected override async Task<bool> ExecuteAsync(JobStepInfo step, ILogger logger, CancellationToken cancellationToken)
+		protected override async Task<bool> ExecuteAsync(IHordeClient hordeClient, JobStepInfo step, ILogger logger, CancellationToken cancellationToken)
 		{
 			// Loop back to JobExecutor's ExecuteAsync again, but with workspace and shared storage dir set
 			DirectoryReference workspaceDir = _workspace.DirectoryPath;
-			return await ExecuteAsync(step, workspaceDir, _workspace.IsPerforceWorkspace, GetLogger(logger), cancellationToken);
+			return await ExecuteAsync(hordeClient, step, workspaceDir, _workspace.IsPerforceWorkspace, GetLogger(logger), cancellationToken);
 		}
 
 		/// <inheritdoc/>
diff --git a/Engine/Source/Programs/Shared/EpicGames.Horde/HordeClient.cs b/Engine/Source/Programs/Shared/EpicGames.Horde/HordeClient.cs
--- a/Engine/Source/Programs/Shared/EpicGames.Horde/HordeClient.cs
+++ b/Engine/Source/Programs/Shared/EpicGames.Horde/HordeClient.cs
@@ -448,16 +448,16 @@ namespace EpicGames.Horde
 		/// <summary>
 		/// Create a client using the user's default access token
 		/// </summary>
-		public HordeClient Create(Uri? serverUrl = null, string? accessToken = null)
+		public HordeClient Create(Uri? serverUrl = null, string? accessToken = null, ILoggerFactory? loggerFactory = null)
 		{
 			serverUrl ??= GetDefaultServerUrl();
 			if (accessToken == null)
 			{
-				return new HordeClientWithDynamicCredentials(serverUrl, _bundleCache, _httpMessageHandler, _hordeOptions, _loggerFactory);
+				return new HordeClientWithDynamicCredentials(serverUrl, _bundleCache, _httpMessageHandler, _hordeOptions, loggerFactory ?? _loggerFactory);
 			}
 			else
 			{
-				return new HordeClientWithStaticCredentials(serverUrl, accessToken, _httpMessageHandler, _bundleCache, _hordeOptions, _loggerFactory);
+				return new HordeClientWithStaticCredentials(serverUrl, accessToken, _httpMessageHandler, _bundleCache, _hordeOptions, loggerFactory ?? _loggerFactory);
 			}
 		}
 	}
