Horde: Fix tests.

#rnx

diff --git a/Engine/Source/Programs/Horde/Drivers/JobDriver.Tests/Execution/WorkspaceExecutorTest.cs b/Engine/Source/Programs/Horde/Drivers/JobDriver.Tests/Execution/WorkspaceExecutorTest.cs
--- a/Engine/Source/Programs/Horde/Drivers/JobDriver.Tests/Execution/WorkspaceExecutorTest.cs
+++ b/Engine/Source/Programs/Horde/Drivers/JobDriver.Tests/Execution/WorkspaceExecutorTest.cs
@@ -7,143 +7,40 @@ using System.Threading;
 using System.Threading.Tasks;
 using EpicGames.Core;
 using EpicGames.Horde;
-using EpicGames.Horde.Compute;
 using EpicGames.Horde.Jobs;
-using EpicGames.Horde.Logs;
-using EpicGames.Horde.Storage;
 using EpicGames.Horde.Streams;
-using Grpc.Core;
-using Grpc.Net.Client;
-using JobDriver.Execution;
 using Horde.Common.Rpc;
 using HordeCommon.Rpc.Tasks;
+using JobDriver.Execution;
 using Microsoft.Extensions.Logging;
 using Microsoft.Extensions.Logging.Abstractions;
 using Microsoft.VisualStudio.TestTools.UnitTesting;
-using EpicGames.Horde.Artifacts;
-using EpicGames.Horde.Secrets;
-using EpicGames.Horde.Projects;
-using EpicGames.Horde.Tools;
 
 namespace JobDriver.Tests.Execution
 {
 	[TestClass]
 	public sealed class WorkspaceExecutorTest : IDisposable
 	{
-		class FakeHordeClient : IHordeClient
-		{
-			class FakeJobRpc : JobRpc.JobRpcClient
-			{
-			}
-
-			private readonly Dictionary<StreamId, RpcGetStreamResponse> _streamIdToStreamResponse = new();
-			private readonly Dictionary<JobId, RpcGetJobResponse> _jobIdToJobResponse = new();
-
-			public event Action? OnAccessTokenStateChanged
-			{
-				add { }
-				remove { }
-			}
-
-			public void AddStream(StreamId streamId, string streamName)
-			{
-				if (_streamIdToStreamResponse.ContainsKey(streamId))
-				{
-					throw new Exception($"Stream ID {streamId} already added");
-				}
-
-				_streamIdToStreamResponse[streamId] = new RpcGetStreamResponse
-				{
-					Name = streamName,
-				};
-			}
-
-			public void AddAgentType(StreamId streamId, string agentType)
-			{
-				if (!_streamIdToStreamResponse.TryGetValue(streamId, out RpcGetStreamResponse? streamResponse))
-				{
-					throw new Exception($"Stream ID {streamId} not found");
-				}
-
-				string tempDir = Path.Join(Path.GetTempPath(), $"horde-agent-type-{agentType}-" + Guid.NewGuid().ToString()[..8]);
-				Directory.CreateDirectory(tempDir);
-
-				streamResponse.AgentTypes[agentType] = new RpcGetAgentTypeResponse
-				{
-					TempStorageDir = tempDir
-				};
-			}
-
-			public void AddJob(JobId jobId, StreamId streamId, int change, int preflightChange)
-			{
-				if (!_streamIdToStreamResponse.ContainsKey(streamId))
-				{
-					throw new Exception($"Stream ID {streamId} not found");
-				}
-
-				_jobIdToJobResponse[jobId] = new RpcGetJobResponse
-				{
-					StreamId = streamId.ToString(),
-					Change = change,
-					PreflightChange = preflightChange
-				};
-			}
-
-			public Uri ServerUrl
-				=> new Uri("http://fake-horde-server");
-
-			public IArtifactCollection Artifacts => throw new NotImplementedException();
-			public IComputeClient Compute => throw new NotImplementedException();
-			public IProjectCollection Projects => throw new NotImplementedException();
-			public ISecretCollection Secrets => throw new NotImplementedException();
-			public IToolCollection Tools => throw new NotImplementedException();
-
-			public IComputeClient CreateComputeClient()
-				=> throw new NotImplementedException();
-
-			public Task<GrpcChannel> CreateGrpcChannelAsync(CancellationToken cancellationToken = default)
-				=> throw new NotImplementedException();
-
-			public Task<TClient> CreateGrpcClientAsync<TClient>(CancellationToken cancellationToken = default) where TClient : ClientBase<TClient>
-			{
-				if (typeof(TClient) == typeof(JobRpc.JobRpcClient))
-				{
-					return Task.FromResult((TClient)(object)new FakeJobRpc());
-				}
-				throw new NotImplementedException();
-			}
-
-			public HordeHttpClient CreateHttpClient()
-				=> throw new NotImplementedException();
-
-			public IServerLogger CreateServerLogger(LogId logId, LogLevel minimumLevel = LogLevel.Information)
-				=> throw new NotImplementedException();
-
-			public IStorageNamespace GetStorageNamespace(string relativePath, string? accessToken = null)
-				=> throw new NotImplementedException();
-
-			public Task<string?> GetAccessTokenAsync(bool interactive, CancellationToken cancellationToken = default)
-				=> Task.FromResult<string?>("access-token");
-
-			public bool HasValidAccessToken()
-				=> true;
-
-			public Task<bool> LoginAsync(bool allowPrompt, CancellationToken cancellationToken)
-				=> throw new NotImplementedException();
-		}
-
 		private readonly StreamId _streamId = new StreamId("foo-main");
 		private readonly JobId _jobId = JobId.Parse("65bd0655591b5d5d7d047b58");
 		private readonly JobStepBatchId _batchId = new JobStepBatchId(0x1234);
 		private const string AgentType = "BogusAgentType";
 
 		private readonly ILogger<WorkspaceExecutorTest> _logger;
-		private readonly FakeHordeClient _hordeClient;
 		private readonly DirectoryReference _workingDir = new DirectoryReference("workspace-test");
 		private readonly FakeWorkspaceMaterializer _autoSdkWorkspace = new();
 		private readonly FakeWorkspaceMaterializer _workspace = new();
 		private readonly WorkspaceExecutor _executor;
 
+		FakeHordeClient CreateHordeClient(ILogger logger)
+		{
+			FakeHordeClient hordeClient = new FakeHordeClient();
+			hordeClient.AddStream(_streamId, "//Foo/Main");
+			hordeClient.AddAgentType(_streamId, AgentType);
+			hordeClient.AddJob(_jobId, _streamId, 1, 0);
+			return hordeClient;
+		}
+
 		public WorkspaceExecutorTest()
 		{
 			using ILoggerFactory loggerFactory = LoggerFactory.Create(builder =>
@@ -154,16 +51,11 @@ namespace JobDriver.Tests.Execution
 
 			_logger = loggerFactory.CreateLogger<WorkspaceExecutorTest>();
 
-			_hordeClient = new FakeHordeClient();
-			_hordeClient.AddStream(_streamId, "//Foo/Main");
-			_hordeClient.AddAgentType(_streamId, AgentType);
-			_hordeClient.AddJob(_jobId, _streamId, 1, 0);
-
 			_autoSdkWorkspace.SetFile(1, "HostWin64/Android/base.h", "base");
 			_workspace.SetFile(1, "main.cpp", "main");
 			_workspace.SetFile(1, "foo/bar/baz.h", "baz");
 
-			JobExecutorOptions executorOptions = new JobExecutorOptions(_hordeClient, _workingDir, null, _jobId, _batchId, default, new RpcJobOptions());
+			JobExecutorOptions executorOptions = new JobExecutorOptions(CreateHordeClient, _workingDir, null, _jobId, _batchId, default, new RpcJobOptions());
 			_executor = new(executorOptions, _workspace, null, SimpleTestExecutor.NoOpTracer, NullLogger.Instance);
 		}
 
@@ -177,8 +69,10 @@ namespace JobDriver.Tests.Execution
 		[TestMethod]
 		public async Task RegularWorkspaceAsync()
 		{
+			await using IHordeClientWithLifetime hordeClient = CreateHordeClient(NullLogger.Instance);
+
 			RpcBeginBatchResponse batch = new RpcBeginBatchResponse { Change = 1 };
-			await _executor.InitializeAsync(batch, _logger, CancellationToken.None);
+			await _executor.InitializeAsync(hordeClient, batch, _logger, CancellationToken.None);
 			AssertWorkspaceFile(_workspace, "main.cpp", "main");
 			AssertWorkspaceFile(_workspace, "foo/bar/baz.h", "baz");
 			await _executor.FinalizeAsync(_logger, CancellationToken.None);
@@ -187,11 +81,13 @@ namespace JobDriver.Tests.Execution
 		[TestMethod]
 		public async Task RegularAndAutoSdkWorkspaceAsync()
 		{
-			JobExecutorOptions executorOptions = new JobExecutorOptions(_hordeClient, _workingDir, Array.Empty<ProcessToTerminate>(), _jobId, _batchId, default, new RpcJobOptions());
+			JobExecutorOptions executorOptions = new JobExecutorOptions(CreateHordeClient, _workingDir, Array.Empty<ProcessToTerminate>(), _jobId, _batchId, default, new RpcJobOptions());
 			using WorkspaceExecutor executor = new(executorOptions, _workspace, _autoSdkWorkspace, SimpleTestExecutor.NoOpTracer, NullLogger.Instance);
 
+			await using IHordeClientWithLifetime hordeClient = CreateHordeClient(NullLogger.Instance);
+
 			RpcBeginBatchResponse batch = new RpcBeginBatchResponse { Change = 1 };
-			await executor.InitializeAsync(batch, _logger, CancellationToken.None);
+			await executor.InitializeAsync(hordeClient, batch, _logger, CancellationToken.None);
 			AssertWorkspaceFile(_autoSdkWorkspace, "HostWin64/Android/base.h", "base");
 			AssertWorkspaceFile(_workspace, "main.cpp", "main");
 			AssertWorkspaceFile(_workspace, "foo/bar/baz.h", "baz");
@@ -201,11 +97,13 @@ namespace JobDriver.Tests.Execution
 		[TestMethod]
 		public async Task EnvVarsAsync()
 		{
-			JobExecutorOptions executorOptions = new JobExecutorOptions(_hordeClient, _workingDir, Array.Empty<ProcessToTerminate>(), _jobId, _batchId, default, new RpcJobOptions());
+			JobExecutorOptions executorOptions = new JobExecutorOptions(CreateHordeClient, _workingDir, Array.Empty<ProcessToTerminate>(), _jobId, _batchId, default, new RpcJobOptions());
 			using WorkspaceExecutor executor = new(executorOptions, _workspace, _autoSdkWorkspace, SimpleTestExecutor.NoOpTracer, NullLogger.Instance);
 
+			await using IHordeClientWithLifetime hordeClient = CreateHordeClient(NullLogger.Instance);
+
 			RpcBeginBatchResponse batch = new RpcBeginBatchResponse { Change = 1, StreamName = "//UE5/Main" };
-			await executor.InitializeAsync(batch, _logger, CancellationToken.None);
+			await executor.InitializeAsync(hordeClient, batch, _logger, CancellationToken.None);
 
 			IReadOnlyDictionary<string, string> envVars = executor.GetEnvVars();
 			Assert.AreEqual("1", envVars["IsBuildMachine"]);
@@ -224,14 +122,22 @@ namespace JobDriver.Tests.Execution
 		{
 			JobId preflightJobId = JobId.Parse("65bd0655591b5d5d7d047b59");
 
-			_hordeClient.AddJob(preflightJobId, _streamId, 1, 1000);
+			FakeHordeClient CreatePreflightHordeClient(ILogger logger)
+			{
+				FakeHordeClient hordeClient = CreateHordeClient(logger);
+				hordeClient.AddJob(preflightJobId, _streamId, 1, 1000);
+				return hordeClient;
+			}
+
+			await using IHordeClientWithLifetime hordeClient = CreatePreflightHordeClient(NullLogger.Instance);
+
 			_workspace.SetFile(1000, "New/Feature/Foo.cs", "foo");
 
-			JobExecutorOptions executorOptions = new JobExecutorOptions(_hordeClient, _workingDir, Array.Empty<ProcessToTerminate>(), preflightJobId, _batchId, default, new RpcJobOptions());
+			JobExecutorOptions executorOptions = new JobExecutorOptions(CreatePreflightHordeClient, _workingDir, Array.Empty<ProcessToTerminate>(), preflightJobId, _batchId, default, new RpcJobOptions());
 			using WorkspaceExecutor executor = new(executorOptions, _workspace, null, SimpleTestExecutor.NoOpTracer, NullLogger.Instance);
 
 			RpcBeginBatchResponse batch = new RpcBeginBatchResponse { Change = 1, PreflightChange = 1000 };
-			await executor.InitializeAsync(batch, _logger, CancellationToken.None);
+			await executor.InitializeAsync(hordeClient, batch, _logger, CancellationToken.None);
 			AssertWorkspaceFile(_workspace, "main.cpp", "main");
 			AssertWorkspaceFile(_workspace, "foo/bar/baz.h", "baz");
 			AssertWorkspaceFile(_workspace, "New/Feature/Foo.cs", "foo");
@@ -243,13 +149,20 @@ namespace JobDriver.Tests.Execution
 		{
 			JobId noChangeJobId = JobId.Parse("65bd0655591b5d5d7d047b5a");
 
-			_hordeClient.AddJob(noChangeJobId, _streamId, 0, 0);
+			FakeHordeClient CreateNoJobHordeClient(ILogger logger)
+			{
+				FakeHordeClient hordeClient = CreateHordeClient(logger);
+				hordeClient.AddJob(noChangeJobId, _streamId, 0, 0);
+				return hordeClient;
+			}
 
-			JobExecutorOptions executorOptions = new JobExecutorOptions(_hordeClient, _workingDir, Array.Empty<ProcessToTerminate>(), noChangeJobId, _batchId, default, new RpcJobOptions());
+			JobExecutorOptions executorOptions = new JobExecutorOptions(CreateNoJobHordeClient, _workingDir, Array.Empty<ProcessToTerminate>(), noChangeJobId, _batchId, default, new RpcJobOptions());
 			using WorkspaceExecutor executor = new(executorOptions, _workspace, null, SimpleTestExecutor.NoOpTracer, NullLogger.Instance);
 
+			await using IHordeClientWithLifetime hordeClient = CreateNoJobHordeClient(NullLogger.Instance);
+
 			RpcBeginBatchResponse batch = new RpcBeginBatchResponse { };
-			await Assert.ThrowsExceptionAsync<WorkspaceMaterializationException>(() => executor.InitializeAsync(batch, _logger, CancellationToken.None));
+			await Assert.ThrowsExceptionAsync<WorkspaceMaterializationException>(() => executor.InitializeAsync(hordeClient, batch, _logger, CancellationToken.None));
 		}
 
 		private static void AssertWorkspaceFile(IWorkspaceMaterializer workspace, string relativePath, string expectedContent)
diff --git a/Engine/Source/Programs/Horde/Drivers/JobDriver.Tests/FakeHordeClient.cs b/Engine/Source/Programs/Horde/Drivers/JobDriver.Tests/FakeHordeClient.cs
new file mode 100644
--- /dev/null
+++ b/Engine/Source/Programs/Horde/Drivers/JobDriver.Tests/FakeHordeClient.cs
@@ -0,0 +1,128 @@
+// Copyright Epic Games, Inc. All Rights Reserved.
+
+using System;
+using System.Collections.Generic;
+using System.IO;
+using System.Threading;
+using System.Threading.Tasks;
+using EpicGames.Horde;
+using EpicGames.Horde.Artifacts;
+using EpicGames.Horde.Compute;
+using EpicGames.Horde.Jobs;
+using EpicGames.Horde.Logs;
+using EpicGames.Horde.Projects;
+using EpicGames.Horde.Secrets;
+using EpicGames.Horde.Storage;
+using EpicGames.Horde.Streams;
+using EpicGames.Horde.Tools;
+using Grpc.Core;
+using Grpc.Net.Client;
+using Horde.Common.Rpc;
+using Microsoft.Extensions.Logging;
+
+namespace JobDriver.Tests
+{
+	class FakeHordeClient : IHordeClientWithLifetime
+	{
+		class FakeJobRpc : JobRpc.JobRpcClient
+		{
+		}
+
+		private readonly Dictionary<StreamId, RpcGetStreamResponse> _streamIdToStreamResponse = new();
+		private readonly Dictionary<JobId, RpcGetJobResponse> _jobIdToJobResponse = new();
+
+		public event Action? OnAccessTokenStateChanged
+		{
+			add { }
+			remove { }
+		}
+
+		public ValueTask DisposeAsync() => default;
+
+		public void AddStream(StreamId streamId, string streamName)
+		{
+			if (_streamIdToStreamResponse.ContainsKey(streamId))
+			{
+				throw new Exception($"Stream ID {streamId} already added");
+			}
+
+			_streamIdToStreamResponse[streamId] = new RpcGetStreamResponse
+			{
+				Name = streamName,
+			};
+		}
+
+		public void AddAgentType(StreamId streamId, string agentType)
+		{
+			if (!_streamIdToStreamResponse.TryGetValue(streamId, out RpcGetStreamResponse? streamResponse))
+			{
+				throw new Exception($"Stream ID {streamId} not found");
+			}
+
+			string tempDir = Path.Join(Path.GetTempPath(), $"horde-agent-type-{agentType}-" + Guid.NewGuid().ToString()[..8]);
+			Directory.CreateDirectory(tempDir);
+
+			streamResponse.AgentTypes[agentType] = new RpcGetAgentTypeResponse
+			{
+				TempStorageDir = tempDir
+			};
+		}
+
+		public void AddJob(JobId jobId, StreamId streamId, int change, int preflightChange)
+		{
+			if (!_streamIdToStreamResponse.ContainsKey(streamId))
+			{
+				throw new Exception($"Stream ID {streamId} not found");
+			}
+
+			_jobIdToJobResponse[jobId] = new RpcGetJobResponse
+			{
+				StreamId = streamId.ToString(),
+				Change = change,
+				PreflightChange = preflightChange
+			};
+		}
+
+		public Uri ServerUrl
+			=> new Uri("http://fake-horde-server");
+
+		public IArtifactCollection Artifacts => throw new NotImplementedException();
+		public IComputeClient Compute => throw new NotImplementedException();
+		public IProjectCollection Projects => throw new NotImplementedException();
+		public ISecretCollection Secrets => throw new NotImplementedException();
+		public IToolCollection Tools => throw new NotImplementedException();
+
+		public IComputeClient CreateComputeClient()
+			=> throw new NotImplementedException();
+
+		public Task<GrpcChannel> CreateGrpcChannelAsync(CancellationToken cancellationToken = default)
+			=> throw new NotImplementedException();
+
+		public Task<TClient> CreateGrpcClientAsync<TClient>(CancellationToken cancellationToken = default) where TClient : ClientBase<TClient>
+		{
+			if (typeof(TClient) == typeof(JobRpc.JobRpcClient))
+			{
+				return Task.FromResult((TClient)(object)new FakeJobRpc());
+			}
+			throw new NotImplementedException();
+		}
+
+		public HordeHttpClient CreateHttpClient()
+			=> throw new NotImplementedException();
+
+		public IServerLogger CreateServerLogger(LogId logId, LogLevel minimumLevel = LogLevel.Information)
+			=> throw new NotImplementedException();
+
+		public IStorageNamespace GetStorageNamespace(string relativePath, string? accessToken = null)
+			=> throw new NotImplementedException();
+
+		public Task<string?> GetAccessTokenAsync(bool interactive, CancellationToken cancellationToken = default)
+			=> Task.FromResult<string?>("access-token");
+
+		public bool HasValidAccessToken()
+			=> true;
+
+		public Task<bool> LoginAsync(bool allowPrompt, CancellationToken cancellationToken)
+			=> throw new NotImplementedException();
+	}
+}
diff --git a/Engine/Source/Programs/Horde/Drivers/JobDriver.Tests/SimpleTestExecutor.cs b/Engine/Source/Programs/Horde/Drivers/JobDriver.Tests/SimpleTestExecutor.cs
--- a/Engine/Source/Programs/Horde/Drivers/JobDriver.Tests/SimpleTestExecutor.cs
+++ b/Engine/Source/Programs/Horde/Drivers/JobDriver.Tests/SimpleTestExecutor.cs
@@ -10,6 +10,7 @@ using HordeCommon.Rpc.Messages;
 using Microsoft.Extensions.Logging;
 using Microsoft.Extensions.Logging.Abstractions;
 using OpenTelemetry.Trace;
+using EpicGames.Horde;
 
 namespace JobDriver.Tests
 {
@@ -21,18 +22,18 @@ namespace JobDriver.Tests
 		public static Tracer NoOpTracer { get;  }= TracerProvider.Default.GetTracer("NoOp");
 
 		public SimpleTestExecutor(Func<JobStepInfo, ILogger, CancellationToken, Task<JobStepOutcome>> func)
-			: base(new JobExecutorOptions(null!, null!, null, default, default, default, null!), NoOpTracer, NullLogger.Instance)
+			: base(new JobExecutorOptions(_ => new FakeHordeClient(), null!, null, default, default, default, null!), NoOpTracer, NullLogger.Instance)
 		{
 			_func = func;
 		}
 
-		public override Task InitializeAsync(RpcBeginBatchResponse batch, ILogger logger, CancellationToken cancellationToken)
+		public override Task InitializeAsync(IHordeClient hordeClient, RpcBeginBatchResponse batch, ILogger logger, CancellationToken cancellationToken)
 		{
 			logger.LogDebug("SimpleTestExecutor.InitializeAsync()");
 			return Task.CompletedTask;
 		}
 
-		public override Task<JobStepOutcome> RunAsync(JobStepInfo step, ILogger logger, CancellationToken cancellationToken)
+		public override Task<JobStepOutcome> RunAsync(IHordeClient hordeClient, JobStepInfo step, ILogger logger, CancellationToken cancellationToken)
 		{
 			logger.LogDebug("SimpleTestExecutor.RunAsync(Step: {Step})", step);
 			return _func(step, logger, cancellationToken);
@@ -44,12 +45,12 @@ namespace JobDriver.Tests
 			return Task.CompletedTask;
 		}
 
-		protected override Task<bool> SetupAsync(JobStepInfo step, ILogger logger, CancellationToken cancellationToken)
+		protected override Task<bool> SetupAsync(IHordeClient hordeClient, JobStepInfo step, ILogger logger, CancellationToken cancellationToken)
 		{
 			throw new NotImplementedException();
 		}
 
-		protected override Task<bool> ExecuteAsync(JobStepInfo step, ILogger logger, CancellationToken cancellationToken)
+		protected override Task<bool> ExecuteAsync(IHordeClient hordeClient, JobStepInfo step, ILogger logger, CancellationToken cancellationToken)
 		{
 			throw new NotImplementedException();
 		}
diff --git a/Engine/Source/Programs/Horde/Drivers/JobDriver/Execution/JobExecutor.cs b/Engine/Source/Programs/Horde/Drivers/JobDriver/Execution/JobExecutor.cs
--- a/Engine/Source/Programs/Horde/Drivers/JobDriver/Execution/JobExecutor.cs
+++ b/Engine/Source/Programs/Horde/Drivers/JobDriver/Execution/JobExecutor.cs
@@ -51,8 +51,7 @@ namespace JobDriver.Execution
 
 	class JobExecutorOptions
 	{
-		public HordeClientFactory HordeClientFactory { get; }
-		public string AccessToken { get; }
+		public Func<ILogger, IHordeClientWithLifetime> CreateHordeClient { get; }
 		public DirectoryReference WorkingDir { get; }
 		public IReadOnlyList<ProcessToTerminate>? ProcessesToTerminate { get; }
 		public JobId JobId { get; }
@@ -60,10 +59,9 @@ namespace JobDriver.Execution
 		public LeaseId LeaseId { get; }
 		public RpcJobOptions JobOptions { get; }
 
-		public JobExecutorOptions(HordeClientFactory hordeClientFactory, string accessToken, DirectoryReference workingDir, IReadOnlyList<ProcessToTerminate>? processesToTerminate, JobId jobId, JobStepBatchId batchId, LeaseId leaseId, RpcJobOptions jobOptions)
+		public JobExecutorOptions(Func<ILogger, IHordeClientWithLifetime> createHordeClient, DirectoryReference workingDir, IReadOnlyList<ProcessToTerminate>? processesToTerminate, JobId jobId, JobStepBatchId batchId, LeaseId leaseId, RpcJobOptions jobOptions)
 		{
-			HordeClientFactory = hordeClientFactory;
-			AccessToken = accessToken;
+			CreateHordeClient = createHordeClient;
 			WorkingDir = workingDir;
 			ProcessesToTerminate = processesToTerminate;
 			JobId = jobId;
@@ -218,8 +216,7 @@ namespace JobDriver.Execution
 		/// </summary>
 		protected ILogger Logger { get; }
 
-		readonly HordeClientFactory _hordeClientFactory;
-		readonly string _accessToken;
+		readonly Func<ILogger, IHordeClientWithLifetime> _createHordeClient;
 
 		protected DirectoryReference WorkingDir { get; }
 		protected JobId JobId { get; }
@@ -254,8 +251,7 @@ namespace JobDriver.Execution
 
 		protected JobExecutor(JobExecutorOptions options, Tracer tracer, ILogger logger)
 		{
-			_hordeClientFactory = options.HordeClientFactory;
-			_accessToken = options.AccessToken;
+			_createHordeClient = options.CreateHordeClient;
 
 			WorkingDir = options.WorkingDir;
 			JobId = options.JobId;
@@ -292,7 +288,7 @@ namespace JobDriver.Execution
 			span.SetAttribute("horde.lease.id", LeaseId.ToString());
 			JobOptions.DecorateSpan(span);
 
-			await using HordeClient hordeClient = _hordeClientFactory.Create(accessToken: _accessToken);
+			await using IHordeClientWithLifetime hordeClient = _createHordeClient(logger);
 
 			JobRpc.JobRpcClient jobRpc = await hordeClient.CreateGrpcClientAsync<JobRpc.JobRpcClient>(cancellationToken);
 
@@ -496,23 +492,6 @@ namespace JobDriver.Execution
 			}
 		}
 
-		class LoggerFactoryWrapper : ILoggerFactory
-		{
-			readonly ILogger _logger;
-
-			public LoggerFactoryWrapper(ILogger logger)
-				=> _logger = logger;
-
-			public void AddProvider(ILoggerProvider provider)
-			{ }
-
-			public ILogger CreateLogger(string categoryName)
-				=> _logger;
-
-			public void Dispose()
-			{ }
-		}
-
 		/// <summary>
 		/// Executes a step
 		/// </summary>
@@ -526,8 +505,7 @@ namespace JobDriver.Execution
 			using CancellationTokenSource combined = CancellationTokenSource.CreateLinkedTokenSource(cancellationToken, stepCancellationToken);
 			try
 			{
-				using LoggerFactoryWrapper loggerFactory = new LoggerFactoryWrapper(stepLogger);
-				await using HordeClient stepHordeClient = _hordeClientFactory.Create(accessToken: _accessToken, loggerFactory: loggerFactory);
+				await using IHordeClientWithLifetime stepHordeClient = _createHordeClient(stepLogger);
 				JobStepOutcome stepOutcome = await RunAsync(stepHordeClient, step, stepLogger, combined.Token);
 				return (stepOutcome, JobStepState.Completed);
 			}
@@ -2300,6 +2278,23 @@ namespace JobDriver.Execution
 
 	static class JobExecutorHelpers
 	{
+		class LoggerFactoryWrapper : ILoggerFactory
+		{
+			readonly ILogger _logger;
+
+			public LoggerFactoryWrapper(ILogger logger)
+				=> _logger = logger;
+
+			public void AddProvider(ILoggerProvider provider)
+			{ }
+
+			public ILogger CreateLogger(string categoryName)
+				=> _logger;
+
+			public void Dispose()
+			{ }
+		}
+
 		public static async Task ExecuteAsync(HordeClientFactory hordeClientFactory, DirectoryReference workingDir, LeaseId leaseId, ExecuteJobTask executeTask, IEnumerable<IJobExecutorFactory> executorFactories, DriverSettings driverSettings, ILogger logger, CancellationToken cancellationToken)
 		{
 			// Create an executor for this job
@@ -2314,7 +2309,13 @@ namespace JobDriver.Execution
 			JobId jobId = JobId.Parse(executeTask.JobId);
 			JobStepBatchId batchId = JobStepBatchId.Parse(executeTask.BatchId);
 
-			JobExecutorOptions options = new JobExecutorOptions(hordeClientFactory, executeTask.Token, workingDir, driverSettings.ProcessesToTerminate, jobId, batchId, leaseId, executeTask.JobOptions);
+			IHordeClientWithLifetime CreateHordeClient(ILogger logger)
+			{
+				LoggerFactoryWrapper loggerFactory = new LoggerFactoryWrapper(logger);
+				return hordeClientFactory.Create(accessToken: executeTask.Token, loggerFactory: loggerFactory);
+			}
+
+			JobExecutorOptions options = new JobExecutorOptions(CreateHordeClient, workingDir, driverSettings.ProcessesToTerminate, jobId, batchId, leaseId, executeTask.JobOptions);
 
 			using JobExecutor executor = await executorFactory.CreateExecutorAsync(executeTask.Workspace, executeTask.AutoSdkWorkspace, options, cancellationToken);
 			await executor.ExecuteAsync(logger, cancellationToken);
diff --git a/Engine/Source/Programs/Horde/Plugins/Build/HordeServer.Build.Tests/Notifications/NotificationServiceTests.cs b/Engine/Source/Programs/Horde/Plugins/Build/HordeServer.Build.Tests/Notifications/NotificationServiceTests.cs
--- a/Engine/Source/Programs/Horde/Plugins/Build/HordeServer.Build.Tests/Notifications/NotificationServiceTests.cs
+++ b/Engine/Source/Programs/Horde/Plugins/Build/HordeServer.Build.Tests/Notifications/NotificationServiceTests.cs
@@ -240,7 +240,7 @@ namespace HordeServer.Tests.Notifications
 		async Task<IJob> AbortStepAsync(IJob job, int batchIdx, int stepIdx)
 		{
 			LogId logId = LogIdUtils.GenerateNewId();
-			ILog log = await LogCollection.AddAsync(job.Id, job.Batches[batchIdx].LeaseId, job.Batches[batchIdx].SessionId, LogType.Json, logId);
+			_ = await LogCollection.AddAsync(job.Id, job.Batches[batchIdx].LeaseId, job.Batches[batchIdx].SessionId, LogType.Json, logId);
 			job = Deref(await JobService.UpdateStepAsync(job, job.Batches[batchIdx].Id, job.Batches[batchIdx].Steps[stepIdx].Id, _streamConfig, newAbortRequested: true, newAbortByUserId: UserId.Anonymous, newLogId: logId));
 			Assert.AreEqual(JobStepState.Aborted, job.Batches[batchIdx].Steps[stepIdx].State);
 			return job;
@@ -251,7 +251,7 @@ namespace HordeServer.Tests.Notifications
 			Assert.AreEqual(JobStepState.Ready, job.Batches[batchIdx].Steps[stepIdx].State);
 
 			LogId logId = LogIdUtils.GenerateNewId();
-			ILog log = await LogCollection.AddAsync(job.Id, job.Batches[batchIdx].LeaseId, job.Batches[batchIdx].SessionId, LogType.Json, logId);
+			_ = await LogCollection.AddAsync(job.Id, job.Batches[batchIdx].LeaseId, job.Batches[batchIdx].SessionId, LogType.Json, logId);
 			job = Deref(await JobService.UpdateStepAsync(job, job.Batches[batchIdx].Id, job.Batches[batchIdx].Steps[stepIdx].Id, _streamConfig, JobStepState.Running, JobStepOutcome.Success, newLogId: logId));
 			Assert.AreEqual(JobStepState.Running, job.Batches[batchIdx].Steps[stepIdx].State);
 			job = Deref(await JobService.UpdateStepAsync(job, job.Batches[batchIdx].Id, job.Batches[batchIdx].Steps[stepIdx].Id, _streamConfig, JobStepState.Completed, outcome, newLogId: logId));
diff --git a/Engine/Source/Programs/Shared/EpicGames.Horde.Tests/Storage/BundleTests.cs b/Engine/Source/Programs/Shared/EpicGames.Horde.Tests/Storage/BundleTests.cs
--- a/Engine/Source/Programs/Shared/EpicGames.Horde.Tests/Storage/BundleTests.cs
+++ b/Engine/Source/Programs/Shared/EpicGames.Horde.Tests/Storage/BundleTests.cs
@@ -2,7 +2,6 @@
 
 using System;
 using System.Collections.Generic;
-using System.IO;
 using System.Linq;
 using System.Threading;
 using System.Threading.Tasks;
diff --git a/Engine/Source/Programs/Shared/EpicGames.Horde/HordeClient.cs b/Engine/Source/Programs/Shared/EpicGames.Horde/HordeClient.cs
--- a/Engine/Source/Programs/Shared/EpicGames.Horde/HordeClient.cs
+++ b/Engine/Source/Programs/Shared/EpicGames.Horde/HordeClient.cs
@@ -31,7 +31,7 @@ namespace EpicGames.Horde
 	/// <summary>
 	/// Default implementation of <see cref="IHordeClient"/>
 	/// </summary>
-	public abstract class HordeClient : IHordeClient, IAsyncDisposable
+	public abstract class HordeClient : IHordeClientWithLifetime
 	{
 		readonly Uri _serverUrl;
 		readonly BundleCache _bundleCache;
