Horde: Use upserts to create new BlobInfo documents, to avoid races with new blobs being uploaded by other pods.

#rnx

diff --git a/Engine/Source/Programs/Horde/Plugins/Storage/HordeServer.Storage/Storage/StorageService.cs b/Engine/Source/Programs/Horde/Plugins/Storage/HordeServer.Storage/Storage/StorageService.cs
--- a/Engine/Source/Programs/Horde/Plugins/Storage/HordeServer.Storage/Storage/StorageService.cs
+++ b/Engine/Source/Programs/Horde/Plugins/Storage/HordeServer.Storage/Storage/StorageService.cs
@@ -852,13 +852,19 @@ namespace HordeServer.Storage
 					}
 
 					// If there's anything left, try to insert shadow blobs for them. We'll do a separate query to figure out which succeeded.
-					List<InsertOneModel<BlobInfo>> inserts = new List<InsertOneModel<BlobInfo>>();
+					List<WriteModel<BlobInfo>> writes = new List<WriteModel<BlobInfo>>();
 					foreach (string import in remaining)
 					{
-						BlobInfo info = new BlobInfo(ObjectId.GenerateNewId(), namespaceId, new BlobLocator(import)) { Shadow = true };
-						inserts.Add(new InsertOneModel<BlobInfo>(info));
+						FilterDefinition<BlobInfo> filter =
+							Builders<BlobInfo>.Filter.Eq(x => x.NamespaceId, namespaceId) &
+							Builders<BlobInfo>.Filter.Eq(x => x.Path, import);
+
+						UpdateDefinition<BlobInfo> update =
+							Builders<BlobInfo>.Update.SetOnInsert(x => x.Shadow, true);
+
+						writes.Add(new UpdateOneModel<BlobInfo>(filter, update) { IsUpsert = true });
 					}
-					await _blobCollection.BulkWriteAsync(inserts, new BulkWriteOptions { IsOrdered = false }, cancellationToken);
+					await _blobCollection.BulkWriteAsync(writes, new BulkWriteOptions { IsOrdered = false }, cancellationToken);
 				}
 
 				// Find all the matching import ids in order
