Horde: Add job state logging for batch and other updates

diff --git a/Engine/Source/Programs/Horde/Plugins/Build/HordeServer.Build/Jobs/JobCollection.cs b/Engine/Source/Programs/Horde/Plugins/Build/HordeServer.Build/Jobs/JobCollection.cs
--- a/Engine/Source/Programs/Horde/Plugins/Build/HordeServer.Build/Jobs/JobCollection.cs
+++ b/Engine/Source/Programs/Horde/Plugins/Build/HordeServer.Build/Jobs/JobCollection.cs
@@ -1801,6 +1801,7 @@ namespace HordeServer.Jobs
 				JobStepBatchDocument batch = jobDocument.Batches[batchIdx];
 				if (batch.State == JobStepBatchState.Ready || batch.State == JobStepBatchState.Waiting)
 				{
+					_logger.LogInformation("Skipping all batches: job {JobId}, batch {BatchId}", jobDocument.Id, batch.Id);
 					batch.State = JobStepBatchState.Complete;
 					batch.Error = reason;
 					batch.FinishTime = DateTimeOffset.UtcNow;
@@ -1826,6 +1827,8 @@ namespace HordeServer.Jobs
 		/// <inheritdoc/>
 		async Task<JobDocument?> TrySkipBatchAsync(JobDocument jobDocument, JobStepBatchId batchId, IGraph graph, JobStepBatchError reason, CancellationToken cancellationToken)
 		{
+			_logger.LogInformation("Trying to skip batch: job {JobId}, batch {BatchId}", jobDocument.Id, batchId);
+
 			JobStepBatchDocument? batch = jobDocument.Batches.FirstOrDefault(x => x.Id == batchId);
 			if (batch == null)
 			{
@@ -2431,7 +2434,7 @@ namespace HordeServer.Jobs
 							}
 							else if (!steps.Any(x => !x.AbortRequested && JobStepExtensions.IsPendingState(x.State)))
 							{
-								logger.LogDebug("Transitioning job {JobId}, batch {BatchId}, step {StepId} to ready state ({Dependencies})", job.Id, batch.Id, step.Id, String.Join(", ", steps.Select(x => x.Id.ToString())));
+								logger.LogInformation("Transitioning job {JobId}, batch {BatchId}, step {StepId} to ready state ({Dependencies})", job.Id, batch.Id, step.Id, String.Join(", ", steps.Select(x => x.Id.ToString())));
 								newState = JobStepState.Ready;
 							}
 						}
@@ -2458,10 +2461,12 @@ namespace HordeServer.Jobs
 						if (batch.State != newState)
 						{
 							batch.State = newState;
+							logger.LogInformation("Transitioning job {JobId}, batch {BatchId} to state {JobStepBatchState}", job.Id, batch.Id, newState);
 							updates.Add(updateBuilder.Set(x => x.Batches[batchIdx].State, batch.State));
 						}
 						if (batch.ReadyTime != newReadyTime)
 						{
+							logger.LogInformation("Setting job {JobId}, batch {BatchId} ready time", job.Id, batch.Id);
 							batch.ReadyTime = newReadyTime;
 							updates.Add(updateBuilder.Set(x => x.Batches[batchIdx].ReadyTime, batch.ReadyTime));
 						}
diff --git a/Engine/Source/Programs/Horde/Plugins/Build/HordeServer.Build/Jobs/JobService.cs b/Engine/Source/Programs/Horde/Plugins/Build/HordeServer.Build/Jobs/JobService.cs
--- a/Engine/Source/Programs/Horde/Plugins/Build/HordeServer.Build/Jobs/JobService.cs
+++ b/Engine/Source/Programs/Horde/Plugins/Build/HordeServer.Build/Jobs/JobService.cs
@@ -779,14 +779,17 @@ namespace HordeServer.Jobs
 						if (error != null)
 						{
 							newError = error.Value;
+							_logger.LogError("Batch {BatchId} marked with error {JobStepBatchError}", batchId, error.Value.ToString());
 						}
 						else if (!batch.Steps.Any(x => x.StartTimeUtc != null))
-						{
+						{							
 							newError = JobStepBatchError.SyncingFailed;
+							_logger.LogError("Batch {BatchId} marked with error SyncingFailed", batchId);
 						}
 						else
-						{
+						{							
 							newError = JobStepBatchError.Incomplete;
+							_logger.LogError("Batch {BatchId} marked with error Incomplete", batchId);
 						}
 					}
 				}
