Horde: Add job audit logging for tracing of the state machine

diff --git a/Engine/Source/Programs/Horde/Plugins/Build/HordeServer.Build/BuildPlugin.cs b/Engine/Source/Programs/Horde/Plugins/Build/HordeServer.Build/BuildPlugin.cs
--- a/Engine/Source/Programs/Horde/Plugins/Build/HordeServer.Build/BuildPlugin.cs
+++ b/Engine/Source/Programs/Horde/Plugins/Build/HordeServer.Build/BuildPlugin.cs
@@ -4,6 +4,7 @@ using System.Diagnostics.CodeAnalysis;
 using EpicGames.Core;
 using EpicGames.Horde.Artifacts;
 using EpicGames.Horde.Devices;
+using EpicGames.Horde.Jobs;
 using EpicGames.Horde.Jobs.Templates;
 using EpicGames.Horde.Streams;
 using HordeServer.Acls;
@@ -114,6 +115,7 @@ namespace HordeServer
 
 			services.AddSingleton<IssueService>();
 			services.AddSingleton<JobService>();
+			services.AddSingleton<IAuditLog<JobId>>(sp => sp.GetRequiredService<IAuditLogFactory<JobId>>().Create("Jobs.Log", "JobId"));
 			services.AddSingleton<ILogExtAuthProvider>(sp => sp.GetRequiredService<JobService>());
 			services.AddSingleton<INotificationService, NotificationService>();
 			services.AddSingleton<UnsyncCache>();
diff --git a/Engine/Source/Programs/Horde/Plugins/Build/HordeServer.Build/Jobs/IJobCollection.cs b/Engine/Source/Programs/Horde/Plugins/Build/HordeServer.Build/Jobs/IJobCollection.cs
--- a/Engine/Source/Programs/Horde/Plugins/Build/HordeServer.Build/Jobs/IJobCollection.cs
+++ b/Engine/Source/Programs/Horde/Plugins/Build/HordeServer.Build/Jobs/IJobCollection.cs
@@ -193,6 +193,14 @@ namespace HordeServer.Jobs
 		/// <returns>Sorted list of jobs to execute</returns>
 		Task<IReadOnlyList<IJob>> GetDispatchQueueAsync(CancellationToken cancellationToken = default);
 
+		/// <summary>
+		/// Write to the job's audit log
+		/// </summary>
+		/// <param name="jobId"></param>
+		/// <param name="message"></param>
+		/// <param name="args"></param>
+		void JobAudit(JobId jobId, string? message, params object?[] args);
+
 		/// <summary>
 		/// Upgrade all documents in the collection
 		/// </summary>
diff --git a/Engine/Source/Programs/Horde/Plugins/Build/HordeServer.Build/Jobs/JobCollection.cs b/Engine/Source/Programs/Horde/Plugins/Build/HordeServer.Build/Jobs/JobCollection.cs
--- a/Engine/Source/Programs/Horde/Plugins/Build/HordeServer.Build/Jobs/JobCollection.cs
+++ b/Engine/Source/Programs/Horde/Plugins/Build/HordeServer.Build/Jobs/JobCollection.cs
@@ -20,6 +20,7 @@ using EpicGames.Horde.Notifications;
 using EpicGames.Horde.Streams;
 using EpicGames.Horde.Users;
 using HordeServer.Acls;
+using HordeServer.Auditing;
 using HordeServer.Commits;
 using HordeServer.Jobs.Graphs;
 using HordeServer.Logs;
@@ -767,11 +768,12 @@ namespace HordeServer.Jobs
 		readonly Tracer _tracer;
 		readonly IOptionsMonitor<BuildConfig> _buildConfig;
 		readonly ILogger<JobCollection> _logger;
+		readonly IAuditLog<JobId> _auditLog;
 
 		/// <summary>
 		/// Constructor
 		/// </summary>
-		public JobCollection(IMongoService mongoService, IClock clock, IGraphCollection graphCollection, ILogCollection logCollection, ICommitService commitService, ITelemetryWriter telemetryWriter, IOptionsMonitor<BuildConfig> buildConfig, Tracer tracer, ILogger<JobCollection> logger)
+		public JobCollection(IMongoService mongoService, IClock clock, IGraphCollection graphCollection, ILogCollection logCollection, ICommitService commitService, ITelemetryWriter telemetryWriter, IOptionsMonitor<BuildConfig> buildConfig, Tracer tracer, ILogger<JobCollection> logger, IAuditLog<JobId> auditLog)
 		{
 			_clock = clock;
 			_graphCollection = graphCollection;
@@ -781,6 +783,7 @@ namespace HordeServer.Jobs
 			_buildConfig = buildConfig;
 			_tracer = tracer;
 			_logger = logger;
+			_auditLog = auditLog;
 
 			List<MongoIndex<JobDocument>> indexes = new List<MongoIndex<JobDocument>>();
 			indexes.Add(keys => keys.Ascending(x => x.StreamId).Descending(x => x.CreateTimeUtc));
@@ -794,6 +797,39 @@ namespace HordeServer.Jobs
 			indexes.Add(keys => keys.Descending(x => x.SchedulePriority));
 			indexes.Add(_startedByBisectTaskIdIndex = MongoIndex.Create<JobDocument>(keys => keys.Descending(x => x.StartedByBisectTaskId), sparse: true));
 			_jobs = mongoService.GetCollection<JobDocument>("Jobs", indexes);
+		}		
+
+		IAuditLogChannel<JobId> GetJobLogger(JobId jobId)
+		{
+			return _auditLog[jobId];
+		}
+
+		public void JobAudit(JobId jobId, string? message, params object?[] args)
+		{
+			try
+			{
+#pragma warning disable CA2254 // Template should be a static expression
+				GetJobLogger(jobId).LogInformation(message, args);
+#pragma warning restore CA2254 // Template should be a static expression
+			}
+			catch (Exception ex)
+			{
+				_logger.LogError(ex, "Error logging to job {JobId} audit", jobId);
+			}
+		}
+
+		static void JobAudit(JobId jobId, ILogger logger, IAuditLogChannel<JobId> jobLogger, string? message, params object?[] args)
+		{
+			try
+			{
+#pragma warning disable CA2254 // Template should be a static expression
+				jobLogger.LogInformation(message, args);
+#pragma warning restore CA2254 // Template should be a static expression
+			}
+			catch (Exception ex)
+			{
+				logger.LogError(ex, "Error logging to job {JobId} audit", jobId);
+			}
 		}
 
 		static Task PostLoadAsync(JobDocument job)
@@ -829,7 +865,10 @@ namespace HordeServer.Jobs
 			}
 
 			JobDocument newJob = new JobDocument(jobId, streamId, templateRefId, templateHash, graph.Id, name, commitIdWithOrder, codeCommitIdWithOrder, options, DateTime.UtcNow);
-			CreateBatches(newJob, graph, _logger);
+
+			JobAudit(jobId, "Job {JobId} added, stream {StreamId}, templateRefId {TemplateId}, name {Name}, commitId {CommitId}", jobId, streamId.Id, templateRefId, name, commitId);
+
+			CreateBatches(newJob, graph, _logger, GetJobLogger(jobId));
 
 			await _jobs.InsertOneAsync(newJob, null, cancellationToken);
 
@@ -873,7 +912,7 @@ namespace HordeServer.Jobs
 
 		/// <inheritdoc/>
 		async Task<bool> TryDeleteAsync(JobDocument job, CancellationToken cancellationToken)
-		{
+		{			
 			DeleteResult result = await _jobs.DeleteOneAsync(x => x.Id == job.Id && x.UpdateIndex == job.UpdateIndex, null, cancellationToken);
 			foreach (JobStepDocument step in job.Batches.SelectMany(x => x.Steps))
 			{
@@ -886,6 +925,12 @@ namespace HordeServer.Jobs
 					}
 				}
 			}
+
+			if (result.DeletedCount > 1)
+			{
+				JobAudit(job.Id, "Job {JobId} deleted", job.Id);
+			}
+
 			return result.DeletedCount > 0;
 		}
 
@@ -1178,37 +1223,46 @@ namespace HordeServer.Jobs
 			// Flag for whether to update batches
 			bool updateBatches = false;
 
+			JobId jobId = jobDocument.Id;
+
 			// Build the update list
 			jobDocument = Clone(jobDocument);
+			
 			if (name != null)
 			{
 				jobDocument.Name = name;
 				updates.Add(updateBuilder.Set(x => x.Name, jobDocument.Name));
+				JobAudit(jobId, "Job {JobId} Name set {Name}", name);
 			}
 			if (priority != null)
 			{
 				jobDocument.Priority = priority.Value;
 				updates.Add(updateBuilder.Set(x => x.Priority, jobDocument.Priority));
+				JobAudit(jobId, "Job {JobId} Priority set {Priority}", priority.Value);
 			}
 			if (autoSubmit != null)
 			{
 				jobDocument.AutoSubmit = autoSubmit.Value;
 				updates.Add(updateBuilder.Set(x => x.AutoSubmit, jobDocument.AutoSubmit));
+				JobAudit(jobId, "Job {JobId} AutoSubmit set {AutoSubmit}", autoSubmit.Value);
 			}
 			if (autoSubmitChange != null)
 			{
 				jobDocument.AutoSubmitChange = autoSubmitChange.Value;
 				updates.Add(updateBuilder.Set(x => x.AutoSubmitChange, jobDocument.AutoSubmitChange));
+				JobAudit(jobId, "Job {JobId} AutoSubmitChange set {AutoSubmitChange}", autoSubmitChange.Value);
 			}
 			if (autoSubmitMessage != null)
 			{
 				jobDocument.AutoSubmitMessage = (autoSubmitMessage.Length == 0) ? null : autoSubmitMessage;
 				updates.Add(updateBuilder.SetOrUnsetNullRef(x => x.AutoSubmitMessage, jobDocument.AutoSubmitMessage));
+				JobAudit(jobId, "Job {JobId} AutoSubmitMessage set {AutoSubmitMessage}", autoSubmitMessage);
 			}
 			if (abortedByUserId != null && jobDocument.AbortedByUserId == null)
 			{
 				jobDocument.AbortedByUserId = abortedByUserId;
 				updates.Add(updateBuilder.Set(x => x.AbortedByUserId, jobDocument.AbortedByUserId));
+				JobAudit(jobId, "Job {JobId} AbortedByUserId set {AbortedByUserId}", abortedByUserId.Value);
 				updateBatches = true;
 			}
 
@@ -1216,12 +1270,14 @@ namespace HordeServer.Jobs
 			{
 				jobDocument.CancellationReason = cancellationReason;
 				updates.Add(updateBuilder.Set(x => x.CancellationReason, jobDocument.CancellationReason));
+				JobAudit(jobId, "Job {JobId} CancellationReason set {CancellationReason}", cancellationReason);
 			}
 
 			if (notificationTriggerId != null)
 			{
 				jobDocument.NotificationTriggerId = notificationTriggerId.Value;
 				updates.Add(updateBuilder.Set(x => x.NotificationTriggerId, notificationTriggerId));
+				JobAudit(jobId, "Job {JobId} NotificationTriggerId set {NotificationTriggerId}", notificationTriggerId.Value);
 			}
 			if (labelIdxToTriggerId != null)
 			{
@@ -1230,7 +1286,7 @@ namespace HordeServer.Jobs
 					throw new ArgumentException("Cannot update label trigger that already exists");
 				}
 				jobDocument._labelNotifications.Add(new LabelNotificationDocument { _labelIdx = labelIdxToTriggerId.Value.Key, _triggerId = labelIdxToTriggerId.Value.Value });
-				updates.Add(updateBuilder.Set(x => x._labelNotifications, jobDocument._labelNotifications));
+				updates.Add(updateBuilder.Set(x => x._labelNotifications, jobDocument._labelNotifications));				
 			}
 			if (jobTrigger != null)
 			{
@@ -1242,6 +1298,13 @@ namespace HordeServer.Jobs
 						int localIdx = idx;
 						jobTriggerDocument.JobId = jobTrigger.Value.Value;
 						updates.Add(updateBuilder.Set(x => x.ChainedJobs[localIdx].JobId, jobTrigger.Value.Value));
+						try
+						{
+							JobAudit(jobId, "Job {JobId}, ChainedJobs Index {Index}, set to JobTrigger {JobId}", jobId, localIdx, jobTrigger.Value.Value.Id);
+						}
+						catch (Exception)
+						{
+						}						
 					}
 				}
 			}
@@ -1251,7 +1314,10 @@ namespace HordeServer.Jobs
 				jobDocument.Reports.RemoveAll(x => reports.Any(y => y.Name == x.Name));
 				jobDocument.Reports.AddRange(reports);
 				updates.Add(updateBuilder.Set(x => x.Reports, jobDocument.Reports));
+
+				JobAudit(jobId, "Job {JobId}, adding Reports {Reports}", jobId, String.Join(',', reports.Select(x => x.Name)));
 			}
+
 			if (arguments != null)
 			{
 				HashSet<string> modifiedArguments = new HashSet<string>(jobDocument.Arguments);
@@ -1267,12 +1333,13 @@ namespace HordeServer.Jobs
 
 				jobDocument.Arguments = arguments.ToList();
 				updates.Add(updateBuilder.Set(x => x.Arguments, jobDocument.Arguments));
+				JobAudit(jobId, "Job {JobId}, setting Arguments {Arguments}", jobId, String.Join(',', jobDocument.Arguments));
 			}
 
 			// Update the batches
 			if (updateBatches)
 			{
-				UpdateBatches(jobDocument, graph, updates, _logger);
+				UpdateBatches(jobDocument, graph, updates, _logger, GetJobLogger(jobId));
 			}
 
 			// Update the new list of job steps
@@ -1358,7 +1425,7 @@ namespace HordeServer.Jobs
 
 				// Force an update of all batches in the job. This will fail or reschedule any nodes that can no longer be executed in this batch.
 				updates.Clear();
-				UpdateBatches(jobDocument, graph, updates, _logger);
+				UpdateBatches(jobDocument, graph, updates, _logger, GetJobLogger(jobDocument.Id));
 
 				if (retriedNodes.Count > 0)
 				{
@@ -1602,14 +1669,14 @@ namespace HordeServer.Jobs
 			if (refreshBatches)
 			{
 				updates.Clear(); // UpdateBatches will update the entire batches list. We need to remove all the individual step updates to avoid an exception.
-				UpdateBatches(jobDocument, graph, updates, _logger);
+				UpdateBatches(jobDocument, graph, updates, _logger, GetJobLogger(jobDocument.Id));
 			}
 
 			// Update the state of dependent jobsteps
 			if (refreshDependentJobSteps)
 			{
-				RefreshDependentJobSteps(jobDocument, graph, updates, _logger);
-				RefreshJobPriority(jobDocument, updates);
+				RefreshDependentJobSteps(jobDocument, graph, updates, _logger, GetJobLogger(jobDocument.Id));
+				RefreshJobPriority(jobDocument, updates, _logger, GetJobLogger(jobDocument.Id));
 			}
 
 			// Update the new list of job steps
@@ -1741,7 +1808,7 @@ namespace HordeServer.Jobs
 			jobDocument.GraphHash = newGraph.Id;
 			updates.Add(updateBuilder.Set(x => x.GraphHash, jobDocument.GraphHash));
 
-			UpdateBatches(jobDocument, newGraph, updates, _logger);
+			UpdateBatches(jobDocument, newGraph, updates, _logger, GetJobLogger(jobDocument.Id));
 
 			// Update the new list of job steps
 			return await TryUpdateAsync(jobDocument, updates, cancellationToken);
@@ -1802,6 +1869,8 @@ namespace HordeServer.Jobs
 				if (batch.State == JobStepBatchState.Ready || batch.State == JobStepBatchState.Waiting)
 				{
 					_logger.LogInformation("Skipping all batches: job {JobId}, batch {BatchId}", jobDocument.Id, batch.Id);
+					JobAudit(jobDocument.Id, "Skipping all batches: job {JobId}, batch {BatchId}", jobDocument.Id, batch.Id);
+
 					batch.State = JobStepBatchState.Complete;
 					batch.Error = reason;
 					batch.FinishTime = DateTimeOffset.UtcNow;
@@ -1819,7 +1888,7 @@ namespace HordeServer.Jobs
 			}
 
 			List<UpdateDefinition<JobDocument>> updates = new List<UpdateDefinition<JobDocument>>();
-			UpdateBatches(jobDocument, graph, updates, _logger);
+			UpdateBatches(jobDocument, graph, updates, _logger, GetJobLogger(jobDocument.Id));
 
 			return await TryUpdateAsync(jobDocument, updates, cancellationToken);
 		}
@@ -1828,6 +1897,7 @@ namespace HordeServer.Jobs
 		async Task<JobDocument?> TrySkipBatchAsync(JobDocument jobDocument, JobStepBatchId batchId, IGraph graph, JobStepBatchError reason, CancellationToken cancellationToken)
 		{
 			_logger.LogInformation("Trying to skip batch: job {JobId}, batch {BatchId}", jobDocument.Id, batchId);
+			JobAudit(jobDocument.Id, "Trying to skip batch: {BatchId}", batchId);
 
 			JobStepBatchDocument? batch = jobDocument.Batches.FirstOrDefault(x => x.Id == batchId);
 			if (batch == null)
@@ -1849,7 +1919,7 @@ namespace HordeServer.Jobs
 			}
 
 			List<UpdateDefinition<JobDocument>> updates = new List<UpdateDefinition<JobDocument>>();
-			UpdateBatches(jobDocument, graph, updates, _logger);
+			UpdateBatches(jobDocument, graph, updates, _logger, GetJobLogger(jobDocument.Id));
 
 			return await TryUpdateAsync(jobDocument, updates, cancellationToken);
 		}
@@ -1872,6 +1942,8 @@ namespace HordeServer.Jobs
 			{
 				string currentSessionId = jobDocument.Batches[batchIdx].SessionId!.Value.ToString();
 				_logger.LogError("Attempt to replace current session ID {CurrSessionId} with {NewSessionId} for batch {JobId}:{BatchId}", currentSessionId, sessionId.ToString(), jobDocument.Id.ToString(), jobDocument.Batches[batchIdx].Id);
+				JobAudit(jobDocument.Id, "Attempt to replace current session ID {CurrSessionId} with {NewSessionId} for batch {JobId}:{BatchId}", currentSessionId, sessionId.ToString(), jobDocument.Id.ToString(), jobDocument.Batches[batchIdx].Id);
+
 				return null;
 			}
 
@@ -1882,6 +1954,7 @@ namespace HordeServer.Jobs
 		async Task<JobDocument?> TryCancelLeaseAsync(JobDocument jobDocument, int batchIdx, CancellationToken cancellationToken)
 		{
 			_logger.LogInformation("Cancelling lease {LeaseId} for agent {AgentId}", jobDocument.Batches[batchIdx].LeaseId, jobDocument.Batches[batchIdx].AgentId);
+			JobAudit(jobDocument.Id, "Cancelling lease {LeaseId} for agent {AgentId}", jobDocument.Batches[batchIdx].LeaseId, jobDocument.Batches[batchIdx].AgentId);
 
 			jobDocument = Clone(jobDocument);
 
@@ -1906,6 +1979,7 @@ namespace HordeServer.Jobs
 
 			JobStepBatchDocument batch = jobDocument.Batches[batchIdx];
 			_logger.LogInformation("Failing batch {JobId}:{BatchId} with error {Error}", jobDocument.Id, batch.Id, error);
+			JobAudit(jobDocument.Id, "Failing batch {JobId}:{BatchId} with error {Error}", jobDocument.Id, batch.Id, error);
 
 			UpdateDefinitionBuilder<JobDocument> updateBuilder = new UpdateDefinitionBuilder<JobDocument>();
 			List<UpdateDefinition<JobDocument>> updates = new List<UpdateDefinition<JobDocument>>();
@@ -1950,8 +2024,8 @@ namespace HordeServer.Jobs
 				}
 			}
 
-			RefreshDependentJobSteps(jobDocument, graph, updates, _logger);
-			RefreshJobPriority(jobDocument, updates);
+			RefreshDependentJobSteps(jobDocument, graph, updates, _logger, GetJobLogger(jobDocument.Id));
+			RefreshJobPriority(jobDocument, updates, _logger, GetJobLogger(jobDocument.Id));
 
 			return await TryUpdateAsync(jobDocument, updateBuilder.Combine(updates), cancellationToken);
 		}
@@ -1962,9 +2036,10 @@ namespace HordeServer.Jobs
 		/// <param name="job">The job to update</param>
 		/// <param name="graph">The graph definition for this job</param>
 		/// <param name="logger">Logger for output messages</param>
-		private static void CreateBatches(JobDocument job, IGraph graph, ILogger logger)
+		/// <param name="jobAudit"></param>
+		private static void CreateBatches(JobDocument job, IGraph graph, ILogger logger, IAuditLogChannel<JobId> jobAudit)
 		{
-			UpdateBatches(job, graph, new List<UpdateDefinition<JobDocument>>(), logger);
+			UpdateBatches(job, graph, new List<UpdateDefinition<JobDocument>>(), logger, jobAudit);
 		}
 
 		/// <summary>
@@ -1974,18 +2049,19 @@ namespace HordeServer.Jobs
 		/// <param name="graph">Graph definition for this job</param>
 		/// <param name="updates">List of updates for the job</param>
 		/// <param name="logger">Logger for updates</param>
-		static void UpdateBatches(JobDocument job, IGraph graph, List<UpdateDefinition<JobDocument>> updates, ILogger logger)
+		/// <param name="jobAudit"></param>
+		static void UpdateBatches(JobDocument job, IGraph graph, List<UpdateDefinition<JobDocument>> updates, ILogger logger, IAuditLogChannel<JobId> jobAudit)
 		{
 			UpdateDefinitionBuilder<JobDocument> updateBuilder = Builders<JobDocument>.Update;
 
 			// Update the list of batches
-			CreateOrUpdateBatches(job, graph, logger);
+			CreateOrUpdateBatches(job, graph, logger, jobAudit);
 			updates.Add(updateBuilder.Set(x => x.Batches, job.Batches));
 			updates.Add(updateBuilder.Set(x => x.NextSubResourceId, job.NextSubResourceId));
 
 			// Update all the dependencies
-			RefreshDependentJobSteps(job, graph, new List<UpdateDefinition<JobDocument>>(), logger);
-			RefreshJobPriority(job, updates);
+			RefreshDependentJobSteps(job, graph, new List<UpdateDefinition<JobDocument>>(), logger, jobAudit);
+			RefreshJobPriority(job, updates, logger, jobAudit);
 		}
 
 		static void RemoveSteps(JobStepBatchDocument batch, Predicate<JobStepDocument> predicate, Dictionary<NodeRef, JobStepId> recycleStepIds)
@@ -2020,8 +2096,12 @@ namespace HordeServer.Jobs
 		/// <param name="job">The job to update</param>
 		/// <param name="graph">The graph for this job</param>
 		/// <param name="logger">Logger for any changes</param>
-		private static void CreateOrUpdateBatches(JobDocument job, IGraph graph, ILogger logger)
+		/// <param name="jobAudit"></param>
+		private static void CreateOrUpdateBatches(JobDocument job, IGraph graph, ILogger logger, IAuditLogChannel<JobId> jobAudit)
 		{
+
+			JobAudit(job.Id, logger, jobAudit, "CreateOrUpdateBatches: Begin");
+
 			// Find the priorities of each node, incorporating all the per-step overrides
 			Dictionary<INode, Priority> nodePriorities = new Dictionary<INode, Priority>();
 			foreach (INodeGroup group in graph.Groups)
@@ -2060,7 +2140,8 @@ namespace HordeServer.Jobs
 					{
 						if (JobStepExtensions.IsPendingState(step.State))
 						{
-							step.State = JobStepState.Skipped;
+							step.State = JobStepState.Skipped;							
+							JobAudit(job.Id, logger, jobAudit, "Setting batch {BatchId}, step {StepId} to Skipped", batch.Id, step.Id);
 						}
 					}
 				}
@@ -2108,6 +2189,10 @@ namespace HordeServer.Jobs
 			// Remove any batches which are now empty
 			Dictionary<int, JobStepBatchId> recycleBatchIds = new Dictionary<int, JobStepBatchId>();
 			RemoveBatches(job, x => x.Steps.Count == 0 && x.LeaseId == null && x.Error == JobStepBatchError.None, recycleBatchIds);
+			if (recycleBatchIds.Count > 0)
+			{
+				JobAudit(job.Id, logger, jobAudit, "Recycled removed batch ids {RecycledIds}", String.Join(',', recycleBatchIds.Values));
+			}
 
 			// Find all the targets in this job
 			HashSet<string> targets = new HashSet<string>(StringComparer.OrdinalIgnoreCase);
@@ -2123,6 +2208,8 @@ namespace HordeServer.Jobs
 				targets.Add(IJob.SetupNodeName);
 			}
 
+			JobAudit(job.Id, logger, jobAudit, "Job targets {Targets}", String.Join(',', targets));
+
 			// Add all the referenced aggregates
 			HashSet<INode> newNodesToExecute = new HashSet<INode>();
 			foreach (IAggregate aggregate in graph.Aggregates)
@@ -2131,7 +2218,7 @@ namespace HordeServer.Jobs
 				{
 					newNodesToExecute.UnionWith(aggregate.Nodes.Select(x => graph.GetNode(x)));
 				}
-			}
+			}			
 
 			// Add any individual nodes
 			foreach (INode node in graph.Groups.SelectMany(x => x.Nodes))
@@ -2231,6 +2318,16 @@ namespace HordeServer.Jobs
 				}
 			}
 
+			if (existingNodesToExecute.Count > 0)
+			{
+				JobAudit(job.Id, logger, jobAudit, "Existing Nodes to execute {Nodes}", String.Join(',', existingNodesToExecute.Select(x => x.Name)));
+			}
+
+			if (newNodesToExecute.Count > 0)
+			{
+				JobAudit(job.Id, logger, jobAudit, "New Nodes to execute {Nodes}", String.Join(',', newNodesToExecute.Select(x => x.Name)));
+			}				
+
 			// Figure out the existing batch for each group
 			JobStepBatchDocument?[] appendToBatches = new JobStepBatchDocument?[graph.Groups.Count];
 			foreach (JobStepBatchDocument batch in job.Batches)
@@ -2285,8 +2382,9 @@ namespace HordeServer.Jobs
 
 							batch = new JobStepBatchDocument(batchId, groupIdx);
 							job.Batches.Add(batch);
-
 							appendToBatches[groupIdx] = batch;
+
+							JobAudit(job.Id, logger, jobAudit, "Created new batch {BatchId} for {NodeName}", batchId, node.Name);							
 						}
 
 						// Don't re-add nodes that have already executed in this batch. If we were missing a dependency, we would have already nulled out the entry in appendToBatches above; anything else
@@ -2302,6 +2400,7 @@ namespace HordeServer.Jobs
 
 							JobStepDocument step = new JobStepDocument(stepId, nodeIdx);
 							batch.Steps.Add(step);
+							JobAudit(job.Id, logger, jobAudit, "Created step {StepId} to batch {BatchId} for {StepName}", step.Id, batch.Id, node.Name);
 						}
 					}
 				}
@@ -2356,6 +2455,8 @@ namespace HordeServer.Jobs
 					nodeExecutionCount[node] = count + 1;
 				}
 			}
+
+			JobAudit(job.Id, logger, jobAudit, "CreateOrUpdateBatches: End");
 		}
 
 		/// <summary>
@@ -2403,7 +2504,7 @@ namespace HordeServer.Jobs
 		/// <param name="graph">The graph for this job</param>
 		/// <param name="updates">List of updates to the job</param>
 		/// <param name="logger">Logger instance</param>
-		static void RefreshDependentJobSteps(JobDocument job, IGraph graph, List<UpdateDefinition<JobDocument>> updates, ILogger logger)
+		static void RefreshDependentJobSteps(JobDocument job, IGraph graph, List<UpdateDefinition<JobDocument>> updates, ILogger logger, IAuditLogChannel<JobId> jobAudit)
 		{
 			// Update the batches
 			UpdateDefinitionBuilder<JobDocument> updateBuilder = Builders<JobDocument>.Update;
@@ -2435,6 +2536,7 @@ namespace HordeServer.Jobs
 							else if (!steps.Any(x => !x.AbortRequested && JobStepExtensions.IsPendingState(x.State)))
 							{
 								logger.LogInformation("Transitioning job {JobId}, batch {BatchId}, step {StepId} to ready state ({Dependencies})", job.Id, batch.Id, step.Id, String.Join(", ", steps.Select(x => x.Id.ToString())));
+								JobAudit(job.Id, logger, jobAudit, "Transitioning job {JobId}, batch {BatchId}, step {StepId} to ready state ({Dependencies})", job.Id, batch.Id, step.Id, String.Join(", ", steps.Select(x => x.Id.ToString())));
 								newState = JobStepState.Ready;
 							}
 						}
@@ -2443,12 +2545,14 @@ namespace HordeServer.Jobs
 						{
 							step.State = newState;
 							updates.Add(updateBuilder.Set(x => x.Batches[batchIdx].Steps[stepIdx].State, newState));
+							JobAudit(job.Id, logger, jobAudit, "Setting step {StepId} state to {StepState}", step.Id, newState);
 						}
 
 						if (newOutcome != step.Outcome)
 						{
 							step.Outcome = newOutcome;
 							updates.Add(updateBuilder.Set(x => x.Batches[batchIdx].Steps[stepIdx].Outcome, newOutcome));
+							JobAudit(job.Id, logger, jobAudit, "Setting step {StepId} outcome to {StepOutcome}", step.Id, newOutcome);
 						}
 
 						stepForNode[node] = step;
@@ -2462,11 +2566,13 @@ namespace HordeServer.Jobs
 						{
 							batch.State = newState;
 							logger.LogInformation("Transitioning job {JobId}, batch {BatchId} to state {JobStepBatchState}", job.Id, batch.Id, newState);
+							JobAudit(job.Id, logger, jobAudit, "Transitioning job {JobId}, batch {BatchId} to state {JobStepBatchState}", job.Id, batch.Id, newState);
 							updates.Add(updateBuilder.Set(x => x.Batches[batchIdx].State, batch.State));
 						}
 						if (batch.ReadyTime != newReadyTime)
 						{
 							logger.LogInformation("Setting job {JobId}, batch {BatchId} ready time", job.Id, batch.Id);
+							JobAudit(job.Id, logger, jobAudit, "Setting job {JobId}, batch {BatchId} ready time", job.Id, batch.Id);
 							batch.ReadyTime = newReadyTime;
 							updates.Add(updateBuilder.Set(x => x.Batches[batchIdx].ReadyTime, batch.ReadyTime));
 						}
@@ -2480,7 +2586,9 @@ namespace HordeServer.Jobs
 		/// </summary>
 		/// <param name="job"></param>
 		/// <param name="updates"></param>
-		static void RefreshJobPriority(JobDocument job, List<UpdateDefinition<JobDocument>> updates)
+		/// <param name="logger"></param>
+		/// <param name="jobAudit"></param>
+		static void RefreshJobPriority(JobDocument job, List<UpdateDefinition<JobDocument>> updates, ILogger logger, IAuditLogChannel<JobId> jobAudit)
 		{
 			// Update the weighted priority for the job
 			int newSchedulePriority = GetSchedulePriority(job);
@@ -2488,6 +2596,7 @@ namespace HordeServer.Jobs
 			{
 				job.SchedulePriority = newSchedulePriority;
 				updates.Add(Builders<JobDocument>.Update.Set(x => x.SchedulePriority, newSchedulePriority));
+				JobAudit(job.Id, logger, jobAudit, "Setting job schedule priority {SchedulePriority}", job.Id, newSchedulePriority);
 			}
 		}
 
diff --git a/Engine/Source/Programs/Horde/Plugins/Build/HordeServer.Build/Jobs/JobService.cs b/Engine/Source/Programs/Horde/Plugins/Build/HordeServer.Build/Jobs/JobService.cs
--- a/Engine/Source/Programs/Horde/Plugins/Build/HordeServer.Build/Jobs/JobService.cs
+++ b/Engine/Source/Programs/Horde/Plugins/Build/HordeServer.Build/Jobs/JobService.cs
@@ -17,6 +17,7 @@ using EpicGames.Horde.Notifications;
 using EpicGames.Horde.Streams;
 using EpicGames.Horde.Users;
 using HordeServer.Agents;
+using HordeServer.Auditing;
 using HordeServer.Jobs.Bisect;
 using HordeServer.Jobs.Graphs;
 using HordeServer.Jobs.Templates;
@@ -85,6 +86,7 @@ namespace HordeServer.Jobs
 		readonly IOptionsMonitor<BuildConfig> _buildConfig;
 		readonly Tracer _tracer;
 		readonly ILogger _logger;
+		readonly IAuditLog<JobId> _auditLog;
 
 		/// <summary>
 		/// Delegate for label update events
@@ -126,7 +128,7 @@ namespace HordeServer.Jobs
 		/// <summary>
 		/// Constructor
 		/// </summary>
-		public JobService(IJobCollection jobs, IGraphCollection graphs, IAgentCollection agents, IJobStepRefCollection jobStepRefs, IBisectTaskCollection bisectTasks, IJobTimingCollection jobTimings, IUserCollection userCollection, INotificationTriggerCollection triggerCollection, JobTaskSource jobTaskSource, IStreamCollection streamCollection, ITemplateCollectionInternal templateCollection, IPerforceService perforceService, Tracer tracer, IOptionsMonitor<BuildConfig> buildConfig, ILogger<JobService> logger)
+		public JobService(IJobCollection jobs, IGraphCollection graphs, IAgentCollection agents, IJobStepRefCollection jobStepRefs, IBisectTaskCollection bisectTasks, IJobTimingCollection jobTimings, IUserCollection userCollection, INotificationTriggerCollection triggerCollection, JobTaskSource jobTaskSource, IStreamCollection streamCollection, ITemplateCollectionInternal templateCollection, IPerforceService perforceService, Tracer tracer, IOptionsMonitor<BuildConfig> buildConfig, ILogger<JobService> logger, IAuditLog<JobId> auditLog)
 		{
 			_jobs = jobs;
 			_graphs = graphs;
@@ -143,6 +145,7 @@ namespace HordeServer.Jobs
 			_buildConfig = buildConfig;
 			_tracer = tracer;
 			_logger = logger;
+			_auditLog = auditLog;
 
 			_jobTaskSource.OnJobScheduled += (pool, numAgentsOnline, job, graph, batchId) =>
 			{
@@ -200,6 +203,7 @@ namespace HordeServer.Jobs
 			using IDisposable? scope = _logger.BeginScope("CreateJobAsync({JobId})", jobIdValue);
 
 			_logger.LogInformation("Creating job at CL {Change}, code CL {CodeChange}, preflight CL {PreflightChange}", commitId, codeCommitId, options.PreflightCommitId);
+			_jobs.JobAudit(jobIdValue, "Creating job at CL {Change}, code CL {CodeChange}, preflight CL {PreflightChange}", commitId, codeCommitId, options.PreflightCommitId);
 
 			Dictionary<string, string> properties = new Dictionary<string, string>();
 			properties["Change"] = commitId.ToString();
@@ -274,6 +278,7 @@ namespace HordeServer.Jobs
 				if (updatedJob == null)
 				{
 					_logger.LogError("Failed marking duplicate job as aborted! Job ID: {JobId}", job.Id);
+					_jobs.JobAudit(job.Id, "Failed marking duplicate job as aborted!");
 				}
 
 				IJob? updatedJob2 = await GetJobAsync(job.Id, cancellationToken);
@@ -769,6 +774,7 @@ namespace HordeServer.Jobs
 								if (await agent.TryUpdateAsync(new UpdateAgentOptions { RequestConform = true }, cancellationToken: cancellationToken) != null)
 								{
 									_logger.LogError("Agent {AgentId} did not complete lease; marking for conform", agent.Id);
+									_jobs.JobAudit(job.Id, "Agent {AgentId} did not complete lease; marking for conform", agent.Id);
 									break;
 								}
 							}
@@ -780,16 +786,19 @@ namespace HordeServer.Jobs
 						{
 							newError = error.Value;
 							_logger.LogError("Batch {BatchId} marked with error {JobStepBatchError}", batchId, error.Value.ToString());
+							_jobs.JobAudit(job.Id, "Batch {BatchId} marked with error {JobStepBatchError}", batchId, error.Value.ToString());
 						}
 						else if (!batch.Steps.Any(x => x.StartTimeUtc != null))
 						{							
 							newError = JobStepBatchError.SyncingFailed;
 							_logger.LogError("Batch {BatchId} marked with error SyncingFailed", batchId);
+							_jobs.JobAudit(job.Id, "Batch {BatchId} marked with error SyncingFailed", batchId);
 						}
 						else
 						{							
 							newError = JobStepBatchError.Incomplete;
 							_logger.LogError("Batch {BatchId} marked with error Incomplete", batchId);
+							_jobs.JobAudit(job.Id, "Batch {BatchId} marked with error Incomplete", batchId);
 						}
 					}
 				}
@@ -822,6 +831,7 @@ namespace HordeServer.Jobs
 					if (result != CheckShelfResult.Ok)
 					{
 						_logger.LogWarning("Job {JobId} is no longer valid - check shelf returned {Result}", job.Id, result);
+						_jobs.JobAudit(job.Id, "Job {JobId} is no longer valid - check shelf returned {Result}", job.Id, result);
 						return JobStepBatchError.UnknownShelf;
 					}
 				}
@@ -830,6 +840,7 @@ namespace HordeServer.Jobs
 			catch (Exception ex)
 			{
 				_logger.LogWarning(ex, "Job {JobId} is no longer valid.", job.Id);
+				_jobs.JobAudit(job.Id, "Job {JobId} is no longer valid, {ExceptionMessage}.", job.Id, ex.Message);
 				return JobStepBatchError.ExecutionError;
 			}
 		}
@@ -973,6 +984,7 @@ namespace HordeServer.Jobs
 				if (oldState != updatedState || oldOutcome != updatedOutcome)
 				{
 					_logger.LogInformation("Transitioned job {JobId}, batch {BatchId}, step {StepId} from {OldState} to {NewState}", job.Id, batchId, stepId, oldState, updatedState);
+					_jobs.JobAudit(job.Id, "Transitioned job {JobId}, batch {BatchId}, step {StepId} from {OldState} to {NewState}", job.Id, batchId, stepId, oldState, updatedState);
 
 					// Send any updates for modified badges
 					if (oldLabelStates != null)
@@ -999,6 +1011,7 @@ namespace HordeServer.Jobs
 								if (settings.AlwaysTagPreflightCL)
 								{
 									_logger.LogInformation("Updating description for {PreflightChange} for {UserId} user settings", job.PreflightCommitId, job.StartedByUserId.Value);
+									_jobs.JobAudit(job.Id, "Updating description for {PreflightChange} for {UserId} user settings", job.PreflightCommitId, job.StartedByUserId.Value);
 									await _perforceService.UpdateChangelistDescriptionAsync(streamConfig.ClusterName, job.PreflightCommitId.GetPerforceChange(), x => x.TrimEnd() + $"\n#preflight {job.Id}", cancellationToken);
 								}
 							}
@@ -1085,6 +1098,7 @@ namespace HordeServer.Jobs
 				CommitId preflightCommitId = job.PreflightCommitId;
 
 				_logger.LogInformation("Updating description for {PreflightChange}", preflightCommitId);
+				_jobs.JobAudit(job.Id, "Updating description for {PreflightChange}", preflightCommitId);
 
 				await _perforceService.UpdateChangelistDescriptionAsync(streamConfig.ClusterName, preflightCommitId.GetPerforceChange(), x => x.TrimEnd() + $"\n#preflight {job.Id}", cancellationToken);
 
@@ -1092,6 +1106,7 @@ namespace HordeServer.Jobs
 				(change, message) = await _perforceService.SubmitShelvedChangeAsync(streamConfig, preflightCommitId.GetPerforceChange(), job.PreflightCommitId.GetPerforceChange(), cancellationToken);
 
 				_logger.LogInformation("Attempt to submit {Change} (through {ChangeCopy}): {Message}", job.PreflightCommitId, preflightCommitId, message);
+				_jobs.JobAudit(job.Id, "Attempt to submit {Change} (through {ChangeCopy}): {Message}", job.PreflightCommitId, preflightCommitId, message);
 
 				if (!String.IsNullOrEmpty(message))
 				{
@@ -1107,6 +1122,8 @@ namespace HordeServer.Jobs
 			{
 				(change, message) = ((int?)null, "Internal error");
 				_logger.LogError(ex, "Unable to submit shelved change");
+				_jobs.JobAudit(job.Id, "Unable to submit shelved change: {ExceptionMessage}", ex.Message);
+
 			}
 
 			for (; ; )
@@ -1206,6 +1223,7 @@ namespace HordeServer.Jobs
 					if (!streamConfig.TryGetTemplate(jobTrigger.TemplateRefId, out templateRefConfig))
 					{
 						_logger.LogWarning("Cannot find template {TemplateId} in stream {StreamId}", jobTrigger.TemplateRefId, newJob.StreamId);
+						_jobs.JobAudit(newJob.Id, "Cannot find template {TemplateId} in stream {StreamId}", jobTrigger.TemplateRefId, newJob.StreamId);
 						break;
 					}
 
@@ -1238,6 +1256,7 @@ namespace HordeServer.Jobs
 					}
 
 					_logger.LogInformation("Creating downstream job {ChainedJobId} from job {JobId} at change {Change}", chainedJobId, newJob.Id, commitId);
+					_jobs.JobAudit(newJob.Id, "Creating downstream job {ChainedJobId} from job {JobId} at change {Change}", chainedJobId, newJob.Id, commitId);
 
 					CreateJobOptions options = new CreateJobOptions(templateRefConfig);
 					options.PreflightCommitId = newJob.PreflightCommitId;
@@ -1338,5 +1357,15 @@ namespace HordeServer.Jobs
 		{
 			return String.Join(",", job.Batches.Select(b => b.SessionId));
 		}
+
+		/// <summary>
+		/// Get the logger for the specified job id
+		/// </summary>
+		/// <param name="jobId"></param>
+		/// <returns></returns>
+		public IAuditLogChannel<JobId> GetJobLogger(JobId jobId)
+		{
+			return _auditLog[jobId];
+		}
 	}
 }
diff --git a/Engine/Source/Programs/Horde/Plugins/Build/HordeServer.Build/Jobs/JobsController.cs b/Engine/Source/Programs/Horde/Plugins/Build/HordeServer.Build/Jobs/JobsController.cs
--- a/Engine/Source/Programs/Horde/Plugins/Build/HordeServer.Build/Jobs/JobsController.cs
+++ b/Engine/Source/Programs/Horde/Plugins/Build/HordeServer.Build/Jobs/JobsController.cs
@@ -1,6 +1,7 @@
 // Copyright Epic Games, Inc. All Rights Reserved.
 
 using EpicGames.Core;
+using EpicGames.Horde.Agents;
 using EpicGames.Horde.Artifacts;
 using EpicGames.Horde.Commits;
 using EpicGames.Horde.Jobs;
@@ -11,6 +12,7 @@ using EpicGames.Horde.Logs;
 using EpicGames.Horde.Notifications;
 using EpicGames.Horde.Streams;
 using EpicGames.Horde.Users;
+using HordeServer.Auditing;
 using HordeServer.Acls;
 using HordeServer.Agents;
 using HordeServer.Commits;
@@ -1856,5 +1858,43 @@ namespace HordeServer.Jobs
 			}
 			return new GetNotificationResponse(subscription);
 		}
+
+		/// <summary>
+		/// Retrieve historical information about a specific job
+		/// </summary>
+		/// <param name="jobId">Id of the job to get information about</param>
+		/// <param name="minTime">Minimum time for records to return</param>
+		/// <param name="maxTime">Maximum time for records to return</param>
+		/// <param name="index">Offset of the first result</param>
+		/// <param name="count">Number of records to return</param>
+		/// <returns>Information about the requested agent</returns>
+		[HttpGet]
+		[Route("/api/v1/jobs/{jobId}/history")]
+		public async Task<ActionResult> GetJobHistoryAsync(JobId jobId, [FromQuery] DateTime? minTime = null, [FromQuery] DateTime? maxTime = null, [FromQuery] int index = 0, [FromQuery] int count = 50, CancellationToken cancellationToken = default)
+		{
+			IJob? job = await _jobService.GetJobAsync(jobId, cancellationToken);
+			if (job == null)
+			{
+				return NotFound(jobId);
+			}
+
+			StreamConfig? streamConfig;
+			if (!_buildConfig.Value.TryGetStream(job.StreamId, out streamConfig))
+			{
+				return NotFound(job.StreamId);
+			}
+			if (!streamConfig.Authorize(JobAclAction.ViewJob, User))
+			{
+				return Forbid(JobAclAction.ViewJob, job.StreamId);
+			}
+
+			Response.ContentType = "application/json";
+			Response.StatusCode = 200;
+			await Response.StartAsync();
+			await _jobService.GetJobLogger(jobId).FindAsync(Response.BodyWriter, minTime, maxTime, index, count);
+			await Response.CompleteAsync();
+
+			return Empty;
+		}
 	}
 }
