Horde: Fix some server warnings

diff --git a/Engine/Source/Programs/Horde/Plugins/Build/HordeServer.Build/Jobs/JobCollection.cs b/Engine/Source/Programs/Horde/Plugins/Build/HordeServer.Build/Jobs/JobCollection.cs
--- a/Engine/Source/Programs/Horde/Plugins/Build/HordeServer.Build/Jobs/JobCollection.cs
+++ b/Engine/Source/Programs/Horde/Plugins/Build/HordeServer.Build/Jobs/JobCollection.cs
@@ -799,11 +799,22 @@ namespace HordeServer.Jobs
 			_jobs = mongoService.GetCollection<JobDocument>("Jobs", indexes);
 		}		
 
+		/// <summary>
+		/// Get Job Logger
+		/// </summary>
+		/// <param name="jobId"></param>
+		/// <returns></returns>
 		IAuditLogChannel<JobId> GetJobLogger(JobId jobId)
 		{
 			return _auditLog[jobId];
 		}
 
+		/// <summary>
+		/// Log to a job audit
+		/// </summary>
+		/// <param name="jobId"></param>
+		/// <param name="message"></param>
+		/// <param name="args"></param>
 		public void JobAudit(JobId jobId, string? message, params object?[] args)
 		{
 			try
@@ -818,6 +829,14 @@ namespace HordeServer.Jobs
 			}
 		}
 
+		/// <summary>
+		/// Log to job audit static flavor
+		/// </summary>
+		/// <param name="jobId"></param>
+		/// <param name="logger"></param>
+		/// <param name="jobLogger"></param>
+		/// <param name="message"></param>
+		/// <param name="args"></param>
 		static void JobAudit(JobId jobId, ILogger logger, IAuditLogChannel<JobId> jobLogger, string? message, params object?[] args)
 		{
 			try
diff --git a/Engine/Source/Programs/Horde/Plugins/Build/HordeServer.Build/Jobs/JobsController.cs b/Engine/Source/Programs/Horde/Plugins/Build/HordeServer.Build/Jobs/JobsController.cs
--- a/Engine/Source/Programs/Horde/Plugins/Build/HordeServer.Build/Jobs/JobsController.cs
+++ b/Engine/Source/Programs/Horde/Plugins/Build/HordeServer.Build/Jobs/JobsController.cs
@@ -1,7 +1,6 @@
 // Copyright Epic Games, Inc. All Rights Reserved.
 
 using EpicGames.Core;
-using EpicGames.Horde.Agents;
 using EpicGames.Horde.Artifacts;
 using EpicGames.Horde.Commits;
 using EpicGames.Horde.Jobs;
@@ -1890,8 +1889,8 @@ namespace HordeServer.Jobs
 
 			Response.ContentType = "application/json";
 			Response.StatusCode = 200;
-			await Response.StartAsync();
-			await _jobService.GetJobLogger(jobId).FindAsync(Response.BodyWriter, minTime, maxTime, index, count);
+			await Response.StartAsync(cancellationToken);
+			await _jobService.GetJobLogger(jobId).FindAsync(Response.BodyWriter, minTime, maxTime, index, count, cancellationToken);
 			await Response.CompleteAsync();
 
 			return Empty;
