Horde: Add job audit view to dashboard
#skipci

diff --git a/Engine/Source/Programs/Horde/HordeDashboard/src/App.tsx b/Engine/Source/Programs/Horde/HordeDashboard/src/App.tsx
--- a/Engine/Source/Programs/Horde/HordeDashboard/src/App.tsx
+++ b/Engine/Source/Programs/Horde/HordeDashboard/src/App.tsx
@@ -160,6 +160,7 @@ const Main: React.FC = () => {
                { path: "audit/agent/:agentId", element: <AuditLogView /> },
                { path: "audit/issue/:issueId", element: <AuditLogView /> },
                { path: "audit/device/:deviceId", element: <AuditLogView /> },
+               { path: "audit/job/:jobId", element: <AuditLogView /> },
                { path: "audit/template/:streamId/:templateId", element: <AuditLogView /> },
                { path: "automation", element: <AutomationView /> },
                { path: "tools", element: <ToolView /> },
diff --git a/Engine/Source/Programs/Horde/HordeDashboard/src/backend/Backend.ts b/Engine/Source/Programs/Horde/HordeDashboard/src/backend/Backend.ts
--- a/Engine/Source/Programs/Horde/HordeDashboard/src/backend/Backend.ts
+++ b/Engine/Source/Programs/Horde/HordeDashboard/src/backend/Backend.ts
@@ -1661,6 +1661,15 @@ export class Backend {
         });
     }
 
+    getJobHistory(id: string, query: AuditLogQuery): Promise<AuditLogEntry[]> {
+        return new Promise<AuditLogEntry[]>((resolve, reject) => {
+            this._fetch.get(`/api/v1/jobs/${encodeURIComponent(id)}/history`, { params: query }).then((response) => {
+                const history = (response.data?.entries ?? []) as AuditLogEntry[];
+                resolve(history);
+            }).catch(reason => { reject(reason); });
+        });
+    }
+
     // get a job's bisection tasks
     getJobBisectTasks(jobId: string): Promise<GetBisectTaskResponse[]> {
         return new Promise<GetBisectTaskResponse[]>((resolve, reject) => {
diff --git a/Engine/Source/Programs/Horde/HordeDashboard/src/components/AuditLog.tsx b/Engine/Source/Programs/Horde/HordeDashboard/src/components/AuditLog.tsx
--- a/Engine/Source/Programs/Horde/HordeDashboard/src/components/AuditLog.tsx
+++ b/Engine/Source/Programs/Horde/HordeDashboard/src/components/AuditLog.tsx
@@ -67,6 +67,7 @@ class AuditLogHandler {
       this.issueId = undefined;
       this.deviceId = undefined;
       this.templateAudit = undefined;
+      this.jobId = undefined;
 
       this.timeSelectKey = "time_live_tail";
 
@@ -132,6 +133,21 @@ class AuditLogHandler {
       this.update();
    }
 
+   setJob(jobId?: string) {
+
+      if (!jobId) {
+         this.clear();
+         return;
+      }
+
+      if (this.jobId === jobId) {
+         return;
+      }
+
+      this.jobId = jobId;
+      this.update();
+   }
+
    setTemplateAudit(templateAudit?: TemplateAudit) {
 
       if (!templateAudit) {
@@ -159,6 +175,10 @@ class AuditLogHandler {
       return !!this.deviceId;
    }
 
+   get isJobLog(): boolean {
+      return !!this.jobId;
+   }
+
    get isTemplateAuditLog(): boolean {
       return !!this.templateAudit;
    }
@@ -182,7 +202,7 @@ class AuditLogHandler {
          this.canceled.add(i);
       }
 
-      if (!this.agentId && !this.issueId && !this.templateAudit && !this.deviceId) {
+      if (!this.agentId && !this.issueId && !this.templateAudit && !this.deviceId && !this.jobId) {
          return;
       }
 
@@ -231,6 +251,10 @@ class AuditLogHandler {
             entries = await backend.getTemplateHistory(this.templateAudit.streamId, this.templateAudit.templateId, { minTime: minTime?.toISOString(), maxTime: maxTime?.toISOString(), count: 8192 * 2 });
          }
 
+         if (this.jobId) {
+            entries = await backend.getJobHistory(this.jobId, { minTime: minTime?.toISOString(), maxTime: maxTime?.toISOString(), count: 8192 * 2 });
+         }
+
          // check for canceled during graph request
          if (this.canceled.has(cancelId)) {
             return;
@@ -298,6 +322,7 @@ class AuditLogHandler {
    agentId?: string;
    issueId?: string;
    deviceId?: string;
+   jobId?: string;
    templateAudit?: TemplateAudit;
 
    timeSelectKey?: string;
@@ -369,7 +394,7 @@ const timeSelections: TimeSelection[] = [
 
 const selection = new Selection({ selectionMode: SelectionMode.multiple });
 
-export const AuditLogPanel: React.FC<{ agentId?: string, deviceId?: string, issueId?: string, templateAudit?: TemplateAudit }> = observer(({ agentId, deviceId, issueId, templateAudit }) => {
+export const AuditLogPanel: React.FC<{ agentId?: string, deviceId?: string, jobId?: string, issueId?: string, templateAudit?: TemplateAudit }> = observer(({ agentId, deviceId, jobId, issueId, templateAudit }) => {
 
    const windowSize = useWindowSize();
    const [showDatePicker, setShowDatePicker] = useState(false);
@@ -401,6 +426,10 @@ export const AuditLogPanel: React.FC<{ agentId?: string, deviceId?: string, issu
       handler.setDevice(deviceId);
    }
 
+   if (jobId && handler.jobId !== jobId) {
+      handler.setJob(jobId);
+   }
+
    if (issueId && handler.issueId !== issueId) {
       handler.setIssue(issueId);
    }
@@ -409,7 +438,7 @@ export const AuditLogPanel: React.FC<{ agentId?: string, deviceId?: string, issu
       handler.setTemplateAudit(templateAudit);
    }
 
-   if (!agentId && !deviceId && !issueId && !templateAudit) {
+   if (!agentId && !deviceId && !issueId && !templateAudit && !jobId) {
       return null;
    }
 
@@ -506,6 +535,19 @@ export const AuditLogPanel: React.FC<{ agentId?: string, deviceId?: string, issu
       ];
    }
 
+   if (jobId) {
+      crumbItems = [
+         {
+            text: `Job`,
+            link: `/job/${jobId}`
+         },
+         {
+            text: `Audit - ${jobId}`,
+            link: `/job/${jobId}`
+         }
+      ];
+   }
+
 
    if (issueId) {
       crumbItems = [
@@ -708,7 +750,7 @@ export const AuditLogView: React.FC = () => {
 
    const { hordeClasses } = getHordeStyling();
 
-   const { agentId, deviceId, issueId, streamId, templateId } = useParams<{ agentId: string, deviceId: string, issueId: string, streamId: string, templateId: string }>();
+   const { agentId, deviceId, jobId, issueId, streamId, templateId } = useParams<{ agentId: string, deviceId: string, jobId, issueId: string, streamId: string, templateId: string }>();
 
    let templateAudit: TemplateAudit | undefined;
    if (!!streamId && !!templateId) {
@@ -724,6 +766,7 @@ export const AuditLogView: React.FC = () => {
          {!!agentId && <AuditLogPanel agentId={agentId} />}
          {!!deviceId && <AuditLogPanel deviceId={deviceId} />}
          {!!issueId && <AuditLogPanel issueId={issueId} />}
+         {!!jobId && <AuditLogPanel jobId={jobId} />}
          {!!templateAudit && <AuditLogPanel templateAudit={templateAudit} />}
       </Stack>
    </Stack>
@@ -739,6 +782,7 @@ enum AuditLogType {
    Unknown,
    Agent,
    Device,
+   Job,
    Issue,
    Template
 }
@@ -778,6 +822,10 @@ class AuditLine {
          instance = new DeviceAuditLine(logType, entry);
       }
 
+      if (logType === AuditLogType.Job) {
+         instance = new JobAuditLine(logType, entry);
+      }
+
       if (instance) {
          instance.search = search;
          instance.process();
@@ -1020,6 +1068,10 @@ class DeviceAuditLine extends AuditLine {
 
 }
 
+class JobAuditLine extends AuditLine {
+
+}
+
 // Line Rendering -------------------------------------------------------------------------------------------------------------------
 
 
@@ -1043,6 +1095,10 @@ const renderAuditEntry = (entry: AuditLogEntry, search?: string) => {
       type = AuditLogType.Device;
    }
 
+   if (handler.isJobLog) {
+      type = AuditLogType.Job;
+   }
+
    if (type === AuditLogType.Unknown) {
       return <div>Unknown Log Type</div>
    }
diff --git a/Engine/Source/Programs/Horde/HordeDashboard/src/components/jobDetailsV2/JobOperationsBar.tsx b/Engine/Source/Programs/Horde/HordeDashboard/src/components/jobDetailsV2/JobOperationsBar.tsx
--- a/Engine/Source/Programs/Horde/HordeDashboard/src/components/jobDetailsV2/JobOperationsBar.tsx
+++ b/Engine/Source/Programs/Horde/HordeDashboard/src/components/jobDetailsV2/JobOperationsBar.tsx
@@ -162,6 +162,16 @@ export const JobOperations: React.FC<{ jobDetails: JobDetailsV2 }> = observer(({
       onClick: () => { setParametersState(ParameterState.Clone); }
    });
 
+   opsList.push({
+      key: 'jobops_audit',
+      text: "Audit",
+      iconProps: { iconName: "SearchTemplate" },
+      onClick: () => { 
+         const href = `/audit/job/${jobId}`;
+         window.open(href, "_blank")
+       }
+   });
+
    const opsItems: ICommandBarItemProps[] = [
       {
          key: 'jobops_items',
