Horde: Reverse job audit logs to be top down
#skipci

diff --git a/Engine/Source/Programs/Horde/HordeDashboard/src/components/AuditLog.tsx b/Engine/Source/Programs/Horde/HordeDashboard/src/components/AuditLog.tsx
--- a/Engine/Source/Programs/Horde/HordeDashboard/src/components/AuditLog.tsx
+++ b/Engine/Source/Programs/Horde/HordeDashboard/src/components/AuditLog.tsx
@@ -492,7 +492,7 @@ export const AuditLogPanel: React.FC<{ agentId?: string, deviceId?: string, jobI
 
    });
 
-   if (issueId) {
+   if (issueId || jobId) {
       logItems.reverse();
    }
 
