Horde: Fix even more server warnings

diff --git a/Engine/Source/Programs/Horde/Plugins/Build/HordeServer.Build/Jobs/JobCollection.cs b/Engine/Source/Programs/Horde/Plugins/Build/HordeServer.Build/Jobs/JobCollection.cs
--- a/Engine/Source/Programs/Horde/Plugins/Build/HordeServer.Build/Jobs/JobCollection.cs
+++ b/Engine/Source/Programs/Horde/Plugins/Build/HordeServer.Build/Jobs/JobCollection.cs
@@ -2523,6 +2523,7 @@ namespace HordeServer.Jobs
 		/// <param name="graph">The graph for this job</param>
 		/// <param name="updates">List of updates to the job</param>
 		/// <param name="logger">Logger instance</param>
+		/// <param name="jobAudit"></param>
 		static void RefreshDependentJobSteps(JobDocument job, IGraph graph, List<UpdateDefinition<JobDocument>> updates, ILogger logger, IAuditLogChannel<JobId> jobAudit)
 		{
 			// Update the batches
diff --git a/Engine/Source/Programs/Horde/Plugins/Build/HordeServer.Build/Jobs/JobsController.cs b/Engine/Source/Programs/Horde/Plugins/Build/HordeServer.Build/Jobs/JobsController.cs
--- a/Engine/Source/Programs/Horde/Plugins/Build/HordeServer.Build/Jobs/JobsController.cs
+++ b/Engine/Source/Programs/Horde/Plugins/Build/HordeServer.Build/Jobs/JobsController.cs
@@ -1866,6 +1866,7 @@ namespace HordeServer.Jobs
 		/// <param name="maxTime">Maximum time for records to return</param>
 		/// <param name="index">Offset of the first result</param>
 		/// <param name="count">Number of records to return</param>
+		/// <param name="cancellationToken"></param>
 		/// <returns>Information about the requested agent</returns>
 		[HttpGet]
 		[Route("/api/v1/jobs/{jobId}/history")]
