Update Device Info dialog after editing
Allow Device Audit to be opened in another tab by using Link node

#rnx

diff --git a/Engine/Source/Programs/Horde/HordeDashboard/src/components/devices/DeviceEditor.tsx b/Engine/Source/Programs/Horde/HordeDashboard/src/components/devices/DeviceEditor.tsx
--- a/Engine/Source/Programs/Horde/HordeDashboard/src/components/devices/DeviceEditor.tsx
+++ b/Engine/Source/Programs/Horde/HordeDashboard/src/components/devices/DeviceEditor.tsx
@@ -18,7 +18,7 @@ type DeviceEditData = {
    notes?: string;
 };
 
-export const DeviceEditor: React.FC<{ handler: DeviceHandler, deviceIn?: GetDeviceResponse | undefined, editNote?: boolean, onClose: () => void }> = ({ handler, deviceIn, editNote, onClose }) => {
+export const DeviceEditor: React.FC<{ handler: DeviceHandler, deviceIn?: GetDeviceResponse | undefined, editNote?: boolean, onClose: (device?: GetDeviceResponse) => void }> = ({ handler, deviceIn, editNote, onClose }) => {
 
    const [state, setState] = useState<{ device?: DeviceEditData, title?: string }>({});
    const [error, setError] = useState<string | undefined>();
@@ -67,7 +67,17 @@ export const DeviceEditor: React.FC<{ handler: DeviceHandler, deviceIn?: GetDevi
       setSubmitting(false);
       setError("");
       setConfirmDelete(false);
-      onClose();
+      const updatedDevice: GetDeviceResponse | undefined = !!deviceIn? {
+         ...deviceIn,
+            platformId: device.platformId,
+            enabled: device.enabled,
+            address: device.address,
+            name: device.name,
+            poolId: device.poolId,
+            modelId: device.modelId,
+            notes: device.notes
+         }: undefined;
+      onClose(updatedDevice);
    };
 
    const onDelete = async () => {
@@ -230,7 +240,8 @@ export const DeviceEditor: React.FC<{ handler: DeviceHandler, deviceIn?: GetDevi
                address: device.address,
                poolId: device.poolId,
                platformId: device.platformId,
-               modelId: device.modelId
+               modelId: device.modelId,
+               enabled: device.enabled
             });
 
             if (!result) {
diff --git a/Engine/Source/Programs/Horde/HordeDashboard/src/components/devices/DeviceInfoView.tsx b/Engine/Source/Programs/Horde/HordeDashboard/src/components/devices/DeviceInfoView.tsx
--- a/Engine/Source/Programs/Horde/HordeDashboard/src/components/devices/DeviceInfoView.tsx
+++ b/Engine/Source/Programs/Horde/HordeDashboard/src/components/devices/DeviceInfoView.tsx
@@ -281,7 +281,7 @@ export const DeviceInfoModal: React.FC<{ handler: DeviceHandler, deviceIn?: GetD
                            <PrimaryButton text="Edit" onClick={() => { onEdit(deviceIn) }} />
                         </Stack>
                         <Stack style={{ paddingTop: 8, paddingRight: 24 }}>
-                           <DefaultButton text="Audit" onClick={() => { navigate(`/audit/device/${deviceIn.id}`) }} />
+                           <Link to={`/audit/device/${deviceIn.id}`} ><DefaultButton text="Audit" onClick={(ev) => { }} /></Link>
                         </Stack>
                      </Stack>
                   </Stack>
diff --git a/Engine/Source/Programs/Horde/HordeDashboard/src/components/devices/DeviceView.tsx b/Engine/Source/Programs/Horde/HordeDashboard/src/components/devices/DeviceView.tsx
--- a/Engine/Source/Programs/Horde/HordeDashboard/src/components/devices/DeviceView.tsx
+++ b/Engine/Source/Programs/Horde/HordeDashboard/src/components/devices/DeviceView.tsx
@@ -782,14 +782,14 @@ const DevicePanel: React.FC = observer(() => {
 
    const checkedOut = handler.getUserDeviceCheckouts(dashboard.userId);
 
-   const hasDevices = !!handler.getDevices().length
+   const hasDevices = !!handler.getDevices().length;
 
    return (<Stack>
       {!!telemState && <DevicePoolTelemetryModal poolsIn={localState.searchState.filterPools} platformsIn={localState.searchState.filterPlatforms} onClose={() => setTelemState(false)} />}
       {(!!checkoutState.checkinId || !!checkoutState.checkoutId) && <InfoModal />}
       {!!checkoutState.showConfirm && checkedOut.length > 0 && <CheckoutConfirmModal check={checkoutState.showConfirm} devices={checkedOut} onClose={() => setCheckoutState({})} />}
       {editState.infoShown && <DeviceInfoModal handler={handler} deviceIn={editState.device?.device} onEdit={(device) => { setEditState({ infoShown: true, shown: true, device: { device: device } }) }} onClose={() => { setEditState({}); localState.resetHistoryItem(); }} />}
-      {editState.shown && <DeviceEditor handler={handler} deviceIn={editState.device?.device} editNote={editState.editNote} onClose={() => { setEditState({ ...editState, shown: false }); localState.resetHistoryItem(); }} />}
+      {editState.shown && <DeviceEditor handler={handler} deviceIn={editState.device?.device} editNote={editState.editNote} onClose={(device) => { const deviceItem = !!device? { device: device }: undefined; setEditState({ ...editState, shown: false, device: deviceItem }); localState.resetHistoryItem(); }} />}
       <Stack styles={{ root: { paddingLeft: 12, paddingRight: 12, width: "100%" } }} >
          <Stack>
             <Stack horizontal verticalAlign="center" style={{ paddingBottom: 12 }} tokens={{ childrenGap: 32 }}>
