Horde: Prevent exception swallowing during job batch execution

diff --git a/Engine/Source/Programs/Horde/Drivers/JobDriver/Execution/JobExecutor.cs b/Engine/Source/Programs/Horde/Drivers/JobDriver/Execution/JobExecutor.cs
--- a/Engine/Source/Programs/Horde/Drivers/JobDriver/Execution/JobExecutor.cs
+++ b/Engine/Source/Programs/Horde/Drivers/JobDriver/Execution/JobExecutor.cs
@@ -459,14 +459,11 @@ namespace JobDriver.Execution
 			{
 				if (cancellationToken.IsCancellationRequested && IsCancellationException(ex))
 				{
-					if (!HordeClient.HasValidAccessToken())
-					{
-						logger.LogError(ex, "Exception while executing batch: {Ex}", ex);
-					}
-					else
-					{
-						logger.LogError("Lease was aborted");
-					}
+					logger.LogError("Lease was aborted");
+				}
+				else
+				{
+					logger.LogError(ex, "Exception while executing batch: {Ex}", ex);
 				}
 			}
 
