Horde: Fixing some job logging now that it is running in dev

diff --git a/Engine/Source/Programs/Horde/Plugins/Build/HordeServer.Build/Jobs/JobCollection.cs b/Engine/Source/Programs/Horde/Plugins/Build/HordeServer.Build/Jobs/JobCollection.cs
--- a/Engine/Source/Programs/Horde/Plugins/Build/HordeServer.Build/Jobs/JobCollection.cs
+++ b/Engine/Source/Programs/Horde/Plugins/Build/HordeServer.Build/Jobs/JobCollection.cs
@@ -1251,37 +1251,37 @@ namespace HordeServer.Jobs
 			{
 				jobDocument.Name = name;
 				updates.Add(updateBuilder.Set(x => x.Name, jobDocument.Name));
-				JobAudit(jobId, "Job {JobId} Name set {Name}", name);
+				JobAudit(jobId, "Job Name set {Name}", name);
 			}
 			if (priority != null)
 			{
 				jobDocument.Priority = priority.Value;
 				updates.Add(updateBuilder.Set(x => x.Priority, jobDocument.Priority));
-				JobAudit(jobId, "Job {JobId} Priority set {Priority}", priority.Value);
+				JobAudit(jobId, "Job Priority set {Priority}", priority.Value);
 			}
 			if (autoSubmit != null)
 			{
 				jobDocument.AutoSubmit = autoSubmit.Value;
 				updates.Add(updateBuilder.Set(x => x.AutoSubmit, jobDocument.AutoSubmit));
-				JobAudit(jobId, "Job {JobId} AutoSubmit set {AutoSubmit}", autoSubmit.Value);
+				JobAudit(jobId, "Job AutoSubmit set to {AutoSubmit}", autoSubmit.Value);
 			}
 			if (autoSubmitChange != null)
 			{
 				jobDocument.AutoSubmitChange = autoSubmitChange.Value;
 				updates.Add(updateBuilder.Set(x => x.AutoSubmitChange, jobDocument.AutoSubmitChange));
-				JobAudit(jobId, "Job {JobId} AutoSubmitChange set {AutoSubmitChange}", autoSubmitChange.Value);
+				JobAudit(jobId, "Job AutoSubmitChange set {AutoSubmitChange}", autoSubmitChange.Value);
 			}
 			if (autoSubmitMessage != null)
 			{
 				jobDocument.AutoSubmitMessage = (autoSubmitMessage.Length == 0) ? null : autoSubmitMessage;
 				updates.Add(updateBuilder.SetOrUnsetNullRef(x => x.AutoSubmitMessage, jobDocument.AutoSubmitMessage));
-				JobAudit(jobId, "Job {JobId} AutoSubmitMessage set {AutoSubmitMessage}", autoSubmitMessage);
+				JobAudit(jobId, "Job AutoSubmitMessage set {AutoSubmitMessage}", autoSubmitMessage);
 			}
 			if (abortedByUserId != null && jobDocument.AbortedByUserId == null)
 			{
 				jobDocument.AbortedByUserId = abortedByUserId;
 				updates.Add(updateBuilder.Set(x => x.AbortedByUserId, jobDocument.AbortedByUserId));
-				JobAudit(jobId, "Job {JobId} AbortedByUserId set {AbortedByUserId}", abortedByUserId.Value);
+				JobAudit(jobId, "Job AbortedByUserId set {AbortedByUserId}", abortedByUserId.Value);
 				updateBatches = true;
 			}
 
@@ -1289,14 +1289,14 @@ namespace HordeServer.Jobs
 			{
 				jobDocument.CancellationReason = cancellationReason;
 				updates.Add(updateBuilder.Set(x => x.CancellationReason, jobDocument.CancellationReason));
-				JobAudit(jobId, "Job {JobId} CancellationReason set {CancellationReason}", cancellationReason);
+				JobAudit(jobId, "Job CancellationReason set {CancellationReason}", cancellationReason);
 			}
 
 			if (notificationTriggerId != null)
 			{
 				jobDocument.NotificationTriggerId = notificationTriggerId.Value;
 				updates.Add(updateBuilder.Set(x => x.NotificationTriggerId, notificationTriggerId));
-				JobAudit(jobId, "Job {JobId} NotificationTriggerId set {NotificationTriggerId}", notificationTriggerId.Value);
+				JobAudit(jobId, "Job NotificationTriggerId set {NotificationTriggerId}", notificationTriggerId.Value);
 			}
 			if (labelIdxToTriggerId != null)
 			{
@@ -1334,7 +1334,7 @@ namespace HordeServer.Jobs
 				jobDocument.Reports.AddRange(reports);
 				updates.Add(updateBuilder.Set(x => x.Reports, jobDocument.Reports));
 
-				JobAudit(jobId, "Job {JobId}, adding Reports {Reports}", jobId, String.Join(',', reports.Select(x => x.Name)));
+				JobAudit(jobId, "Job, adding Reports {Reports}", String.Join(',', reports.Select(x => x.Name)));
 			}
 
 			if (arguments != null)
@@ -1352,7 +1352,7 @@ namespace HordeServer.Jobs
 
 				jobDocument.Arguments = arguments.ToList();
 				updates.Add(updateBuilder.Set(x => x.Arguments, jobDocument.Arguments));
-				JobAudit(jobId, "Job {JobId}, setting Arguments {Arguments}", jobId, String.Join(',', jobDocument.Arguments));
+				JobAudit(jobId, "Job setting Arguments {Arguments}", String.Join(',', jobDocument.Arguments));
 			}
 
 			// Update the batches
@@ -2616,7 +2616,7 @@ namespace HordeServer.Jobs
 			{
 				job.SchedulePriority = newSchedulePriority;
 				updates.Add(Builders<JobDocument>.Update.Set(x => x.SchedulePriority, newSchedulePriority));
-				JobAudit(job.Id, logger, jobAudit, "Setting job schedule priority {SchedulePriority}", job.Id, newSchedulePriority);
+				JobAudit(job.Id, logger, jobAudit, "Setting job schedule priority {SchedulePriority}", newSchedulePriority);
 			}
 		}
 
