Horde: Fix compiler warnings

diff --git a/Engine/Source/Programs/Horde/Plugins/Build/HordeServer.Build.Tests/Notifications/NotificationServiceTests.cs b/Engine/Source/Programs/Horde/Plugins/Build/HordeServer.Build.Tests/Notifications/NotificationServiceTests.cs
--- a/Engine/Source/Programs/Horde/Plugins/Build/HordeServer.Build.Tests/Notifications/NotificationServiceTests.cs
+++ b/Engine/Source/Programs/Horde/Plugins/Build/HordeServer.Build.Tests/Notifications/NotificationServiceTests.cs
@@ -240,7 +240,7 @@ namespace HordeServer.Tests.Notifications
 		async Task<IJob> AbortStepAsync(IJob job, int batchIdx, int stepIdx)
 		{
 			LogId logId = LogIdUtils.GenerateNewId();
-			ILog log = await LogCollection.AddAsync(job.Id, job.Batches[batchIdx].LeaseId, job.Batches[batchIdx].SessionId, LogType.Json, logId);
+			await LogCollection.AddAsync(job.Id, job.Batches[batchIdx].LeaseId, job.Batches[batchIdx].SessionId, LogType.Json, logId);
 			job = Deref(await JobService.UpdateStepAsync(job, job.Batches[batchIdx].Id, job.Batches[batchIdx].Steps[stepIdx].Id, _streamConfig, newAbortRequested: true, newAbortByUserId: UserId.Anonymous, newLogId: logId));
 			Assert.AreEqual(JobStepState.Aborted, job.Batches[batchIdx].Steps[stepIdx].State);
 			return job;
@@ -251,7 +251,7 @@ namespace HordeServer.Tests.Notifications
 			Assert.AreEqual(JobStepState.Ready, job.Batches[batchIdx].Steps[stepIdx].State);
 
 			LogId logId = LogIdUtils.GenerateNewId();
-			ILog log = await LogCollection.AddAsync(job.Id, job.Batches[batchIdx].LeaseId, job.Batches[batchIdx].SessionId, LogType.Json, logId);
+			await LogCollection.AddAsync(job.Id, job.Batches[batchIdx].LeaseId, job.Batches[batchIdx].SessionId, LogType.Json, logId);
 			job = Deref(await JobService.UpdateStepAsync(job, job.Batches[batchIdx].Id, job.Batches[batchIdx].Steps[stepIdx].Id, _streamConfig, JobStepState.Running, JobStepOutcome.Success, newLogId: logId));
 			Assert.AreEqual(JobStepState.Running, job.Batches[batchIdx].Steps[stepIdx].State);
 			job = Deref(await JobService.UpdateStepAsync(job, job.Batches[batchIdx].Id, job.Batches[batchIdx].Steps[stepIdx].Id, _streamConfig, JobStepState.Completed, outcome, newLogId: logId));
diff --git a/Engine/Source/Programs/Shared/EpicGames.Horde.Tests/Storage/BundleTests.cs b/Engine/Source/Programs/Shared/EpicGames.Horde.Tests/Storage/BundleTests.cs
--- a/Engine/Source/Programs/Shared/EpicGames.Horde.Tests/Storage/BundleTests.cs
+++ b/Engine/Source/Programs/Shared/EpicGames.Horde.Tests/Storage/BundleTests.cs
@@ -2,7 +2,6 @@
 
 using System;
 using System.Collections.Generic;
-using System.IO;
 using System.Linq;
 using System.Threading;
 using System.Threading.Tasks;
