Horde: Add database schema migration tool for MongoDB

diff --git a/Engine/Source/Programs/Horde/HordeServer.Shared/Server/MongoMigration.cs b/Engine/Source/Programs/Horde/HordeServer.Shared/Server/MongoMigration.cs
new file mode 100644
--- /dev/null
+++ b/Engine/Source/Programs/Horde/HordeServer.Shared/Server/MongoMigration.cs
@@ -0,0 +1,63 @@
+﻿// Copyright Epic Games, Inc. All Rights Reserved.
+
+using Microsoft.Extensions.Logging;
+using MongoDB.Driver;
+
+namespace HordeServer.Server;
+
+/// <summary>
+/// Attribute to decorate MongoDB migration classes with metadata
+/// </summary>
+[AttributeUsage(AttributeTargets.Class)]
+public sealed class MongoMigrationAttribute(string plugin, int version, string name, bool autoLoad = true) : Attribute
+{
+	/// <summary>
+	/// Plugin identifier for this migration.
+	/// </summary>
+	public string Plugin { get; } = plugin;
+	
+	/// <summary>
+	/// Gets the version number for this migration and plugin
+	/// Must be monotonically increasing and unique within in plugin namespace
+	/// </summary>
+	public int Version { get; } = version;
+	
+	/// <summary>
+	/// Descriptive name of the change this migration performs
+	/// </summary>
+	public string Name { get; } = name;
+	
+	/// <summary>
+	/// Whether migration should be detected and loaded during full reflection scans of available migrations
+	/// If set to false, the migration must manually be registered with a migrator
+	/// </summary>
+	public bool AutoLoad { get; } = autoLoad;
+}
+
+/// <summary>
+/// Provides context and services needed for executing migrations
+/// </summary>
+public record MigrationContext(IMongoService MongoService, ILogger Logger)
+{
+	/// <summary>
+	/// Shorthand for accessing the MongoDB database object
+	/// </summary>
+	public IMongoDatabase Database => MongoService.Database;
+}
+
+/// <summary>
+/// Defines operations for upgrading and downgrading database schemas
+/// Implemented by classes providing a migration
+/// </summary>
+public interface IMongoMigration
+{
+	/// <summary>
+	/// Upgrade the database schema version
+	/// </summary>
+	public Task UpAsync(MigrationContext context, CancellationToken token);
+	
+	/// <summary>
+	/// Downgrade the database schema version
+	/// </summary>
+	public Task DownAsync(MigrationContext context, CancellationToken token);
+}
\ No newline at end of file
diff --git a/Engine/Source/Programs/Horde/HordeServer.Tests/Server/MongoMigrationTests.cs b/Engine/Source/Programs/Horde/HordeServer.Tests/Server/MongoMigrationTests.cs
new file mode 100644
--- /dev/null
+++ b/Engine/Source/Programs/Horde/HordeServer.Tests/Server/MongoMigrationTests.cs
@@ -0,0 +1,277 @@
+﻿// Copyright Epic Games, Inc. All Rights Reserved.
+
+using System;
+using System.Collections.Generic;
+using System.Threading;
+using System.Threading.Tasks;
+using HordeServer.Server;
+using Microsoft.Extensions.Logging;
+using Microsoft.VisualStudio.TestTools.UnitTesting;
+using MongoDB.Driver;
+
+namespace HordeServer.Tests.Server;
+
+[TestClass]
+public class MongoMigratorTests : DatabaseIntegrationTest
+{
+	private abstract class TestMigration : IMongoMigration
+	{
+		public virtual Task UpAsync(MigrationContext context, CancellationToken token) { throw new NotImplementedException(); }
+		public virtual Task DownAsync(MigrationContext context, CancellationToken token) { throw new NotImplementedException(); }
+	}
+	
+	private class MissingAttributeMigration : TestMigration;
+	
+	[MongoMigration("duplicate", 1, "First upgrade", false)]
+	private class DuplicateVersionMigrationA : TestMigration;
+	
+	[MongoMigration("duplicate", 1, "First upgrade again", false)]
+	private class DuplicateVersionMigrationB : TestMigration;
+	
+	[MongoMigration("test", 0, "Invalid version", false)]
+	private class InvalidVersionMigration : TestMigration { }
+	
+	[MongoMigration("test", 1, "ab", false)]
+	private class ShortNameMigration : TestMigration { }
+	
+	[MongoMigration("failing", 1, "Failing migration", false)]
+	private class FailingUpgradeMigration : TestMigration
+	{
+		public override Task UpAsync(MigrationContext context, CancellationToken token) { throw new Exception("Simulated failure"); }
+	}
+	
+	[MongoMigration("failing", 1, "Failing migration", false)]
+	private class FailingDowngradeMigration : TestMigration
+	{
+		public override Task UpAsync(MigrationContext context, CancellationToken token) { return Task.CompletedTask; }
+		public override Task DownAsync(MigrationContext context, CancellationToken token) { throw new Exception("Simulated failure"); }
+	}
+	
+	[MongoMigration("foo", 1, "Create collection")]
+	private class FooMigration1 : TestMigration
+	{
+		public override async Task UpAsync(MigrationContext c, CancellationToken token)
+		{
+			await c.Database.CreateCollectionAsync("MyCollection-1", cancellationToken: token);
+		}
+		
+		public override async Task DownAsync(MigrationContext c, CancellationToken token)
+		{
+			await c.Database.DropCollectionAsync("MyCollection-1", cancellationToken: token);
+		}
+	}
+	
+	[MongoMigration("foo", 2, "Rename collection 1 -> 2")]
+	private class FooMigration2 : TestMigration
+	{
+		public override async Task UpAsync(MigrationContext c, CancellationToken token)
+		{
+			await c.Database.RenameCollectionAsync("MyCollection-1", "MyCollection-2", cancellationToken: token);
+		}
+		
+		public override async Task DownAsync(MigrationContext c, CancellationToken token)
+		{
+			await c.Database.RenameCollectionAsync("MyCollection-2", "MyCollection-1", cancellationToken: token);
+		}
+	}
+	
+	[MongoMigration("foo", 3, "Rename collection 2 -> 3")]
+	private class FooMigration3 : TestMigration
+	{
+		public override async Task UpAsync(MigrationContext c, CancellationToken token)
+		{
+			await c.Database.RenameCollectionAsync("MyCollection-2", "MyCollection-3", cancellationToken: token);
+		}
+		
+		public override async Task DownAsync(MigrationContext c, CancellationToken token)
+		{
+			await c.Database.RenameCollectionAsync("MyCollection-3", "MyCollection-2", cancellationToken: token);
+		}
+	}
+	
+	private readonly MongoMigrator _migrator;
+	private readonly IMongoService _mongoService;
+	
+	public MongoMigratorTests()
+	{
+		using ILoggerFactory loggerFactory = LoggerFactory.Create(builder =>
+		{
+			builder.SetMinimumLevel(LogLevel.Debug);
+			builder.AddSimpleConsole(options => { options.SingleLine = true; });
+		});
+		
+		_migrator = new MongoMigrator(GetMongoServiceSingleton(), loggerFactory.CreateLogger<MongoMigrator>());
+		_mongoService = GetMongoServiceSingleton();
+	}
+	
+	[TestMethod]
+	public void AddMigration_WithMissingAttribute()
+	{
+		MigrationException me = Assert.ThrowsException<MigrationException>(() => _migrator.AddMigration(new MissingAttributeMigration()));
+		Assert.IsTrue(me.Message.Contains("is missing [MongoMigration] attribute", StringComparison.Ordinal));
+	}
+	
+	[TestMethod]
+	public void AddMigration_WithDuplicateVersion()
+	{
+		_migrator.AddMigration(new DuplicateVersionMigrationA());
+		MigrationException me = Assert.ThrowsException<MigrationException>(() => _migrator.AddMigration(new DuplicateVersionMigrationB()));
+		Assert.IsTrue(me.Message.Contains("already contains a migration with version 1", StringComparison.Ordinal));
+	}
+	
+	[TestMethod]
+	public void AddMigration_WithOutOfOrderVersions()
+	{
+		_migrator.AddMigration(new FooMigration2());
+		_migrator.AddMigration(new FooMigration3());
+		_migrator.AddMigration(new FooMigration1());
+		IReadOnlyList<MigrationDefinition> migrations = _migrator.GetMigrations("foo");
+		Assert.AreEqual(3, migrations.Count);
+		Assert.AreEqual(1, migrations[0].Version);
+		Assert.AreEqual(2, migrations[1].Version);
+		Assert.AreEqual(3, migrations[2].Version);
+	}
+	
+	[TestMethod]
+	public void AutoAddMigrations_WithAllClasses()
+	{
+		_migrator.AutoAddMigrations();
+		IReadOnlyList<MigrationDefinition> migrations = _migrator.GetMigrations("foo");
+		Assert.AreEqual(3, migrations.Count);
+		Assert.AreEqual(1, migrations[0].Version);
+		Assert.AreEqual(2, migrations[1].Version);
+		Assert.AreEqual(3, migrations[2].Version);
+	}
+	
+	[TestMethod]
+	public async Task Upgrade_WithVersionGap_Async()
+	{
+		_migrator.AddMigration(new FooMigration1());
+		_migrator.AddMigration(new FooMigration3());
+		MigrationException me = await Assert.ThrowsExceptionAsync<MigrationException>(() => _migrator.UpgradeAsync("foo", null, CancellationToken.None));
+		Assert.IsTrue(me.Message.StartsWith("Gap detected", StringComparison.Ordinal));
+	}
+	
+	[TestMethod]
+	public void AddMigration_WithInvalidVersionNumber()
+	{
+		MigrationException me = Assert.ThrowsException<MigrationException>(() => _migrator.AddMigration(new InvalidVersionMigration()));
+		Assert.IsTrue(me.Message.Contains("invalid version number", StringComparison.Ordinal));
+	}
+	
+	[TestMethod]
+	public void AddMigration_WithInvalidName()
+	{
+		MigrationException me = Assert.ThrowsException<MigrationException>(() => _migrator.AddMigration(new ShortNameMigration()));
+		Assert.IsTrue(me.Message.Contains("too short/long name", StringComparison.Ordinal));
+	}
+	
+	[TestMethod]
+	public async Task UpgradeAsync()
+	{
+		_migrator.AddMigration(new FooMigration1());
+		await _migrator.UpgradeAsync("foo", null, CancellationToken.None);
+		AssertCollectionExists("MyCollection-1");
+	}
+	
+	[TestMethod]
+	public async Task Upgrade_WithSpecificVersion_Async()
+	{
+		// Test upgrading to a specific version, not all the way
+		_migrator.AddMigration(new FooMigration1());
+		_migrator.AddMigration(new FooMigration2());
+		_migrator.AddMigration(new FooMigration3());
+		
+		await _migrator.UpgradeAsync("foo", 2, CancellationToken.None);
+		
+		AssertCollectionExists("MyCollection-2");
+		SchemaVersionResult result = await _migrator.GetSchemaStateAsync("foo", CancellationToken.None);
+		Assert.AreEqual(2, result.DatabaseVersion);
+	}
+	
+	[TestMethod]
+	public async Task UpgradeAndDowngrade_AppliesCorrectly_Async()
+	{
+		_migrator.AddMigration(new FooMigration1());
+		await _migrator.UpgradeAsync("foo", null, CancellationToken.None);
+		AssertCollectionExists("MyCollection-1");
+		
+		_migrator.AddMigration(new FooMigration2());
+		await _migrator.UpgradeAsync("foo", null, CancellationToken.None);
+		AssertCollectionExists("MyCollection-2");
+		
+		await _migrator.DowngradeAsync("foo", 1, CancellationToken.None);
+		AssertCollectionExists("MyCollection-1");
+	}
+	
+	[TestMethod]
+	public async Task CheckSchemaVersion_WhenVersionsMatch_Async()
+	{
+		// Setup initial state at version 1
+		_migrator.AddMigration(new FooMigration1());
+		await _migrator.UpgradeAsync("foo", null, CancellationToken.None);
+		
+		// Check state
+		SchemaVersionResult result = await _migrator.GetSchemaStateAsync("foo", CancellationToken.None);
+		
+		Assert.AreEqual(SchemaVersionState.VersionMatch, result.State);
+		Assert.AreEqual(0, result.MissingMigrations.Count);
+	}
+	
+	[TestMethod]
+	public async Task CheckSchemaVersion_WhenDatabaseOutdated_Async()
+	{
+		// Add migrations but don't apply them yet
+		_migrator.AddMigration(new FooMigration1());
+		_migrator.AddMigration(new FooMigration2());
+		_migrator.AddMigration(new FooMigration3());
+		
+		// Upgrade only to version 1
+		await _migrator.UpgradeAsync("foo", 1, CancellationToken.None);
+		
+		// Check state - should show versions 2 and 3 as missing
+		SchemaVersionResult result = await _migrator.GetSchemaStateAsync("foo", CancellationToken.None);
+		
+		Console.WriteLine("AppVersion: " + result.AppVersion);
+		Console.WriteLine("DatabaseVersion: " + result.DatabaseVersion);
+		
+		Assert.AreEqual(SchemaVersionState.DatabaseOutdated, result.State);
+		Assert.AreEqual(2, result.MissingMigrations.Count);
+		Assert.AreEqual(2, result.MissingMigrations[0].Version);
+		Assert.AreEqual(3, result.MissingMigrations[1].Version);
+	}
+	
+	[TestMethod]
+	public async Task CheckSchemaVersion_WhenApplicationOutdated_Async()
+	{
+		_migrator.AddMigration(new FooMigration1());
+		_migrator.AddMigration(new FooMigration2());
+		_migrator.AddMigration(new FooMigration3());
+		await _migrator.UpgradeAsync("foo", null, CancellationToken.None);
+		
+		// Create new migrator instance with fewer migrations to simulate older app version
+		using ILoggerFactory loggerFactory = LoggerFactory.Create(builder => builder.AddSimpleConsole());
+		MongoMigrator newMigrator = new (_mongoService, loggerFactory.CreateLogger<MongoMigrator>());
+		newMigrator.AddMigration(new FooMigration1());
+		newMigrator.AddMigration(new FooMigration2());
+		
+		// Check state - should show version 3 as missing from application
+		SchemaVersionResult result = await newMigrator.GetSchemaStateAsync("foo", CancellationToken.None);
+		Assert.AreEqual(SchemaVersionState.ApplicationOutdated, result.State);
+		Assert.AreEqual(1, result.MissingMigrations.Count);
+		Assert.AreEqual(3, result.MissingMigrations[0].Version);
+	}
+	
+	[TestMethod]
+	public async Task CheckSchemaVersion_WithUnknownPlugin_Async()
+	{
+		MigrationException me = await Assert.ThrowsExceptionAsync<MigrationException>(() => _migrator.GetSchemaStateAsync("unknown-plugin", CancellationToken.None));
+		Assert.AreEqual("Unknown plugin unknown-plugin", me.Message);
+	}
+	
+	private void AssertCollectionExists(string name)
+	{
+		List<string> collectionNames = _mongoService.Database.ListCollectionNamesAsync().Result.ToList();
+		CollectionAssert.Contains(collectionNames, name, $"Collection {name} does not exist");
+	}
+}
\ No newline at end of file
diff --git a/Engine/Source/Programs/Horde/HordeServer/Commands/MongoCommands.cs b/Engine/Source/Programs/Horde/HordeServer/Commands/MongoCommands.cs
new file mode 100644
--- /dev/null
+++ b/Engine/Source/Programs/Horde/HordeServer/Commands/MongoCommands.cs
@@ -0,0 +1,76 @@
+// Copyright Epic Games, Inc. All Rights Reserved.
+
+using System;
+using System.Collections.Generic;
+using System.Threading;
+using System.Threading.Tasks;
+using EpicGames.Core;
+using HordeServer.Server;
+using Microsoft.Extensions.Configuration;
+using Microsoft.Extensions.DependencyInjection;
+using Microsoft.Extensions.Hosting;
+using Microsoft.Extensions.Logging;
+
+namespace HordeServer.Commands;
+
+[Command("mongo", "upgrade", "Upgrade database schema(s) by applying migration scripts")]
+class MongoUpgradeCommand(IConfiguration config) : Command
+{
+	/// <inheritdoc/>
+	public override async Task<int> ExecuteAsync(ILogger logger)
+	{
+		using IHost host = ServerCommand.CreateMinimalHostBuilder(config).Build();
+		MongoMigrator migrator = host.Services.GetRequiredService<MongoMigrator>();
+		
+		try
+		{
+			migrator.AutoAddMigrations();
+			await migrator.UpgradeAsync(CancellationToken.None);
+		}
+		catch (Exception e)
+		{
+			logger.LogError(e, "Database schema upgrade failed: {Exception}", e);
+			return 1;
+		}
+		
+		logger.LogInformation("Database schemas successfully upgraded!");
+		return 0;
+	}
+}
+
+[Command("mongo", "validate", "Validate database schema versions")]
+class MongoValidateCommand(IConfiguration config) : Command
+{
+	/// <inheritdoc/>
+	public override async Task<int> ExecuteAsync(ILogger logger)
+	{
+		using IHost host = ServerCommand.CreateMinimalHostBuilder(config).Build();
+		MongoMigrator migrator = host.Services.GetRequiredService<MongoMigrator>();
+
+		try
+		{
+			migrator.AutoAddMigrations();
+			Dictionary<string,SchemaVersionResult> results = await migrator.GetAllSchemaStateAsync(CancellationToken.None);
+			foreach ((string plugin, SchemaVersionResult svr) in results)
+			{
+				logger.LogInformation("Plugin {Plugin}: {State}", plugin, svr.State);
+				if (svr.MissingMigrations.Count > 0)
+				{
+					logger.LogInformation("\tMissing migrations:");
+					foreach (MigrationDefinition md in svr.MissingMigrations)
+					{
+						logger.LogInformation("\t{Version,3}: {Class} - {Name}", md.Version, md.Migration.GetType().FullName, md.Name);
+					}
+				}
+			}
+		}
+		catch (Exception e)
+		{
+			logger.LogError(e, "Database validation failed: {Exception}", e);
+			return 1;
+		}
+		
+		return 0;
+	}
+}
+
diff --git a/Engine/Source/Programs/Horde/HordeServer/Commands/ServerCommand.cs b/Engine/Source/Programs/Horde/HordeServer/Commands/ServerCommand.cs
--- a/Engine/Source/Programs/Horde/HordeServer/Commands/ServerCommand.cs
+++ b/Engine/Source/Programs/Horde/HordeServer/Commands/ServerCommand.cs
@@ -8,6 +8,7 @@ using System.Security.Cryptography.X509Certificates;
 using System.Threading;
 using System.Threading.Tasks;
 using EpicGames.Core;
+using HordeServer.Server;
 using Microsoft.AspNetCore.Hosting;
 using Microsoft.AspNetCore.Server.Kestrel.Core;
 using Microsoft.Extensions.Configuration;
@@ -34,27 +35,38 @@ namespace HordeServer.Commands
 			_hordeSettings = settings;
 			_config = config;
 		}
-
+		
+		/// <inheritdoc/>
 		public override void Configure(CommandLineArguments arguments, ILogger logger)
 		{
 			base.Configure(arguments, logger);
 			_args = arguments.GetRawArray();
 		}
-
+		
+		/// <inheritdoc/>
 		public override async Task<int> ExecuteAsync(ILogger logger)
 		{
 			logger.LogInformation("Server version: {Version}", ServerApp.Version);
 			logger.LogInformation("App directory: {AppDir}", ServerApp.AppDir);
 			logger.LogInformation("Data directory: {DataDir}", ServerApp.DataDir);
 			logger.LogInformation("Server config: {ConfigFile}", ServerApp.ServerConfigFile);
-
-			using (X509Certificate2? grpcCertificate = ReadGrpcCertificate(_hordeSettings))
+			
+			using X509Certificate2? grpcCertificate = ReadGrpcCertificate(_hordeSettings);
+			using IHost host = CreateHostBuilderWithCert(_args, _config, _hordeSettings, grpcCertificate).Build();
+			
+			if (_hordeSettings.MongoMigrationsEnabled)
 			{
-				using IHost host = CreateHostBuilderWithCert(_args, _config, _hordeSettings, grpcCertificate).Build();
-
-				await host.RunAsync();
-				return 0;
+				MongoMigrator migrator = host.Services.GetRequiredService<MongoMigrator>();
+				migrator.AutoAddMigrations();
+				bool schemasValid = await migrator.ValidateAllSchemasAsync(_hordeSettings.MongoMigrationsAutoUpgrade, CancellationToken.None);
+				if (!schemasValid)
+				{
+					return 1;
+				}
 			}
+			
+			await host.RunAsync();
+			return 0;
 		}
 
 		static IHostBuilder CreateHostBuilderWithCert(string[] args, IConfiguration config, ServerSettings serverSettings, X509Certificate2? sslCert)
@@ -231,5 +243,13 @@ namespace HordeServer.Commands
 			ServerSettings hordeSettings = new ServerSettings();
 			return CreateHostBuilderWithCert(args, new ConfigurationBuilder().Build(), hordeSettings, null);
 		}
+		
+		public static IHostBuilder CreateMinimalHostBuilder(IConfiguration config)
+		{
+			return Host.CreateDefaultBuilder()
+				.UseSerilog()
+				.ConfigureAppConfiguration(builder => builder.AddConfiguration(config))
+				.ConfigureWebHostDefaults(webBuilder => { webBuilder.UseStartup<Startup>(); });
+		}
 	}
 }
diff --git a/Engine/Source/Programs/Horde/HordeServer/Server/MongoMigration.cs b/Engine/Source/Programs/Horde/HordeServer/Server/MongoMigration.cs
new file mode 100644
--- /dev/null
+++ b/Engine/Source/Programs/Horde/HordeServer/Server/MongoMigration.cs
@@ -0,0 +1,427 @@
+﻿// Copyright Epic Games, Inc. All Rights Reserved.
+
+using System;
+using System.Collections.Generic;
+using System.Linq;
+using System.Reflection;
+using System.Threading;
+using System.Threading.Tasks;
+using HordeServer.Utilities;
+using Microsoft.Extensions.Logging;
+
+namespace HordeServer.Server;
+
+/// <summary>
+/// Indicates the version relationship between the application's defined schema and the database's applied schema.
+/// Used to determine if migrations or code updates are required.
+/// </summary>
+public enum SchemaVersionState
+{
+	/// <summary>
+	/// The application and database schemas are at the same version.
+	/// <remarks>
+	/// - No action required
+	/// - Application can operate normally
+	/// - Represents the ideal state
+	/// </remarks>
+	/// </summary>
+	VersionMatch = 0,
+	
+	/// <summary>
+	/// The application schema version is newer than the database version.
+	/// <remarks>
+	/// - Database requires migration
+	/// - Application should not run or process requests
+	/// - Requires manual (or automated) migration intervention
+	/// </remarks>
+	/// </summary>
+	DatabaseOutdated,
+	
+	/// <summary>
+	/// The database schema version is newer than the application version.
+	/// <remarks>
+	/// - Application requires deployment of a newer version of itself
+	/// - Can be a valid temporary state during rolling updates
+	/// - Safe to continue operation
+	/// </remarks>
+	/// </summary>
+	ApplicationOutdated,
+}
+
+/// <summary>
+/// Results of comparing application and database schema versions
+/// </summary>
+/// <param name="State">Current version state</param>
+/// <param name="AppVersion">Application schema version</param>
+/// <param name="DatabaseVersion">Database schema version</param>
+/// <param name="MissingMigrations">List of migrations needed to reconcile versions</param>
+public record SchemaVersionResult(SchemaVersionState State, int AppVersion, int DatabaseVersion, List<MigrationDefinition> MissingMigrations);
+
+/// <summary>
+/// A database migration with its metadata
+/// </summary>
+public class MigrationDefinition(IMongoMigration migration, string plugin, int version, string name)
+{
+	/// <inheritdoc cref="IMongoMigration"/>
+	public IMongoMigration Migration { get; } = migration;
+	
+	/// <inheritdoc cref="MongoMigrationAttribute.Plugin"/>
+	public string Plugin { get; } = plugin;
+	
+	/// <inheritdoc cref="MongoMigrationAttribute.Version"/>
+	public int Version { get; } = version;
+	
+	/// <inheritdoc cref="MongoMigrationAttribute.Name"/>
+	public string Name { get; } = name;
+	
+	/// <inheritdoc/>
+	public override string ToString()
+	{
+		return $"Plugin: {Plugin}, Version: {Version}, Name: {Name}";
+	}
+	
+	/// <inheritdoc/>
+	public override bool Equals(object? obj)
+	{
+		if (ReferenceEquals(null, obj))
+		{
+			return false;
+		}
+		
+		if (ReferenceEquals(this, obj))
+		{
+			return true;
+		}
+		
+		return obj.GetType() == GetType() && Equals((MigrationDefinition)obj);
+	}
+	
+	/// <inheritdoc/>
+	public override int GetHashCode()
+	{
+		return HashCode.Combine(Plugin, Version);
+	}
+	
+	private bool Equals(MigrationDefinition other)
+	{
+		return Plugin == other.Plugin && Version == other.Version;
+	}
+}
+
+/// <summary>
+/// Exception thrown during migration operations
+/// </summary>
+public class MigrationException : Exception
+{
+	/// <inheritdoc/>
+	public MigrationException(string message) : base(message) {}
+	
+	/// <inheritdoc/>
+	public MigrationException(string message, Exception innerException) : base(message, innerException) { }
+}
+
+/// <summary>
+/// Handles database schema migrations for MongoDB database
+/// This class is not thread-safe
+/// </summary>
+public class MongoMigrator
+{
+	private readonly Dictionary<string, HashSet<MigrationDefinition>> _migrations = [];
+	private readonly IMongoService _mongoService;
+	private readonly ILogger<MongoMigrator> _logger;
+	
+	[SingletonDocument("mongo-migrations")]
+	private class MigrationState : SingletonBase
+	{
+		public Dictionary<string, int> Versions { get; set; } = new ();
+	}
+	
+	/// <summary>
+	/// Constructor
+	/// </summary>
+	public MongoMigrator(IMongoService mongoService, ILogger<MongoMigrator> logger)
+	{
+		_mongoService = mongoService;
+		_logger = logger;
+	}
+	
+	/// <summary>
+	/// Load all migrations
+	/// Will scan all loaded assemblies via reflection and add any class decorated with MongoMigrationAttribute
+	/// </summary>
+	/// <exception cref="MigrationException">If migration is invalid or conflicts</exception>
+	public void AutoAddMigrations()
+	{
+		foreach (Assembly assembly in AppDomain.CurrentDomain.GetAssemblies())
+		{
+			// Skip most assemblies as iterating types is slow
+			bool isHordePlugin = assembly.GetName().Name?.Contains("Horde", StringComparison.InvariantCultureIgnoreCase) ?? false;
+			if (isHordePlugin)
+			{
+				foreach (Type type in assembly.GetTypes())
+				{
+					MongoMigrationAttribute? mma = type.GetCustomAttribute<MongoMigrationAttribute>();
+					if (mma is {AutoLoad: true} && Activator.CreateInstance(type) is IMongoMigration migration)
+					{
+						AddMigration(migration);
+					}
+				}
+			}
+		}
+	}
+	
+	/// <summary>
+	/// Adds a new migration to be managed by this migrator
+	/// </summary>
+	/// <param name="migration">Migration implementation to add</param>
+	/// <exception cref="MigrationException">If migration is invalid or conflicts</exception>
+	public void AddMigration(IMongoMigration migration)
+	{
+		MongoMigrationAttribute? attribute = migration.GetType().GetCustomAttribute<MongoMigrationAttribute>();
+		if (attribute == null)
+		{
+			throw new MigrationException($"Migration class {migration.GetType()} is missing [MongoMigration] attribute");
+		}
+
+		MigrationDefinition definition = new (migration, attribute.Plugin, attribute.Version, attribute.Name);
+		if (!_migrations.TryGetValue(definition.Plugin, out HashSet<MigrationDefinition>? pluginMigrations))
+		{
+			pluginMigrations = new HashSet<MigrationDefinition>();
+			_migrations[definition.Plugin] = pluginMigrations;
+		}
+
+		if (definition.Version is < 1 or > 50000)
+		{
+			throw new MigrationException($"Migration {definition.Plugin}/{definition.Name} has an invalid version number {definition.Version}");
+		}
+		
+		if (definition.Name.Length is <= 3 or > 200)
+		{
+			throw new MigrationException($"Migration {definition.Plugin}/{definition.Name} has a too short/long name");
+		}
+
+		if (!pluginMigrations.Add(definition))
+		{
+			throw new MigrationException($"Migrations for plugin {definition.Plugin} already contains a migration with version {definition.Version}");
+		}
+	}
+	
+	/// <summary>
+	/// Validates database schema versions for all plugins against application versions
+	/// Logs error for outdated schemas with migration instructions.
+	/// </summary>
+	/// <param name="autoUpgrade">Whether automatic upgrade of the schemas should be performed</param>
+	/// <param name="cancellationToken">Token to cancel the validation operation</param>
+	/// <returns>True if all schemas are valid or newer than app, false if any schema is outdated</returns>
+	public async Task<bool> ValidateAllSchemasAsync(bool autoUpgrade, CancellationToken cancellationToken)
+	{
+		Dictionary<string, SchemaVersionResult> schemaResults = await GetAllSchemaStateAsync(cancellationToken);
+		bool databaseOutdated = false;
+		foreach ((string plugin, SchemaVersionResult svr) in schemaResults)
+		{
+			string dbVersion = svr.DatabaseVersion > 0 ? Convert.ToString(svr.DatabaseVersion) : "(not applied)";
+			switch (svr.State)
+			{
+				case SchemaVersionState.VersionMatch: continue;
+				case SchemaVersionState.DatabaseOutdated:
+					_logger.Log(autoUpgrade ? LogLevel.Warning : LogLevel.Error, "Database schema for plugin {Plugin} is outdated. Schema versions: app = {AppVersion} database = {DatabaseVersion}", plugin, svr.AppVersion, dbVersion);
+					databaseOutdated = true;
+					break;
+				case SchemaVersionState.ApplicationOutdated:
+					_logger.LogWarning("Database schema for plugin {Plugin} newer than app. Schema versions: app = {AppVersion} database = {DatabaseVersion}", plugin, svr.AppVersion, dbVersion);
+					break;
+				default: throw new ArgumentOutOfRangeException($"Unknown schema state {svr}");
+			}
+		}
+		
+		if (autoUpgrade)
+		{
+			_logger.LogInformation("Automatically upgrading database schemas...");
+			await UpgradeAsync(cancellationToken);
+		}
+		else if (databaseOutdated)
+		{
+			_logger.LogInformation("To apply migration scripts, run: dotnet HordeServer.dll mongo upgrade (or equivalent)");
+			_logger.LogInformation("For dev or test deployments, auto-apply of migrations can be enabled in server settings. This setting is *not* recommended for production use.");
+		}
+		
+		return !databaseOutdated;
+	}
+	
+	/// <summary>
+	/// Retrieves the schema version states for all registered plugins
+	/// </summary>
+	/// <param name="cancellationToken">Cancellation token</param>
+	/// <returns>Schema version comparison results per plugin</returns>
+	public async Task<Dictionary<string, SchemaVersionResult>> GetAllSchemaStateAsync(CancellationToken cancellationToken)
+	{
+		Dictionary<string, SchemaVersionResult> result = new ();
+		foreach (string plugin in _migrations.Keys)
+		{
+			SchemaVersionResult svr = await GetSchemaStateAsync(plugin, cancellationToken);
+			result[plugin] = svr;
+		}
+		
+		return result;
+	}
+	
+	/// <summary>
+	/// Checks the current schema version state for a plugin
+	/// </summary>
+	/// <param name="plugin">Plugin identifier</param>
+	/// <param name="cancellationToken">Cancellation token</param>
+	/// <returns>Schema version comparison results</returns>
+	public async Task<SchemaVersionResult> GetSchemaStateAsync(string plugin, CancellationToken cancellationToken)
+	{
+		IReadOnlyList<MigrationDefinition> migrations = GetMigrations(plugin);
+		int appVersion = migrations.Count == 0 ? -1 : migrations.Max(x => x.Version);
+		int dbVersion = await GetVersionAsync(plugin, cancellationToken);
+		
+		if (appVersion > dbVersion)
+		{
+			List<MigrationDefinition> missingMigrations = migrations
+				.Where(x => x.Version > dbVersion && x.Version <= appVersion)
+				.OrderBy(x => x.Version)
+				.ToList();
+			
+			return new SchemaVersionResult(SchemaVersionState.DatabaseOutdated, appVersion, dbVersion, missingMigrations);
+		}
+		if (appVersion < dbVersion)
+		{
+			List<MigrationDefinition> missingMigrations = [];
+			for (int v = appVersion + 1; v <= dbVersion; v++)
+			{
+				missingMigrations.Add(new MigrationDefinition(null!, plugin, v, $"Unknown migration {v}"));
+			}
+			
+			return new SchemaVersionResult(SchemaVersionState.ApplicationOutdated, appVersion, dbVersion, missingMigrations);
+		}
+		
+		return new SchemaVersionResult(SchemaVersionState.VersionMatch, appVersion, dbVersion, []);
+	}
+	
+	/// <summary>
+	/// Gets all migrations registered for a plugin
+	/// </summary>
+	/// <param name="plugin">Plugin identifier</param>
+	/// <returns>List of migrations ordered by version</returns>
+	/// <exception cref="MigrationException">If plugin is unknown</exception>
+	public IReadOnlyList<MigrationDefinition> GetMigrations(string plugin)
+	{
+		if (_migrations.TryGetValue(plugin, out HashSet<MigrationDefinition>? pluginMigrations))
+		{
+			return new List<MigrationDefinition>(pluginMigrations).OrderBy(x => x.Version).ToList();
+		}
+		
+		throw new MigrationException($"Unknown plugin {plugin}");
+	}
+	
+	/// <summary>
+	/// Upgrades all database schemas to latest version
+	/// </summary>
+	/// <param name="cancellationToken">Cancellation token</param>
+	public async Task UpgradeAsync(CancellationToken cancellationToken)
+	{
+		foreach (string plugin in _migrations.Keys)
+		{
+			await UpgradeAsync(plugin, toVersion: null, cancellationToken);
+		}
+	}
+	
+	/// <summary>
+	/// Upgrades database schema for a plugin to specified version or latest if not specified
+	/// </summary>
+	/// <param name="plugin">Plugin identifier</param>
+	/// <param name="toVersion">Target version, or null for latest</param>
+	/// <param name="cancellationToken">Cancellation token</param>
+	public async Task UpgradeAsync(string plugin, int? toVersion, CancellationToken cancellationToken)
+	{
+		IReadOnlyList<MigrationDefinition> migrations = GetMigrations(plugin);
+		ValidateMigrations(plugin, migrations);
+		int appliedVersion = await GetVersionAsync(plugin, cancellationToken);
+		MigrationContext context = new (_mongoService, _logger);
+		
+		foreach (MigrationDefinition mi in migrations.Where(x => x.Version > appliedVersion && (!toVersion.HasValue || x.Version <= toVersion.Value)))
+		{
+			try
+			{
+				_logger.LogInformation("Upgrading {Plugin} to version {Version} using migration: {Name} ({ClassName})", mi.Plugin, mi.Version, mi.Name, mi.Migration.GetType().FullName);
+				await mi.Migration.UpAsync(context, cancellationToken);
+				await SetVersionAsync(plugin, mi.Version, cancellationToken);
+			}
+			catch (Exception e)
+			{
+				throw new MigrationException($"Failed to upgrade {mi.Plugin} to version {mi.Version} using migration: {mi.Name}", e);
+			}
+		}
+	}
+	
+	/// <summary>
+	/// Downgrades database schema for a plugin to specified version
+	/// </summary>
+	/// <param name="plugin">Plugin identifier</param>
+	/// <param name="toVersion">Target version</param>
+	/// <param name="cancellationToken">Cancellation token</param>
+	public async Task DowngradeAsync(string plugin, int toVersion, CancellationToken cancellationToken)
+	{
+		IReadOnlyList<MigrationDefinition> migrations = GetMigrations(plugin);
+		ValidateMigrations(plugin, migrations);
+		MigrationContext context = new (_mongoService, _logger);
+		
+		IEnumerable<MigrationDefinition> migrationsToApply = migrations
+			.OrderByDescending(x => x.Version).Where(x => x.Version > toVersion);
+		
+		foreach (MigrationDefinition mi in migrationsToApply)
+		{
+			try
+			{
+				_logger.LogInformation("Downgrading {Plugin} to version {Version} using migration: {Name}", mi.Plugin, mi.Version, mi.Name);
+				await mi.Migration.DownAsync(context, cancellationToken);
+				await SetVersionAsync(plugin, mi.Version, cancellationToken);
+			}
+			catch (Exception e)
+			{
+				throw new MigrationException($"Failed to downgrade {mi.Plugin} to version {mi.Version} using migration: {mi.Name}", e);
+			}
+		}
+	}
+	
+	private static void ValidateMigrations(string plugin, IReadOnlyList<MigrationDefinition> migrations)
+	{
+		if (migrations.Count == 0)
+		{
+			throw new MigrationException($"No migrations found for {plugin}");
+		}
+		
+		if (migrations[0].Version != 1)
+		{
+			throw new MigrationException($"Missing migration version 1 for {plugin}");
+		}
+		
+		// Check for gaps in version numbers
+		for (int i = 0; i < migrations.Count - 1; i++)
+		{
+			int currentVersion = migrations[i].Version;
+			int nextVersion = migrations[i + 1].Version;
+			
+			if (nextVersion != currentVersion + 1)
+			{
+				throw new MigrationException($"Gap detected in migration versions for {plugin}. Expected version {currentVersion + 1} but found version {nextVersion}");
+			}
+		}
+	}
+	
+	private async Task SetVersionAsync(string plugin, int version, CancellationToken cancellationToken)
+	{
+		await _mongoService.UpdateSingletonAsync<MigrationState>(state => state.Versions[plugin] = version, cancellationToken);
+	}
+	
+	/// <summary>
+	/// Get currently applied version of migrations for a plugin as stored in database
+	/// </summary>
+	private async Task<int> GetVersionAsync(string plugin, CancellationToken cancellationToken)
+	{
+		MigrationState state = await _mongoService.GetSingletonAsync<MigrationState>(cancellationToken);
+		return state.Versions.TryGetValue(plugin, out int version) ? version : -1;
+	}
+}
\ No newline at end of file
diff --git a/Engine/Source/Programs/Horde/HordeServer/ServerApp.cs b/Engine/Source/Programs/Horde/HordeServer/ServerApp.cs
--- a/Engine/Source/Programs/Horde/HordeServer/ServerApp.cs
+++ b/Engine/Source/Programs/Horde/HordeServer/ServerApp.cs
@@ -235,7 +235,19 @@ namespace HordeServer
 		{
 			s_pluginCollection ??= CreatePluginCollection(new ConfigurationBuilder().Build(), NullLogger.Instance);
 		}
-
+		
+		/// <summary>
+		/// Creates a plugin collection by loading and configuring plugins from assemblies based on configuration settings.
+		///
+		/// The method performs the following steps:
+		/// 1. Loads plugin configurations from the "Horde.Plugins" configuration section
+		/// 2. Scans the application directory for DLL files matching the pattern "HordeServer.*.dll"
+		/// 3. Loads each matching assembly and searches for types decorated with <see cref="PluginAttribute"/>
+		/// 4. Enables plugins based on configuration settings and default values
+		/// 5. Validates that all plugins enabled in configuration are found
+		/// </summary>
+		/// <returns>An initialized <see cref="IPluginCollection"/> containing all successfully loaded and enabled plugins.</returns>
+		/// <exception cref="InvalidOperationException">Thrown when one or more plugins enabled in the configuration file cannot be found.</exception>
 		static IPluginCollection CreatePluginCollection(IConfiguration configuration, ILogger logger)
 		{
 			Dictionary<string, PluginServerConfig> pluginConfigs = new Dictionary<string, PluginServerConfig>(StringComparer.OrdinalIgnoreCase);
diff --git a/Engine/Source/Programs/Horde/HordeServer/ServerSettings.cs b/Engine/Source/Programs/Horde/HordeServer/ServerSettings.cs
--- a/Engine/Source/Programs/Horde/HordeServer/ServerSettings.cs
+++ b/Engine/Source/Programs/Horde/HordeServer/ServerSettings.cs
@@ -120,6 +120,17 @@ namespace HordeServer
 			get => MongoReadOnlyMode;
 			set => MongoReadOnlyMode = value;
 		}
+		
+		/// <summary>
+		/// Whether database schema migrations are enabled
+		/// </summary>
+		public bool MongoMigrationsEnabled { get; set; } = false;
+		
+		/// <summary>
+		/// Whether database schema should automatically be applied
+		/// Only recommended for dev or test environments
+		/// </summary>
+		public bool MongoMigrationsAutoUpgrade { get; set; } = false;
 
 		/// <summary>
 		/// Shutdown the current server process if memory usage reaches this threshold (specified in MB)
diff --git a/Engine/Source/Programs/Horde/HordeServer/Startup.cs b/Engine/Source/Programs/Horde/HordeServer/Startup.cs
--- a/Engine/Source/Programs/Horde/HordeServer/Startup.cs
+++ b/Engine/Source/Programs/Horde/HordeServer/Startup.cs
@@ -365,6 +365,7 @@ namespace HordeServer
 			services.AddSingleton<MongoCommandTracer>();
 			services.AddSingleton<MongoService>();
 			services.AddSingleton<IMongoService>(sp => sp.GetRequiredService<MongoService>());
+			services.AddSingleton<MongoMigrator>();
 			services.AddSingleton<GlobalsService>();
 			services.AddSingleton<IClock, Clock>();
 			services.AddSingleton<IDowntimeService, DowntimeService>();
