Show partial elements selected on Templates context menu

diff --git a/Engine/Source/Programs/Horde/HordeDashboard/src/components/JobSearchSimple.tsx b/Engine/Source/Programs/Horde/HordeDashboard/src/components/JobSearchSimple.tsx
--- a/Engine/Source/Programs/Horde/HordeDashboard/src/components/JobSearchSimple.tsx
+++ b/Engine/Source/Programs/Horde/HordeDashboard/src/components/JobSearchSimple.tsx
@@ -137,6 +137,7 @@ export const JobSearchSimpleModal: React.FC<{ streamIdIn: string, changeListIn?:
 
       const isSelected = (ct: GetTemplateRefResponse) => searchState.templates.has(ct.id);
       let allCatTemplatesSelected: boolean = catTemplates.every(isSelected);
+      let someCatTemplatesSelected: boolean = catTemplates.some(isSelected) && !allCatTemplatesSelected;
 
       if (catTemplates?.length > 2) {
          subItems.unshift({
@@ -162,7 +163,7 @@ export const JobSearchSimpleModal: React.FC<{ streamIdIn: string, changeListIn?:
                   }
                   return false;
                }} >
-               <Checkbox checked={allCatTemplatesSelected} />
+               <Checkbox indeterminate={someCatTemplatesSelected} checked={allCatTemplatesSelected} />
                <Text style={{ color: dashboard.darktheme ? "#ABABAB" : "#575757" }}>Select All</Text></Stack>
             }
          });
@@ -173,6 +174,7 @@ export const JobSearchSimpleModal: React.FC<{ streamIdIn: string, changeListIn?:
    })
 
    let allTemplatesSelected: boolean =  searchState.templates.size === allTemplates.size;
+   let someTemplatesSelected: boolean =  searchState.templates.size > 0 && !allTemplatesSelected;
 
    templateOptions.unshift({
       key: `select_all_templates`,
@@ -193,7 +195,7 @@ export const JobSearchSimpleModal: React.FC<{ streamIdIn: string, changeListIn?:
             }
             return false;
          }} >
-         <Checkbox checked={allTemplatesSelected} />
+         <Checkbox indeterminate={someTemplatesSelected} checked={allTemplatesSelected} />
          <Text style={{ padding: "6px 0px"}}>Select All</Text></Stack>
       }
    });
