Add Select All to small Templates submenus

diff --git a/Engine/Source/Programs/Horde/HordeDashboard/src/components/JobSearchSimple.tsx b/Engine/Source/Programs/Horde/HordeDashboard/src/components/JobSearchSimple.tsx
--- a/Engine/Source/Programs/Horde/HordeDashboard/src/components/JobSearchSimple.tsx
+++ b/Engine/Source/Programs/Horde/HordeDashboard/src/components/JobSearchSimple.tsx
@@ -139,7 +139,7 @@ export const JobSearchSimpleModal: React.FC<{ streamIdIn: string, changeListIn?:
       let allCatTemplatesSelected: boolean = catTemplates.every(isSelected);
       let someCatTemplatesSelected: boolean = catTemplates.some(isSelected) && !allCatTemplatesSelected;
 
-      if (catTemplates?.length > 2) {
+      if (catTemplates?.length > 1) {
          subItems.unshift({
             key: `select_all_${cat}`,
             onRender: () => {
