Display N/A Start and Stop Times for User Notices

diff --git a/Engine/Source/Programs/Horde/HordeDashboard/src/components/NoticeView.tsx b/Engine/Source/Programs/Horde/HordeDashboard/src/components/NoticeView.tsx
--- a/Engine/Source/Programs/Horde/HordeDashboard/src/components/NoticeView.tsx
+++ b/Engine/Source/Programs/Horde/HordeDashboard/src/components/NoticeView.tsx
@@ -1,6 +1,6 @@
 // Copyright Epic Games, Inc. All Rights Reserved.
 
-import { DefaultButton, DetailsList, DetailsListLayoutMode, Dialog, DialogFooter, DialogType, IColumn, Icon, Modal, PrimaryButton, SelectionMode, Spinner, SpinnerSize, Stack, Text, TextField } from "@fluentui/react";
+import { DefaultButton, DetailsList, DetailsListLayoutMode, Dialog, DialogFooter, DialogType, DirectionalHint, IColumn, Icon, Modal, PrimaryButton, SelectionMode, Spinner, SpinnerSize, Stack, Text, TextField, TooltipHost } from "@fluentui/react";
 import { observer } from "mobx-react-lite";
 import moment from "moment-timezone";
 import React, { useState } from "react";
@@ -56,17 +56,32 @@ const NoticePanel: React.FC = observer(() => {
       }
 
       if (column?.name === "Start Time") {
-         if (!item.startTime) {
-            return null;
+         if (item.startTime) {
+            return <Text>{getDisplayTime(item.startTime as Date)}</Text>
          }
-         return <Text>{getDisplayTime(item.startTime as Date)}</Text>
+         
+         return <Stack style={{width: 25}}>
+               <TooltipHost
+                  content={'User notices displayed immediately'}
+                  directionalHint={DirectionalHint.bottomCenter}>
+                  <Text>{'N/A'}</Text>
+               </TooltipHost>
+         </Stack>
       }
 
       if (column?.name === "Finish Time") {
-         if (!item.finishTime) {
-            return null;
+         if (item.finishTime) {
+            return <Text>{getDisplayTime(item.finishTime as Date)}</Text>
          }
-         return <Text>{getDisplayTime(item.finishTime as Date)}</Text>
+
+         return <Stack style={{width: 25}}>
+            <TooltipHost
+               content={'User notices displayed until issue is resolved'}
+               directionalHint={DirectionalHint.bottomCenter}>
+               <Text>{'N/A'}</Text>
+            </TooltipHost>
+         </Stack>
+         
       }
 
       if (column?.name === "Created By") {
