Allow attaching a message to device audit update

#rnx

diff --git a/Engine/Source/Programs/Horde/Plugins/Build/HordeServer.Build.Tests/Devices/DeviceServiceTests.cs b/Engine/Source/Programs/Horde/Plugins/Build/HordeServer.Build.Tests/Devices/DeviceServiceTests.cs
--- a/Engine/Source/Programs/Horde/Plugins/Build/HordeServer.Build.Tests/Devices/DeviceServiceTests.cs
+++ b/Engine/Source/Programs/Horde/Plugins/Build/HordeServer.Build.Tests/Devices/DeviceServiceTests.cs
@@ -608,8 +608,14 @@ namespace HordeServer.Tests.Devices
 			Assert.AreEqual(telemetry.Count, 1);
 			// check that problem was recorded
 			Assert.IsNotNull(telemetry[0].Telemetry[0].ProblemTimeUtc);
-			// make sure the reswrvation time was finished for the problem device
+			// make sure the reservation time was finished for the problem device
 			Assert.IsNotNull(telemetry[0].Telemetry[0].ReservationFinishUtc);
+
+			// report a problem with message
+			DeviceErrorRequest requestWithMessage = new();
+			requestWithMessage.Message = "Unit Test message";
+			problem = (await DeviceController!.PutDeviceErrorAsync(reservation.DeviceNames[0], requestWithMessage)) as OkResult;
+			Assert.IsNotNull(problem);
 		}
 
 		[TestMethod]
diff --git a/Engine/Source/Programs/Horde/Plugins/Build/HordeServer.Build/Devices/DeviceCollection.cs b/Engine/Source/Programs/Horde/Plugins/Build/HordeServer.Build/Devices/DeviceCollection.cs
--- a/Engine/Source/Programs/Horde/Plugins/Build/HordeServer.Build/Devices/DeviceCollection.cs
+++ b/Engine/Source/Programs/Horde/Plugins/Build/HordeServer.Build/Devices/DeviceCollection.cs
@@ -529,7 +529,7 @@ namespace HordeServer.Devices
 		}
 
 		/// <inheritdoc/>
-		public async Task UpdateDeviceAsync(DeviceId deviceId, DevicePoolId? newPoolId, string? newName, string? newAddress, string? newModelId, string? newNotes, bool? newEnabled, bool? newProblem, bool? newMaintenance, UserId? modifiedByUserId = null)
+		public async Task UpdateDeviceAsync(DeviceId deviceId, DevicePoolId? newPoolId, string? newName, string? newAddress, string? newModelId, string? newNotes, bool? newEnabled, bool? newProblem, bool? newMaintenance, UserId? modifiedByUserId = null, string? message = null)
 		{
 			UpdateDefinitionBuilder<DeviceDocument> updateBuilder = Builders<DeviceDocument>.Update;
 			List<UpdateDefinition<DeviceDocument>> updates = new List<UpdateDefinition<DeviceDocument>>();
@@ -745,6 +745,14 @@ namespace HordeServer.Devices
 					}
 				}
 			}
+
+			if (message != null)
+			{
+				if (log)
+				{
+					logger.LogInformation("{Message}", message);
+				}
+			}
 		}
 
 		/// <inheritdoc/>
diff --git a/Engine/Source/Programs/Horde/Plugins/Build/HordeServer.Build/Devices/DeviceService.cs b/Engine/Source/Programs/Horde/Plugins/Build/HordeServer.Build/Devices/DeviceService.cs
--- a/Engine/Source/Programs/Horde/Plugins/Build/HordeServer.Build/Devices/DeviceService.cs
+++ b/Engine/Source/Programs/Horde/Plugins/Build/HordeServer.Build/Devices/DeviceService.cs
@@ -405,7 +405,7 @@ namespace HordeServer.Devices
 		/// <summary>
 		/// Update a device
 		/// </summary>
-		public async Task UpdateDeviceAsync(DeviceId deviceId, DevicePoolId? newPoolId = null, string? newName = null, string? newAddress = null, string? newModelId = null, string? newNotes = null, bool? newEnabled = null, bool? newProblem = null, bool? newMaintenance = null, UserId? modifiedByUserId = null)
+		public async Task UpdateDeviceAsync(DeviceId deviceId, DevicePoolId? newPoolId = null, string? newName = null, string? newAddress = null, string? newModelId = null, string? newNotes = null, bool? newEnabled = null, bool? newProblem = null, bool? newMaintenance = null, UserId? modifiedByUserId = null, string? message = null)
 		{
 			if (newProblem == true)
 			{
@@ -415,7 +415,7 @@ namespace HordeServer.Devices
 					await TryUpdateReservationAsync(reservation.Id, problemDevice: deviceId);
 				}
 			}
-			await _devices.UpdateDeviceAsync(deviceId, newPoolId, newName, newAddress, newModelId, newNotes, newEnabled, newProblem, newMaintenance, modifiedByUserId);
+			await _devices.UpdateDeviceAsync(deviceId, newPoolId, newName, newAddress, newModelId, newNotes, newEnabled, newProblem, newMaintenance, modifiedByUserId, message);
 		}
 
 		/// <summary>
diff --git a/Engine/Source/Programs/Horde/Plugins/Build/HordeServer.Build/Devices/DevicesController.cs b/Engine/Source/Programs/Horde/Plugins/Build/HordeServer.Build/Devices/DevicesController.cs
--- a/Engine/Source/Programs/Horde/Plugins/Build/HordeServer.Build/Devices/DevicesController.cs
+++ b/Engine/Source/Programs/Horde/Plugins/Build/HordeServer.Build/Devices/DevicesController.cs
@@ -249,7 +249,7 @@ namespace HordeServer.Devices
 			string? name = update.Name?.Trim();
 			string? address = update.Address?.Trim();
 
-			await _deviceService.UpdateDeviceAsync(deviceIdValue, poolIdValue, name, address, modelIdValue, update.Notes, update.Enabled, update.Problem, update.Maintenance, internalUser.Id);
+			await _deviceService.UpdateDeviceAsync(deviceIdValue, poolIdValue, name, address, modelIdValue, update.Notes, update.Enabled, update.Problem, update.Maintenance, internalUser.Id, update.Message);
 
 			return Ok();
 		}
@@ -1047,7 +1047,7 @@ namespace HordeServer.Devices
 		/// </summary>
 		[HttpPut]
 		[Route("/api/v1/deviceerror/{deviceName}")]
-		public async Task<ActionResult> PutDeviceErrorAsync(string deviceName)
+		public async Task<ActionResult> PutDeviceErrorAsync(string deviceName, [FromBody] DeviceErrorRequest? request = null)
 		{
 			IDevice? device = await _deviceService.GetDeviceByNameAsync(deviceName);
 
@@ -1056,7 +1056,7 @@ namespace HordeServer.Devices
 				return BadRequest($"Unknown device {deviceName}");
 			}
 
-			await _deviceService.UpdateDeviceAsync(device.Id, newProblem: true);
+			await _deviceService.UpdateDeviceAsync(device.Id, newProblem: true, message: request?.Message);
 
 			return Ok();
 		}
diff --git a/Engine/Source/Programs/Horde/Plugins/Build/HordeServer.Build/Devices/IDeviceCollection.cs b/Engine/Source/Programs/Horde/Plugins/Build/HordeServer.Build/Devices/IDeviceCollection.cs
--- a/Engine/Source/Programs/Horde/Plugins/Build/HordeServer.Build/Devices/IDeviceCollection.cs
+++ b/Engine/Source/Programs/Horde/Plugins/Build/HordeServer.Build/Devices/IDeviceCollection.cs
@@ -101,7 +101,8 @@ namespace HordeServer.Devices
 		/// <param name="newProblem">Whether to set or clear problem state</param>
 		/// <param name="newMaintenance">Whether to set or clear maintenance state</param>
 		/// <param name="modifiedByUserId">The user who is updating the device</param>
-		Task UpdateDeviceAsync(DeviceId deviceId, DevicePoolId? newPoolId, string? newName, string? newAddress, string? newModelId, string? newNotes, bool? newEnabled, bool? newProblem, bool? newMaintenance, UserId? modifiedByUserId = null);
+		/// <param name="message">Attached a message to the log of the device</param>
+		Task UpdateDeviceAsync(DeviceId deviceId, DevicePoolId? newPoolId, string? newName, string? newAddress, string? newModelId, string? newNotes, bool? newEnabled, bool? newProblem, bool? newMaintenance, UserId? modifiedByUserId = null, string? message = null);
 
 		/// <summary>
 		/// Delete a device from the collection
diff --git a/Engine/Source/Programs/Shared/EpicGames.Horde/Devices/DeviceMessages.cs b/Engine/Source/Programs/Shared/EpicGames.Horde/Devices/DeviceMessages.cs
--- a/Engine/Source/Programs/Shared/EpicGames.Horde/Devices/DeviceMessages.cs
+++ b/Engine/Source/Programs/Shared/EpicGames.Horde/Devices/DeviceMessages.cs
@@ -429,6 +429,22 @@ namespace EpicGames.Horde.Devices
 		/// Whether to set or clear any device problem state
 		/// </summary>
 		public bool? Problem { get; set; }
+
+		/// <summary>
+		/// Additional message to attach to this request
+		/// </summary>
+		public string? Message { get; set; }
+	}
+
+	/// <summary>
+	/// Device error report request object
+	/// </summary>
+	public class DeviceErrorRequest
+	{
+		/// <summary>
+		/// Additional message to attach to this request
+		/// </summary>
+		public string? Message { get; set; }
 	}
 
 	/// <summary>
