Horde: Run job executor cleanup when the job batch process fails or hangs

diff --git a/Engine/Source/Programs/Horde/Drivers/JobDriver/Execution/JobExecutor.cs b/Engine/Source/Programs/Horde/Drivers/JobDriver/Execution/JobExecutor.cs
--- a/Engine/Source/Programs/Horde/Drivers/JobDriver/Execution/JobExecutor.cs
+++ b/Engine/Source/Programs/Horde/Drivers/JobDriver/Execution/JobExecutor.cs
@@ -285,7 +285,15 @@ namespace JobDriver.Execution
 			span.SetAttribute("horde.job.executor", JobOptions.Executor);
 			span.SetAttribute("horde.lease.id", LeaseId.ToString());
 			JobOptions.DecorateSpan(span);
-			
+
+			string? cleanupVar = Environment.GetEnvironmentVariable("UE_HORDE_JOB_BATCH_CLEANUP");
+			if (cleanupVar != null && Int32.TryParse(cleanupVar, out int v) && v != 0)
+			{
+				logger.LogInformation("Cleaning up failed job batch.");
+				await FinalizeAsync(logger, CancellationToken.None);
+				return;
+			}
+
 			JobRpc.JobRpcClient jobRpc = await HordeClient.CreateGrpcClientAsync<JobRpc.JobRpcClient>(cancellationToken);
 
 			// Create a storage client for this session
diff --git a/Engine/Source/Programs/Horde/HordeAgent/Leases/Handlers/JobHandler.cs b/Engine/Source/Programs/Horde/HordeAgent/Leases/Handlers/JobHandler.cs
--- a/Engine/Source/Programs/Horde/HordeAgent/Leases/Handlers/JobHandler.cs
+++ b/Engine/Source/Programs/Horde/HordeAgent/Leases/Handlers/JobHandler.cs
@@ -60,6 +60,25 @@ namespace HordeAgent.Leases.Handlers
 				int exitCode = await RunDotNetProcessAsync(driverAssembly, arguments, environment, AgentApp.IsSelfContained, logger, cancellationToken);
 				logger.LogInformation("Driver finished with exit code {ExitCode}", exitCode);
 
+				if (exitCode != 0)
+				{
+					try
+					{
+						logger.LogInformation("Driver initiating job batch cleanup");
+						// Run job batch cleanup
+						environment["UE_HORDE_JOB_BATCH_CLEANUP"] = "1";
+						int cleanupCode = await RunDotNetProcessAsync(driverAssembly, arguments, environment, AgentApp.IsSelfContained, logger, cancellationToken);
+						if (cleanupCode != 0)
+						{
+							logger.LogInformation("Driver job batch cleanup finished with exit code {ExitCode}", cleanupCode);
+						}
+					}
+					catch (Exception ex)
+					{
+						logger.LogError(ex, "Driver job batch cleanup error ({Message})", ex.Message);
+					}
+				}
+
 				return (exitCode == 0) ? LeaseResult.Success : LeaseResult.Failed;
 			}
 			catch (OperationCanceledException ex)
