Horde: Add small HTTP server in agent for administration and health checks

In this initial version, the health check endpoint returns a basic status whether it is connected to the Horde server or not.

diff --git a/Engine/Source/Programs/Horde/HordeAgent.Tests/Services/AwsInstanceLifecycleServiceTests.cs b/Engine/Source/Programs/Horde/HordeAgent.Tests/Services/AwsInstanceLifecycleServiceTests.cs
--- a/Engine/Source/Programs/Horde/HordeAgent.Tests/Services/AwsInstanceLifecycleServiceTests.cs
+++ b/Engine/Source/Programs/Horde/HordeAgent.Tests/Services/AwsInstanceLifecycleServiceTests.cs
@@ -33,6 +33,8 @@ public class WorkerServiceStub : IWorkerService
 	{
 		TerminateSession = true;
 	}
+	
+	public bool IsConnected { get; set; }
 }
 
 /// <summary>
diff --git a/Engine/Source/Programs/Horde/HordeAgent.Tests/Services/ManagementServiceTests.cs b/Engine/Source/Programs/Horde/HordeAgent.Tests/Services/ManagementServiceTests.cs
new file mode 100644
--- /dev/null
+++ b/Engine/Source/Programs/Horde/HordeAgent.Tests/Services/ManagementServiceTests.cs
@@ -0,0 +1,86 @@
+// Copyright Epic Games, Inc. All Rights Reserved.
+
+using System;
+using System.IO;
+using System.Text.Json;
+using System.Threading;
+using System.Threading.Tasks;
+using HordeAgent.Services;
+using Microsoft.Extensions.Logging.Abstractions;
+using Microsoft.VisualStudio.TestTools.UnitTesting;
+
+namespace HordeAgent.Tests.Services;
+
+internal class HttpListenerRequestStub(Uri? url) : IHttpListenerRequest
+{
+	public Uri? Url { get; } = url;
+}
+
+internal class HttpListenerResponseStub : IHttpListenerResponse
+{
+	public string? ContentType { get; set; }
+	public Stream OutputStream { get; } = new MemoryStream();
+	public void Close()
+	{
+	}
+}
+
+[TestClass]
+public sealed class ManagementServiceTests
+{
+	private readonly WorkerServiceStub _workerServiceStub = new();
+	
+	[TestMethod]
+	public async Task HttpRequest_NotFound_Async()
+	{
+		(ErrorResponse? res, IHttpListenerResponse? httpRes) = await RequestAdminAsync<ErrorResponse>("http://localhost/not-found");
+		Assert.IsFalse(res!.Success);
+		Assert.AreEqual("application/json", httpRes.ContentType);
+		Assert.AreEqual("Not Found", res.Message);
+	}
+	
+	[TestMethod]
+	public async Task HttpRequest_HealthCheck_Async()
+	{
+		{
+			_workerServiceStub.IsConnected = true;
+			(HealthCheckResponse? res, _) = await RequestHealthAsync<HealthCheckResponse>("http://localhost/");
+			Assert.IsTrue(res!.IsHealthy);
+			Assert.IsTrue(res.IsConnected);
+		}
+		{
+			_workerServiceStub.IsConnected = false;
+			(HealthCheckResponse? res, _) = await RequestHealthAsync<HealthCheckResponse>("http://localhost/");
+			Assert.IsFalse(res!.IsHealthy);
+			Assert.IsFalse(res.IsConnected);
+		}
+	}
+	
+	private Task<(T?, IHttpListenerResponse)> RequestAdminAsync<T>(string path)
+	{
+		return RequestAsync<T>(path, (ms, req, res) => ms.TestAdminRouteAsync(req, res));
+	}
+	
+	private Task<(T?, IHttpListenerResponse)> RequestHealthAsync<T>(string path)
+	{
+		return RequestAsync<T>(path, (ms, req, res) => ms.TestHealthRouteAsync(req, res));
+	}
+	
+	private async Task<(T?, IHttpListenerResponse)> RequestAsync<T>(string path, Func<ManagementService, IHttpListenerRequest, IHttpListenerResponse, Task> cb)
+	{
+		await using ManagementService ms = GetMs();
+		HttpListenerRequestStub req = new(new Uri(path));
+		HttpListenerResponseStub res = new();
+		await cb(ms, req, res);
+		res.OutputStream.Position = 0;
+		T? deserialized = await JsonSerializer.DeserializeAsync<T>(res.OutputStream, JsonSerializerOptions.Default, CancellationToken.None);
+		return (deserialized, res);
+	}
+	
+	private ManagementService GetMs()
+	{
+		AgentSettings settings = new () { AdminEndpoints = ["http://localhost:3232/"], HealthCheckEndpoints = ["http://localhost:3233/"] };
+		TestOptionsMonitor<AgentSettings> settingsOpt = new(settings);
+		return new ManagementService(settingsOpt, _workerServiceStub, NullLogger<ManagementService>.Instance);
+	}
+}
diff --git a/Engine/Source/Programs/Horde/HordeAgent/AgentApp.cs b/Engine/Source/Programs/Horde/HordeAgent/AgentApp.cs
--- a/Engine/Source/Programs/Horde/HordeAgent/AgentApp.cs
+++ b/Engine/Source/Programs/Horde/HordeAgent/AgentApp.cs
@@ -205,6 +205,8 @@ namespace HordeAgent
 			services.AddSingleton<JobHandlerFactory>();
 			services.AddSingleton<StatusService>();
 			services.AddHostedService<StatusService>(sp => sp.GetRequiredService<StatusService>());
+			services.AddSingleton<ManagementService>();
+			services.AddHostedService<ManagementService>(sp => sp.GetRequiredService<ManagementService>());
 
 			services.AddSingleton<LeaseHandlerFactory, ComputeHandlerFactory>();
 			services.AddSingleton<LeaseHandlerFactory, ConformHandlerFactory>();
diff --git a/Engine/Source/Programs/Horde/HordeAgent/AgentSettings.cs b/Engine/Source/Programs/Horde/HordeAgent/AgentSettings.cs
--- a/Engine/Source/Programs/Horde/HordeAgent/AgentSettings.cs
+++ b/Engine/Source/Programs/Horde/HordeAgent/AgentSettings.cs
@@ -216,6 +216,21 @@ namespace HordeAgent
 		/// Key/value properties in addition to those set internally by the agent
 		/// </summary>
 		public Dictionary<string, string> Properties { get; } = new();
+		
+		/// <summary>
+		/// Listen addresses for the built-in HTTP admin/management server. Disabled when empty.
+		/// If activated, it's recommended to bind only to localhost for security reasons.
+		/// Example: localhost:7008 to listen on localhost, port 7008
+		/// </summary>
+		public string[] AdminEndpoints { get; set; } = [];
+		
+		/// <summary>
+		/// Listen addresses for the built-in HTTP health check server. Disabled when empty.
+		/// If activated, it's recommended to bind only to localhost for security reasons.
+		/// Example: *:7009 to listen on all interfaces/IPs, port 7009
+		/// If all interfaces are bound with *, make sure to run process as administrator.
+		/// </summary>
+		public string[] HealthCheckEndpoints { get; set; } = [];
 
 		/// <summary>
 		/// Gets the current server settings
diff --git a/Engine/Source/Programs/Horde/HordeAgent/Services/ManagementService.cs b/Engine/Source/Programs/Horde/HordeAgent/Services/ManagementService.cs
new file mode 100644
--- /dev/null
+++ b/Engine/Source/Programs/Horde/HordeAgent/Services/ManagementService.cs
@@ -0,0 +1,221 @@
+// Copyright Epic Games, Inc. All Rights Reserved.
+
+using System.Net;
+using System.Text.Json;
+using EpicGames.Core;
+using Microsoft.Extensions.Hosting;
+using Microsoft.Extensions.Logging;
+using Microsoft.Extensions.Options;
+
+using RoutesDict = System.Collections.Generic.Dictionary<string, System.Action<HordeAgent.Services.IHttpListenerRequest, HordeAgent.Services.IHttpListenerResponse>>;
+
+namespace HordeAgent.Services;
+
+#region Wrapper classes for HttpListener to make it testable
+internal interface IHttpListenerRequest
+{
+	Uri? Url { get; }
+}
+
+internal interface IHttpListenerResponse
+{
+	string? ContentType { get; set;  }
+	Stream OutputStream { get; }
+	void Close();
+}
+
+internal class WrappedHttpListenerRequest(HttpListenerRequest impl) : IHttpListenerRequest
+{
+	public Uri? Url => impl.Url;
+}
+
+internal class WrappedHttpListenerResponse(HttpListenerResponse impl) : IHttpListenerResponse
+{
+	public string? ContentType
+	{
+		get => impl.ContentType;
+		set => impl.ContentType = value;
+	}
+	
+	public Stream OutputStream => impl.OutputStream;
+	public void Close() => impl.Close();
+}
+#endregion
+
+internal class ErrorResponse(bool success, string message)
+{
+	public bool Success { get; set; } = success;
+	public string Message { get; set; } = message;
+}
+
+internal class HealthCheckResponse(bool isHealthy, bool isConnected)
+{
+	public bool IsHealthy { get; set; } = isHealthy;
+	public bool IsConnected { get; set; } = isConnected;
+}
+
+/// <summary>
+/// Handles agent's built-in HTTP server used for admin tasks and health check
+/// </summary>
+internal class ManagementService : IHostedService, IAsyncDisposable
+{
+	private readonly HttpListener? _adminListener;
+	private readonly HttpListener? _healthListener;
+	private readonly BackgroundTask? _adminTask;
+	private readonly BackgroundTask? _healthTask;
+	private readonly ILogger<ManagementService> _logger;
+	private static readonly JsonSerializerOptions s_jsonSerializerOptions = new () { WriteIndented = true, PropertyNameCaseInsensitive = true, PropertyNamingPolicy = JsonNamingPolicy.CamelCase};
+	
+	private readonly RoutesDict _adminRoutes = new()
+	{
+		["/"] = HandleJson((_, _) => "Admin server"),
+	};
+	
+	private readonly RoutesDict _healthRoutes;
+	
+	private static Action<IHttpListenerRequest, IHttpListenerResponse> HandleJson<T>(Func<IHttpListenerRequest, IHttpListenerResponse, T> callback)
+	{
+		return (req, res) =>
+		{
+			res.ContentType = "application/json";
+			T result = callback(req, res);
+			JsonSerializer.Serialize(res.OutputStream, result, s_jsonSerializerOptions);
+		};
+	}
+	
+	/// <summary>
+	/// Constructor
+	/// </summary>
+	public ManagementService(IOptionsMonitor<AgentSettings> settings, IWorkerService workerService, ILogger<ManagementService> logger)
+	{
+		_logger = logger;
+		_healthRoutes = new RoutesDict
+		{
+			["/"] = HandleJson<HealthCheckResponse>((_, _) => new HealthCheckResponse(workerService.IsConnected, workerService.IsConnected)),
+		};
+		
+		string SettingsEndpointToPrefix(string endpoint) => $"http://{endpoint}/";
+		
+		if (settings.CurrentValue.AdminEndpoints.Length > 0)
+		{
+			_adminListener = new HttpListener();
+			foreach (string endpoint in settings.CurrentValue.AdminEndpoints)
+			{
+				_adminListener.Prefixes.Add(SettingsEndpointToPrefix(endpoint));
+			}
+			_logger.LogInformation("Admin HTTP server listening on {Endpoints}", String.Join(",", settings.CurrentValue.AdminEndpoints));
+			_adminTask = new BackgroundTask(ct => RunServerAsync(_adminListener, _adminRoutes, ct));
+		}
+		
+		if (settings.CurrentValue.HealthCheckEndpoints.Length > 0)
+		{
+			_healthListener = new HttpListener();
+			foreach (string endpoint in settings.CurrentValue.HealthCheckEndpoints)
+			{
+				_healthListener.Prefixes.Add(SettingsEndpointToPrefix(endpoint));
+			}
+			_logger.LogInformation("Health check HTTP server listening on {Endpoints}", String.Join(",", settings.CurrentValue.HealthCheckEndpoints));
+			_healthTask = new BackgroundTask(ct => RunServerAsync(_healthListener, _healthRoutes, ct));
+		}
+	}
+
+	/// <inheritdoc/>
+	public async ValueTask DisposeAsync()
+	{
+		_adminListener?.Close();
+		_healthListener?.Close();
+		
+		if (_adminTask != null)
+		{
+			await _adminTask.DisposeAsync();
+		}
+		
+		if (_healthTask != null)
+		{
+			await _healthTask.DisposeAsync();
+		}
+	}
+
+	/// <inheritdoc/>
+	public Task StartAsync(CancellationToken cancellationToken)
+	{
+		try
+		{
+			_adminListener?.Start();
+			_healthListener?.Start();
+			_adminTask?.Start();
+			_healthTask?.Start();
+		}
+		catch (Exception e)
+		{
+			_logger.LogError(e, "Failed to start built-in HTTP server(s) for administration and health checks");
+		}
+		
+		return Task.CompletedTask;
+	}
+
+	/// <inheritdoc/>
+	public async Task StopAsync(CancellationToken cancellationToken)
+	{
+		if (_adminTask != null)
+		{
+			await _adminTask.StopAsync(cancellationToken);
+		}
+		
+		if (_healthTask != null)
+		{
+			await _healthTask.StopAsync(cancellationToken);
+		}
+
+		_adminListener?.Stop();
+		_healthListener?.Stop();
+	}
+	
+	internal Task TestAdminRouteAsync(IHttpListenerRequest req, IHttpListenerResponse res)
+	{
+		return HandleRequestAsync(req, res, _adminRoutes);
+	}
+	
+	internal Task TestHealthRouteAsync(IHttpListenerRequest req, IHttpListenerResponse res)
+	{
+		return HandleRequestAsync(req, res, _healthRoutes);
+	}
+	
+	private Task HandleRequestAsync(IHttpListenerRequest req, IHttpListenerResponse res, RoutesDict routes)
+	{
+		try
+		{
+			string path = req.Url?.AbsolutePath ?? "/";
+			if (routes.TryGetValue(path, out Action<IHttpListenerRequest, IHttpListenerResponse>? handler))
+			{
+				handler(req, res);
+			}
+			else
+			{
+				HandleJson((_, _) => new ErrorResponse(false, "Not Found"))(req, res);
+			}
+		}
+		catch (Exception e)
+		{
+			HandleJson((_, _) => new ErrorResponse(false, "Internal Server Error"))(req, res);
+			_logger.LogError(e, "Error while serving HTTP request for management service: {Exception}", e);
+		}
+		finally
+		{
+			res.Close();
+		}
+		
+		return Task.CompletedTask;
+	}
+	
+	private async Task RunServerAsync(HttpListener listener, RoutesDict routes, CancellationToken cancellationToken)
+	{
+		while (!cancellationToken.IsCancellationRequested)
+		{
+			HttpListenerContext context = await listener.GetContextAsync();
+			WrappedHttpListenerRequest wrappedReq = new (context.Request);
+			WrappedHttpListenerResponse wrappedRes = new (context.Response);
+			await HandleRequestAsync(wrappedReq, wrappedRes, routes);
+		}
+	}
+}
diff --git a/Engine/Source/Programs/Horde/HordeAgent/Services/WorkerService.cs b/Engine/Source/Programs/Horde/HordeAgent/Services/WorkerService.cs
--- a/Engine/Source/Programs/Horde/HordeAgent/Services/WorkerService.cs
+++ b/Engine/Source/Programs/Horde/HordeAgent/Services/WorkerService.cs
@@ -24,6 +24,11 @@ namespace HordeAgent.Services
 		/// Stop the session loop after at least one lease has finished executing
 		/// </summary>
 		public void TerminateSessionAfterLease();
+		
+		/// <summary>
+		/// Whether agent is connected to Horde server
+		/// </summary>
+		public bool IsConnected { get; }
 	}
 	
 	/// <summary>
@@ -72,7 +77,10 @@ namespace HordeAgent.Services
 				_currentLeaseManager.TerminateSessionAfterLease = true;
 			}
 		}
-
+		
+		/// <inheritdoc/>
+		public bool IsConnected { get; private set; } = false;
+		
 		/// <summary>
 		/// Executes the ServerTaskAsync method and swallows the exception for the task being cancelled. This allows waiting for it to terminate.
 		/// </summary>
@@ -129,8 +137,10 @@ namespace HordeAgent.Services
 
 						await using (ISession session = await _sessionFactory.CreateAsync(stoppingToken))
 						{
+							IsConnected = true;
 							_currentLeaseManager = new LeaseManager(session, _serviceProvider);
 							result = await _currentLeaseManager.RunAsync(stoppingToken);
+							IsConnected = false;
 						}
 
 						failureCount = 0;
