Horde: Dashboard side of disabling step run again when temp storage has expired
#rnx

diff --git a/Engine/Source/Programs/Horde/HordeDashboard/src/backend/Api.ts b/Engine/Source/Programs/Horde/HordeDashboard/src/backend/Api.ts
--- a/Engine/Source/Programs/Horde/HordeDashboard/src/backend/Api.ts
+++ b/Engine/Source/Programs/Horde/HordeDashboard/src/backend/Api.ts
@@ -425,6 +425,12 @@ export enum AuthMethod {
 	Horde = "Horde"
 }
 
+/// Artifact type information
+export type GetDashboardArtifactTypeResponse = {
+	type: string;
+	keepDays?: number;
+}
+
 /** Setting information required by dashboard */
 export type GetDashboardConfigResponse = {
 
@@ -458,6 +464,9 @@ export type GetDashboardConfigResponse = {
 	/** Categories to display on the pools page */
 	artifactTypes: string[];
 
+	/** Artifact type information */
+	artifactTypeInfo: GetDashboardArtifactTypeResponse[];
+
 	/** Telemetry views */
 	telemetryViews: any[];
 }
diff --git a/Engine/Source/Programs/Horde/HordeDashboard/src/backend/Dashboard.ts b/Engine/Source/Programs/Horde/HordeDashboard/src/backend/Dashboard.ts
--- a/Engine/Source/Programs/Horde/HordeDashboard/src/backend/Dashboard.ts
+++ b/Engine/Source/Programs/Horde/HordeDashboard/src/backend/Dashboard.ts
@@ -2,7 +2,7 @@
 
 import { action, makeObservable, observable } from 'mobx';
 import backend from '.';
-import { AuthMethod, DashboardPreference, DownloadFormat, GetDashboardAgentCategoryResponse, GetDashboardConfigResponse, GetDashboardPoolCategoryResponse, GetJobTemplateSettingsResponse, GetUserResponse, UserClaim } from './Api';
+import { AuthMethod, DashboardPreference, DownloadFormat, GetDashboardAgentCategoryResponse, GetDashboardArtifactTypeResponse, GetDashboardConfigResponse, GetDashboardPoolCategoryResponse, GetJobTemplateSettingsResponse, GetUserResponse, UserClaim } from './Api';
 import { getSiteConfig } from './Config';
 
 export enum StatusColor {
@@ -206,6 +206,10 @@ export class Dashboard {
         return this.config?.artifactTypes ?? [];
     }
 
+    getArtifactTypeInfo(type:string): GetDashboardArtifactTypeResponse | undefined {
+        return this.config?.artifactTypeInfo?.find(a => a.type === type);
+    }
+
     get deviceProblemCooldownMinutes(): number {
         return this.config?.deviceProblemCooldownMinutes ?? 30;
     }
@@ -558,7 +562,7 @@ export class Dashboard {
                     this.config = await backend.getDashboardConfig();
                 } catch (reason) {
                     console.error("Error getting dashboard config, defaults used: " + reason);
-                    this.config = { agentCategories: [], poolCategories: [], telemetryViews: [], artifactTypes: [] };
+                    this.config = { agentCategories: [], poolCategories: [], telemetryViews: [], artifactTypes: [], artifactTypeInfo: [] };
                 }
             }
 
diff --git a/Engine/Source/Programs/Horde/HordeDashboard/src/components/jobDetailsV2/JobOperationsBar.tsx b/Engine/Source/Programs/Horde/HordeDashboard/src/components/jobDetailsV2/JobOperationsBar.tsx
--- a/Engine/Source/Programs/Horde/HordeDashboard/src/components/jobDetailsV2/JobOperationsBar.tsx
+++ b/Engine/Source/Programs/Horde/HordeDashboard/src/components/jobDetailsV2/JobOperationsBar.tsx
@@ -401,8 +401,19 @@ const StepOperations: React.FC<{ jobDetails: JobDetailsV2, stepId: string }> = o
       return null;
    }
 
+   const stepOutputInfo = dashboard.getArtifactTypeInfo("step-output")
 
-   const canRunDisabled = !step?.allowRetry || !!step.retriedByUserInfo || !!jobData.abortedByUserInfo || !jobDetails.template?.canRun;
+   let tempStorageExpired = false;
+   if (stepOutputInfo && stepOutputInfo.keepDays && step.startTime) {
+
+      const expireMs = stepOutputInfo.keepDays * 86400000;
+      if ((Date.now() - new Date(step.startTime).getTime()) > expireMs)
+      {         
+         tempStorageExpired = true;
+      }      
+   }
+
+   const canRunDisabled = !step?.allowRetry || !!step.retriedByUserInfo || !!jobData.abortedByUserInfo || !jobDetails.template?.canRun || tempStorageExpired;
    const canTryFix = jobDetails.template?.allowPreflights && !!jobDetails.template?.canRun;
    let canBisect = (step.outcome === JobStepOutcome.Failure || step.outcome === JobStepOutcome.Warnings) && !!jobDetails.template?.allowPreflights && !!jobDetails.template?.canRun;
 
