Horde: Re-enable and fix tests for agent's HTTP server

diff --git a/Engine/Source/Programs/Horde/HordeAgent.Tests/Services/ManagementServiceTests.cs b/Engine/Source/Programs/Horde/HordeAgent.Tests/Services/ManagementServiceTests.cs
--- a/Engine/Source/Programs/Horde/HordeAgent.Tests/Services/ManagementServiceTests.cs
+++ b/Engine/Source/Programs/Horde/HordeAgent.Tests/Services/ManagementServiceTests.cs
@@ -31,7 +31,6 @@ public sealed class ManagementServiceTests
 	private readonly WorkerServiceStub _workerServiceStub = new();
 	
 	[TestMethod]
-	[Ignore]
 	public async Task HttpRequest_NotFound_Async()
 	{
 		(ErrorResponse? res, IHttpListenerResponse? httpRes) = await RequestAdminAsync<ErrorResponse>("http://localhost/not-found");
@@ -41,7 +40,6 @@ public sealed class ManagementServiceTests
 	}
 	
 	[TestMethod]
-	[Ignore]
 	public async Task HttpRequest_HealthCheck_Async()
 	{
 		{
@@ -75,13 +73,14 @@ public sealed class ManagementServiceTests
 		HttpListenerResponseStub res = new();
 		await cb(ms, req, res);
 		res.OutputStream.Position = 0;
-		T? deserialized = await JsonSerializer.DeserializeAsync<T>(res.OutputStream, JsonSerializerOptions.Default, CancellationToken.None);
+		
+		T? deserialized = await JsonSerializer.DeserializeAsync<T>(res.OutputStream, ManagementService.JsonSerializerOptions, CancellationToken.None);
 		return (deserialized, res);
 	}
 	
 	private ManagementService GetMs()
 	{
-		AgentSettings settings = new () { AdminEndpoints = ["http://localhost:3232/"], HealthCheckEndpoints = ["http://localhost:3233/"] };
+		AgentSettings settings = new ();
 		TestOptionsMonitor<AgentSettings> settingsOpt = new(settings);
 		return new ManagementService(settingsOpt, _workerServiceStub, NullLogger<ManagementService>.Instance);
 	}
diff --git a/Engine/Source/Programs/Horde/HordeAgent/Services/ManagementService.cs b/Engine/Source/Programs/Horde/HordeAgent/Services/ManagementService.cs
--- a/Engine/Source/Programs/Horde/HordeAgent/Services/ManagementService.cs
+++ b/Engine/Source/Programs/Horde/HordeAgent/Services/ManagementService.cs
@@ -64,7 +64,7 @@ internal class ManagementService : IHostedService, IAsyncDisposable
 	private readonly BackgroundTask? _adminTask;
 	private readonly BackgroundTask? _healthTask;
 	private readonly ILogger<ManagementService> _logger;
-	private static readonly JsonSerializerOptions s_jsonSerializerOptions = new () { WriteIndented = true, PropertyNameCaseInsensitive = true, PropertyNamingPolicy = JsonNamingPolicy.CamelCase};
+	public static readonly JsonSerializerOptions JsonSerializerOptions = new () { WriteIndented = true, PropertyNameCaseInsensitive = true, PropertyNamingPolicy = JsonNamingPolicy.CamelCase};
 	
 	private readonly RoutesDict _adminRoutes = new()
 	{
@@ -79,7 +79,7 @@ internal class ManagementService : IHostedService, IAsyncDisposable
 		{
 			res.ContentType = "application/json";
 			T result = callback(req, res);
-			JsonSerializer.Serialize(res.OutputStream, result, s_jsonSerializerOptions);
+			JsonSerializer.Serialize(res.OutputStream, result, JsonSerializerOptions);
 		};
 	}
 	
