Horde: Return non-200 status code when health check fails in agent

Also improve logging when failing to bind agent HTTP server to all interfaces.

diff --git a/Engine/Source/Programs/Horde/HordeAgent.Tests/Services/ManagementServiceTests.cs b/Engine/Source/Programs/Horde/HordeAgent.Tests/Services/ManagementServiceTests.cs
--- a/Engine/Source/Programs/Horde/HordeAgent.Tests/Services/ManagementServiceTests.cs
+++ b/Engine/Source/Programs/Horde/HordeAgent.Tests/Services/ManagementServiceTests.cs
@@ -2,6 +2,7 @@
 
 using System;
 using System.IO;
+using System.Net;
 using System.Text.Json;
 using System.Threading;
 using System.Threading.Tasks;
@@ -19,6 +20,10 @@ internal class HttpListenerRequestStub(Uri? url) : IHttpListenerRequest
 internal class HttpListenerResponseStub : IHttpListenerResponse
 {
 	public string? ContentType { get; set; }
+	public int StatusCode { get; set; }
+	public long ContentLength64 { get; set; }
+	public WebHeaderCollection Headers { get; set; } = [];
+	public CookieCollection Cookies { get; set; } = new();
 	public Stream OutputStream { get; } = new MemoryStream();
 	public void Close()
 	{
@@ -34,6 +39,7 @@ public sealed class ManagementServiceTests
 	public async Task HttpRequest_NotFound_Async()
 	{
 		(ErrorResponse? res, IHttpListenerResponse? httpRes) = await RequestAdminAsync<ErrorResponse>("http://localhost/not-found");
+		Assert.AreEqual(404, httpRes.StatusCode);
 		Assert.IsFalse(res!.Success);
 		Assert.AreEqual("application/json", httpRes.ContentType);
 		Assert.AreEqual("Not Found", res.Message);
@@ -44,13 +50,15 @@ public sealed class ManagementServiceTests
 	{
 		{
 			_workerServiceStub.IsConnected = true;
-			(HealthCheckResponse? res, _) = await RequestHealthAsync<HealthCheckResponse>("http://localhost/");
+			(HealthCheckResponse? res, IHttpListenerResponse? httpRes) = await RequestHealthAsync<HealthCheckResponse>("http://localhost/");
+			Assert.AreEqual(200, httpRes.StatusCode);
 			Assert.IsTrue(res!.IsHealthy);
 			Assert.IsTrue(res.IsConnected);
 		}
 		{
 			_workerServiceStub.IsConnected = false;
-			(HealthCheckResponse? res, _) = await RequestHealthAsync<HealthCheckResponse>("http://localhost/");
+			(HealthCheckResponse? res, IHttpListenerResponse? httpRes) = await RequestHealthAsync<HealthCheckResponse>("http://localhost/");
+			Assert.AreEqual(503, httpRes.StatusCode);
 			Assert.IsFalse(res!.IsHealthy);
 			Assert.IsFalse(res.IsConnected);
 		}
diff --git a/Engine/Source/Programs/Horde/HordeAgent/Services/ManagementService.cs b/Engine/Source/Programs/Horde/HordeAgent/Services/ManagementService.cs
--- a/Engine/Source/Programs/Horde/HordeAgent/Services/ManagementService.cs
+++ b/Engine/Source/Programs/Horde/HordeAgent/Services/ManagementService.cs
@@ -6,12 +6,12 @@ using EpicGames.Core;
 using Microsoft.Extensions.Hosting;
 using Microsoft.Extensions.Logging;
 using Microsoft.Extensions.Options;
-
 using RoutesDict = System.Collections.Generic.Dictionary<string, System.Action<HordeAgent.Services.IHttpListenerRequest, HordeAgent.Services.IHttpListenerResponse>>;
 
 namespace HordeAgent.Services;
 
 #region Wrapper classes for HttpListener to make it testable
+
 internal interface IHttpListenerRequest
 {
 	Uri? Url { get; }
@@ -19,7 +19,11 @@ internal interface IHttpListenerRequest
 
 internal interface IHttpListenerResponse
 {
-	string? ContentType { get; set;  }
+	string? ContentType { get; set; }
+	int StatusCode { get; set; }
+	long ContentLength64 { get; set; }
+	WebHeaderCollection Headers { get; set; }
+	CookieCollection Cookies { get; set; }
 	Stream OutputStream { get; }
 	void Close();
 }
@@ -31,15 +35,15 @@ internal class WrappedHttpListenerRequest(HttpListenerRequest impl) : IHttpListe
 
 internal class WrappedHttpListenerResponse(HttpListenerResponse impl) : IHttpListenerResponse
 {
-	public string? ContentType
-	{
-		get => impl.ContentType;
-		set => impl.ContentType = value;
-	}
-	
+	public string? ContentType { get => impl.ContentType; set => impl.ContentType = value; }
+	public int StatusCode { get => impl.StatusCode; set => impl.StatusCode = value; }
+	public long ContentLength64 { get => impl.ContentLength64; set => impl.ContentLength64 = value; }
+	public WebHeaderCollection Headers { get => impl.Headers; set => impl.Headers = value; }
+	public CookieCollection Cookies { get => impl.Cookies; set => impl.Cookies = value; }
 	public Stream OutputStream => impl.OutputStream;
 	public void Close() => impl.Close();
 }
+
 #endregion
 
 internal class ErrorResponse(bool success, string message)
@@ -64,20 +68,23 @@ internal class ManagementService : IHostedService, IAsyncDisposable
 	private readonly BackgroundTask? _adminTask;
 	private readonly BackgroundTask? _healthTask;
 	private readonly ILogger<ManagementService> _logger;
-	public static readonly JsonSerializerOptions JsonSerializerOptions = new () { WriteIndented = true, PropertyNameCaseInsensitive = true, PropertyNamingPolicy = JsonNamingPolicy.CamelCase};
 	
-	private readonly RoutesDict _adminRoutes = new()
+	public static readonly JsonSerializerOptions JsonSerializerOptions = new()
 	{
-		["/"] = HandleJson((_, _) => "Admin server"),
+		WriteIndented = true, PropertyNameCaseInsensitive = true, PropertyNamingPolicy = JsonNamingPolicy.CamelCase
 	};
 	
+	private readonly RoutesDict _adminRoutes = new() { ["/"] = HandleJson((_, _) => "Admin server"), };
+	
 	private readonly RoutesDict _healthRoutes;
 	
 	private static Action<IHttpListenerRequest, IHttpListenerResponse> HandleJson<T>(Func<IHttpListenerRequest, IHttpListenerResponse, T> callback)
 	{
 		return (req, res) =>
 		{
+			res.Headers.Set("Server", ""); // Remove the default "Microsoft-HTTPAPI/2.0" getting sent
 			res.ContentType = "application/json";
+			res.StatusCode = 200;
 			T result = callback(req, res);
 			JsonSerializer.Serialize(res.OutputStream, result, JsonSerializerOptions);
 		};
@@ -91,7 +98,12 @@ internal class ManagementService : IHostedService, IAsyncDisposable
 		_logger = logger;
 		_healthRoutes = new RoutesDict
 		{
-			["/"] = HandleJson<HealthCheckResponse>((_, _) => new HealthCheckResponse(workerService.IsConnected, workerService.IsConnected)),
+			["/"] = HandleJson<HealthCheckResponse>((_, res) =>
+			{
+				bool isHealthy = workerService.IsConnected;
+				res.StatusCode = isHealthy ? 200 : 503;
+				return new HealthCheckResponse(isHealthy, workerService.IsConnected);
+			}),
 		};
 		
 		string SettingsEndpointToPrefix(string endpoint) => $"http://{endpoint}/";
@@ -103,6 +115,7 @@ internal class ManagementService : IHostedService, IAsyncDisposable
 			{
 				_adminListener.Prefixes.Add(SettingsEndpointToPrefix(endpoint));
 			}
+			
 			_logger.LogInformation("Admin HTTP server listening on {Endpoints}", String.Join(",", settings.CurrentValue.AdminEndpoints));
 			_adminTask = new BackgroundTask(ct => RunServerAsync(_adminListener, _adminRoutes, ct));
 		}
@@ -114,11 +127,12 @@ internal class ManagementService : IHostedService, IAsyncDisposable
 			{
 				_healthListener.Prefixes.Add(SettingsEndpointToPrefix(endpoint));
 			}
+			
 			_logger.LogInformation("Health check HTTP server listening on {Endpoints}", String.Join(",", settings.CurrentValue.HealthCheckEndpoints));
 			_healthTask = new BackgroundTask(ct => RunServerAsync(_healthListener, _healthRoutes, ct));
 		}
 	}
-
+	
 	/// <inheritdoc/>
 	public async ValueTask DisposeAsync()
 	{
@@ -135,7 +149,7 @@ internal class ManagementService : IHostedService, IAsyncDisposable
 			await _healthTask.DisposeAsync();
 		}
 	}
-
+	
 	/// <inheritdoc/>
 	public Task StartAsync(CancellationToken cancellationToken)
 	{
@@ -148,12 +162,19 @@ internal class ManagementService : IHostedService, IAsyncDisposable
 		}
 		catch (Exception e)
 		{
-			_logger.LogError(e, "Failed to start built-in HTTP server(s) for administration and health checks");
+			// Access denied exception is raised on Windows trying to bind on all interfaces as Administrator privileges are needed
+			// Skip logging the stacktrace to reduce noise in log
+			bool isAccessDenied = e.Message.Contains("Access is denied", StringComparison.OrdinalIgnoreCase);
+			_logger.LogError(isAccessDenied ? null : e, "Failed to start built-in HTTP server(s) for administration and health checks");
+			if (isAccessDenied)
+			{
+				_logger.LogError("Binding on all interfaces requires elevated process permissions (administrator/root)");
+			}
 		}
 		
 		return Task.CompletedTask;
 	}
-
+	
 	/// <inheritdoc/>
 	public async Task StopAsync(CancellationToken cancellationToken)
 	{
@@ -166,7 +187,7 @@ internal class ManagementService : IHostedService, IAsyncDisposable
 		{
 			await _healthTask.StopAsync(cancellationToken);
 		}
-
+		
 		_adminListener?.Stop();
 		_healthListener?.Stop();
 	}
@@ -192,7 +213,11 @@ internal class ManagementService : IHostedService, IAsyncDisposable
 			}
 			else
 			{
-				HandleJson((_, _) => new ErrorResponse(false, "Not Found"))(req, res);
+				HandleJson((_, res) =>
+				{
+					res.StatusCode = 404;
+					return new ErrorResponse(false, "Not Found");
+				})(req, res);
 			}
 		}
 		catch (Exception e)
@@ -213,9 +238,9 @@ internal class ManagementService : IHostedService, IAsyncDisposable
 		while (!cancellationToken.IsCancellationRequested)
 		{
 			HttpListenerContext context = await listener.GetContextAsync();
-			WrappedHttpListenerRequest wrappedReq = new (context.Request);
-			WrappedHttpListenerResponse wrappedRes = new (context.Response);
+			WrappedHttpListenerRequest wrappedReq = new(context.Request);
+			WrappedHttpListenerResponse wrappedRes = new(context.Response);
 			await HandleRequestAsync(wrappedReq, wrappedRes, routes);
 		}
 	}
-}
+}
\ No newline at end of file
