Horde: Add triage issue severity workflow setting to enable warning/error only workflows

diff --git a/Engine/Source/Programs/Horde/Plugins/Build/HordeServer.Build/Issues/WorkflowConfig.cs b/Engine/Source/Programs/Horde/Plugins/Build/HordeServer.Build/Issues/WorkflowConfig.cs
--- a/Engine/Source/Programs/Horde/Plugins/Build/HordeServer.Build/Issues/WorkflowConfig.cs
+++ b/Engine/Source/Programs/Horde/Plugins/Build/HordeServer.Build/Issues/WorkflowConfig.cs
@@ -47,6 +47,12 @@ namespace HordeServer.Issues
 		/// <inheritdoc/>
 		public string? TriageAlias { get; set; }
 
+		/// <inheritdoc/>
+		public bool TriageErrors { get; set; } = true;
+
+		/// <inheritdoc/>
+		public bool TriageWarnings { get; set; } = true;
+
 		/// <inheritdoc cref="IWorkflow.TriageTypeAliases"/>
 		public Dictionary<string, string>? TriageTypeAliases { get; set; }
 		IReadOnlyDictionary<string, string>? IWorkflow.TriageTypeAliases => TriageTypeAliases;
diff --git a/Engine/Source/Programs/Horde/Plugins/Build/HordeServer.Build/Notifications/Sinks/SlackNotificationSink.cs b/Engine/Source/Programs/Horde/Plugins/Build/HordeServer.Build/Notifications/Sinks/SlackNotificationSink.cs
--- a/Engine/Source/Programs/Horde/Plugins/Build/HordeServer.Build/Notifications/Sinks/SlackNotificationSink.cs
+++ b/Engine/Source/Programs/Horde/Plugins/Build/HordeServer.Build/Notifications/Sinks/SlackNotificationSink.cs
@@ -865,8 +865,12 @@ namespace HordeServer.Notifications.Sinks
 					StreamConfig? streamConfig;
 					if (buildConfig.TryGetStream(span.StreamId, out streamConfig) && streamConfig.TryGetWorkflow(workflowId.Value, out workflow) && !String.IsNullOrEmpty(workflow.TriageChannel))
 					{
-						await CreateOrUpdateWorkflowThreadAsync(workflow.TriageChannel, issue, span, details.Spans, workflow, cancellationToken);
-						notifyOwner = notifySuspects = false;
+						bool skip = (issue.Severity == IssueSeverity.Warning && !workflow.TriageWarnings) || (issue.Severity == IssueSeverity.Error && !workflow.TriageErrors);
+						if (!skip)
+						{
+							await CreateOrUpdateWorkflowThreadAsync(workflow.TriageChannel, issue, span, details.Spans, workflow, cancellationToken);
+							notifyOwner = notifySuspects = false;
+						}
 					}
 				}
 			}
@@ -1554,10 +1558,14 @@ namespace HordeServer.Notifications.Sinks
 				StreamConfig? streamConfig;
 				if (buildConfig.TryGetStream(span.StreamId, out streamConfig) && streamConfig.TryGetWorkflow(workflowId.Value, out WorkflowConfig? workflow) && !String.IsNullOrEmpty(workflow.TriageChannel) && workflow.AllowMentions)
 				{
-					MessageStateDocument? state = await GetMessageStateAsync(workflow.TriageChannel, GetTriageThreadEventId(issue.Id), cancellationToken);
-					if (state != null)
-					{
-						summaryBuilder.Append($" See *<{state.Permalink}|discussion thread>*.");
+					bool skip = (issue.Severity == IssueSeverity.Warning && !workflow.TriageWarnings) || (issue.Severity == IssueSeverity.Error && !workflow.TriageErrors);
+					if (!skip)
+					{ 
+						MessageStateDocument? state = await GetMessageStateAsync(workflow.TriageChannel, GetTriageThreadEventId(issue.Id), cancellationToken);
+						if (state != null)
+						{
+							summaryBuilder.Append($" See *<{state.Permalink}|discussion thread>*.");
+						}
 					}
 				}
 			}
@@ -2853,7 +2861,7 @@ namespace HordeServer.Notifications.Sinks
 			{
 				return null;
 			}
-			if (String.IsNullOrEmpty(workflow.TriageChannel) || workflow.EscalateAlias == null || workflow.EscalateTimes.Count == 0)
+			if (String.IsNullOrEmpty(workflow.TriageChannel) || workflow.EscalateAlias == null || workflow.EscalateTimes.Count == 0 || (!workflow.TriageWarnings && issue.Severity == IssueSeverity.Warning) || (!workflow.TriageErrors && issue.Severity == IssueSeverity.Error))
 			{
 				return null;
 			}
diff --git a/Engine/Source/Programs/Shared/EpicGames.Horde/Issues/IWorkflow.cs b/Engine/Source/Programs/Shared/EpicGames.Horde/Issues/IWorkflow.cs
--- a/Engine/Source/Programs/Shared/EpicGames.Horde/Issues/IWorkflow.cs
+++ b/Engine/Source/Programs/Shared/EpicGames.Horde/Issues/IWorkflow.cs
@@ -66,6 +66,16 @@ namespace EpicGames.Horde.Issues
 		/// </summary>
 		string? TriageAlias { get; }
 
+		/// <summary>
+		/// Whether to include issues with an error status in the triage
+		/// </summary>
+		bool TriageErrors { get; }
+
+		/// <summary>
+		/// Whether to include issues with a warning status in the triage
+		/// </summary>
+		bool TriageWarnings { get; }
+
 		/// <summary>
 		/// Slack user/alias to ping for specific issue types (such as Systemic), if there is nobody assigned to an issue by default.
 		/// </summary>
diff --git a/Engine/Source/Programs/Shared/EpicGames.Horde/Issues/WorkflowMessages.cs b/Engine/Source/Programs/Shared/EpicGames.Horde/Issues/WorkflowMessages.cs
--- a/Engine/Source/Programs/Shared/EpicGames.Horde/Issues/WorkflowMessages.cs
+++ b/Engine/Source/Programs/Shared/EpicGames.Horde/Issues/WorkflowMessages.cs
@@ -48,6 +48,12 @@ namespace EpicGames.Horde.Issues
 		/// <inheritdoc/>
 		public string? TriageAlias { get; set; }
 
+		/// <inheritdoc/>
+		public bool TriageErrors { get; set; } = true;
+
+		/// <inheritdoc/>
+		public bool TriageWarnings { get; set; } = true;
+
 		/// <inheritdoc cref="IWorkflow.TriageTypeAliases"/>
 		public Dictionary<string, string>? TriageTypeAliases { get; set; }
 		IReadOnlyDictionary<string, string>? IWorkflow.TriageTypeAliases => TriageTypeAliases;
@@ -112,6 +118,8 @@ namespace EpicGames.Horde.Issues
 			TriageInstructions = workflow.TriageInstructions;
 			TriageAlias = workflow.TriageAlias;
 			TriageTypeAliases = workflow.TriageTypeAliases?.ToDictionary();
+			TriageWarnings = workflow.TriageWarnings;
+			TriageErrors = workflow.TriageErrors;
 			EscalateAlias = workflow.EscalateAlias;
 			EscalateTimes = workflow.EscalateTimes.ToList();
 			MaxMentions = workflow.MaxMentions;
