Horde: Validate ServerUrl for auth modes and load balancers in settings

- Add ServerUrl validation check for non-anonymous auth modes to prevent redirect issues
- Add middleware to detect reverse proxy/load balancer deployments
- Log warning when ServerUrl is not configured in proxy environments
- Forgetting to set this has been a source of issues when deploying Horde

diff --git a/Engine/Source/Programs/Horde/HordeServer/Server/ReverseProxyDetection.cs b/Engine/Source/Programs/Horde/HordeServer/Server/ReverseProxyDetection.cs
new file mode 100644
--- /dev/null
+++ b/Engine/Source/Programs/Horde/HordeServer/Server/ReverseProxyDetection.cs
@@ -0,0 +1,61 @@
+﻿// Copyright Epic Games, Inc. All Rights Reserved.
+
+using System.Collections.Generic;
+using System.Threading.Tasks;
+using Microsoft.AspNetCore.Builder;
+using Microsoft.AspNetCore.Http;
+using Microsoft.Extensions.Logging;
+using Microsoft.Extensions.Options;
+
+namespace HordeServer.Server;
+
+/// <summary>
+/// Middleware that detects if the server is running behind a load balancer or reverse proxy.
+/// If detected, logs an error if <see cref="ServerSettings.ServerUrl"/> is not configured.
+/// </summary>
+/// <param name="next">Next middleware in the pipeline</param>
+/// <param name="logger">Logger</param>
+/// <param name="settings">Settings</param>
+public class ReverseProxyDetectionMiddleware(RequestDelegate next, ILogger<ReverseProxyDetectionMiddleware> logger, IOptionsMonitor<ServerSettings> settings)
+{
+	private const int MaxWarnings = 5;
+	private const int MaxRequests = 5000;
+	private readonly List<string> _headers = ["X-Forwarded-For", "X-Forwarded-Proto", "X-Forwarded-Host", "X-Forwarded-Server",];
+	private int _requestCount;
+	private int _warningCount;
+	
+	/// <summary>
+	/// Processes an HTTP request to detect reverse proxy headers and validate server configuration
+	/// </summary>
+	public async Task InvokeAsync(HttpContext context)
+	{
+		// Only check the first X number of requests for performance
+		if (_requestCount < MaxRequests && _warningCount < MaxWarnings)
+		{
+			_requestCount++;
+			bool isReverseProxyRequest = _headers.Exists(header => context.Request.Headers.ContainsKey(header));
+			if (isReverseProxyRequest && settings.CurrentValue.IsDefaultServerUrlUsed())
+			{
+				_warningCount++;
+				logger.LogError("Load balancer or reverse proxy detected. Ensure {Setting} is configured in server settings",
+					nameof(ServerSettings.ServerUrl));
+			}
+		}
+		
+		await next(context);
+	}
+}
+
+/// <summary>
+/// Extension methods for configuring the reverse proxy detection middleware.
+/// </summary>
+public static class ReverseProxyDetectionMiddlewareExtensions
+{
+	/// <summary>
+	/// Adds middleware to detect if the application is running behind a reverse proxy
+	/// </summary>
+	public static IApplicationBuilder UseReverseProxyDetection(this IApplicationBuilder builder)
+	{
+		return builder.UseMiddleware<ReverseProxyDetectionMiddleware>();
+	}
+}
\ No newline at end of file
diff --git a/Engine/Source/Programs/Horde/HordeServer/ServerApp.cs b/Engine/Source/Programs/Horde/HordeServer/ServerApp.cs
--- a/Engine/Source/Programs/Horde/HordeServer/ServerApp.cs
+++ b/Engine/Source/Programs/Horde/HordeServer/ServerApp.cs
@@ -195,6 +195,8 @@ namespace HordeServer
 #pragma warning disable CA2000 // Dispose objects before losing scope
 			ILogger startupLogger = new SerilogLoggerFactory().CreateLogger(typeof(ServerApp).FullName ?? "ServerApp");
 #pragma warning restore CA2000 // Dispose objects before losing scope
+			
+			serverSettings.Validate(startupLogger);
 
 			try
 			{
diff --git a/Engine/Source/Programs/Horde/HordeServer/ServerSettings.cs b/Engine/Source/Programs/Horde/HordeServer/ServerSettings.cs
--- a/Engine/Source/Programs/Horde/HordeServer/ServerSettings.cs
+++ b/Engine/Source/Programs/Horde/HordeServer/ServerSettings.cs
@@ -7,6 +7,7 @@ using System.Net;
 using EpicGames.Horde.Server;
 using EpicGames.Horde.Utilities;
 using HordeServer.Utilities;
+using Microsoft.Extensions.Logging;
 
 namespace HordeServer
 {
@@ -154,6 +155,7 @@ namespace HordeServer
 		/// Type of authentication (e.g anonymous, OIDC, built-in Horde accounts)
 		/// If "Horde" auth mode is used, be sure to configure "ServerUrl" as well.
 		/// </summary>
+		/// <see cref="ServerUrl"/>
 		public AuthMethod AuthMethod { get; set; } = AuthMethod.Anonymous;
 		
 		/// <summary>
@@ -267,9 +269,9 @@ namespace HordeServer
 
 		/// <summary>
 		/// Base URL this Horde server is accessible from
-		/// For example https://horde.mystudio.com
-		/// If not set, a default is used based on current hostname.
-		/// In more advanced setups where a reverse proxy is present in front of Horde, this must manually be set.
+		/// For example https://horde.mystudio.com. If not set, a default is used based on current hostname.
+		/// It's important this URL matches where users and agents access the server as it's used for signing auth tokens etc.
+		/// Must be configured manually when running behind a reverse proxy or load balancer
 		/// </summary>
 		public Uri ServerUrl
 		{
@@ -285,6 +287,8 @@ namespace HordeServer
 			get => _jwtIssuer ?? ServerUrl.ToString();
 			set => _jwtIssuer = value;
 		}
+		
+		internal bool IsDefaultServerUrlUsed() =>_serverUrl == GetDefaultServerUrl();
 
 		Uri? _serverUrl;
 		string? _jwtIssuer;
@@ -476,12 +480,17 @@ namespace HordeServer
 		/// Validate the settings object does not contain any invalid fields
 		/// </summary>
 		/// <exception cref="ArgumentException"></exception>
-		public void Validate()
+		public void Validate(ILogger? logger = null)
 		{
 			if (RunModes != null && IsRunModeActive(RunMode.None))
 			{
 				throw new ArgumentException($"Settings key '{nameof(RunModes)}' contains one or more invalid entries");
 			}
+			
+			if (AuthMethod != AuthMethod.Anonymous && IsDefaultServerUrlUsed() && logger != null)
+			{
+				logger.LogError("Configure {Setting} in server settings when non-anonymous auth method is used", nameof(ServerUrl));
+			}
 		}
 	}
 }
diff --git a/Engine/Source/Programs/Horde/HordeServer/Startup.cs b/Engine/Source/Programs/Horde/HordeServer/Startup.cs
--- a/Engine/Source/Programs/Horde/HordeServer/Startup.cs
+++ b/Engine/Source/Programs/Horde/HordeServer/Startup.cs
@@ -215,11 +215,11 @@ namespace HordeServer
 
 		public IConfiguration Configuration { get; }
 
-		public static void BindServerSettings(IConfiguration configuration, ServerSettings settings)
+		public static void BindServerSettings(IConfiguration configuration, ServerSettings settings, ILogger? logger = null)
 		{
 			IConfigurationSection hordeSection = configuration.GetSection("Horde");
 			hordeSection.Bind(settings);
-			settings.Validate();
+			settings.Validate(logger);
 			settings.AddDefaultOidcScopesAndMappings();
 		}
 
@@ -998,7 +998,8 @@ namespace HordeServer
 			{
 				app.UseMiddleware<RequestTrackerMiddleware>();
 			}
-
+			
+			app.UseReverseProxyDetection();
 			app.UseExceptionHandler("/api/v1/exception");
 
 			DirectoryReference dashboardDir = DirectoryReference.Combine(ServerApp.AppDir, "DashboardApp");
