Horde: Updating release notes
#rnx

diff --git a/Engine/Source/Programs/Horde/Docs/ReleaseNotes.md b/Engine/Source/Programs/Horde/Docs/ReleaseNotes.md
--- a/Engine/Source/Programs/Horde/Docs/ReleaseNotes.md
+++ b/Engine/Source/Programs/Horde/Docs/ReleaseNotes.md
@@ -1,5 +1,41 @@
 # Release Notes
 
+## 2025-01-30
+
+* Validate ServerUrl for auth modes and load balancers in settings
+  - Add ServerUrl validation check for non-anonymous auth modes to prevent redirect issues
+  - Add middleware to detect reverse proxy/load balancer deployments
+  - Log warning when ServerUrl is not configured in proxy environments
+  - Forgetting to set this has been a source of issues when deploying Horde (39586905)
+* Fixed build warning in CbField json serialization (39586077)
+* Fixed incorrect json serialization of cb fields that had fields without names (mostly encountered when serializing an array). (39585891)
+* Add triage issue severity workflow setting to enable warning/error only workflows (39563253)
+* Return non-200 status code when health check fails in agent
+  Also improve logging when failing to bind agent HTTP server to all interfaces. (39550702)
+* Improve where optional attributes can appear in enum, class, and struct definitions. (39548575)
+* Re-enable and fix tests for agent's HTTP server (39546007)
+* Add small HTTP server in agent for administration and health checks
+  In this initial version, the health check endpoint returns a basic status whether it is connected to the Horde server or not. (39518264)
+* OidcToken - Added missing copyright to tests and fixed other analyze errors in oidc token (39513944)
+* Adding artifact type expiry information to dashboard config query (39513838)
+* OidcToken - Refactored token manager to allow for tests and added some functional tests (39512897)
+* Cloud Artifacts - Support uploading files larger then 2GB - implemented support for compressed buffer serialization in C# to read from a stream thus supporting compressing buffers larger then 2GB, still allocating this all in memory so will use a lot of memory for large buffers. (39512887)
+* Run job executor cleanup when the job batch process fails or hangs (39494168)
+* EpicGames.Perforce: Handle ... syntax with file extension in view mappings (four dots) (39478199)
+* PR #11833: Add throw if Refresh token was not returned in response. (39477147)
+* Display N/A Start and Stop Times for User Notices (39438421)
+* Add Select All to small Templates submenus (39438351)
+* Show partial elements selected on Templates context menu (39438288)
+* [EpicGames.UBA] Fixed so uba logger log scope logic works (39429420)
+* Add database schema migration tool for MongoDB (39407820)
+* Fixed ToolTips not working on List type parameters in stream configs (39407062)
+* Prevent exception swallowing during job batch execution (39371897)
+* Reverse job audit logs to be top down (39356984)
+* Add job audit view to dashboard (39356298)
+* Add job audit logging for tracing of the state machine (39354365,39346324)
+* EpicGames.Core: Fix unsafe type cast in LoggerScopeCollection.Scope constructor (39343576)
+* Improve logging for job step log writer errors (39323447)
+
 ## 2025-01-14
 
 * Set the "background": true flag when creating new indexes, and replace the deletion of unused indexes with a warning message. We don't want old pods to delete indexes created by new pods. (39182560)
