Horde: Add UX to log view to skip to next block of issue lines

diff --git a/Engine/Source/Programs/Horde/HordeDashboard/src/components/LogView.tsx b/Engine/Source/Programs/Horde/HordeDashboard/src/components/LogView.tsx
--- a/Engine/Source/Programs/Horde/HordeDashboard/src/components/LogView.tsx
+++ b/Engine/Source/Programs/Horde/HordeDashboard/src/components/LogView.tsx
@@ -84,6 +84,36 @@ class LogHandler {
       return this.getLogEvent(this.currentLine)
    }
 
+   getEventsOfSeverity(warnings?: boolean){
+      return this.events.filter(e => {
+         if (warnings && e.severity !== EventSeverity.Warning) {
+            return false;
+         }
+         if (!warnings && e.severity !== EventSeverity.Error) {
+            return false;
+         }
+         return true;
+      });
+   }
+
+   getEventBlocks(events: GetLogEventResponse[]){
+      return events.filter((e, i, events) => {
+         // first event must be a block
+         if (i == 0) {
+            return true
+         }
+         // consider different issue a new block, regardless of spacing
+         if (e.issueId != events[i-1].issueId) {
+            return true
+         }
+         
+         if (e.lineIndex == events[i-1].lineIndex + events[i-1].lineCount) {
+            return false
+         }
+         return true
+      });
+   }
+
    getNextLogEvent(warnings?: boolean) {
 
       let currentLine = -1;
@@ -93,15 +123,7 @@ class LogHandler {
 
       const startLine = currentLine + 1;
 
-      const events = this.events.filter(e => {
-         if (warnings && e.severity !== EventSeverity.Warning) {
-            return false;
-         }
-         if (!warnings && e.severity !== EventSeverity.Error) {
-            return false;
-         }
-         return true;
-      });
+      const events = this.getEventsOfSeverity(warnings);
 
       let event = events.find(e => e.lineIndex >= startLine);
 
@@ -112,6 +134,27 @@ class LogHandler {
       return event;
    }
 
+   getNextLogEventBlock(warnings?: boolean) {
+      
+      let currentLine = -1;
+      if (this.currentLine !== undefined) {
+         currentLine = this.currentLine - 1;
+      }
+
+      const startLine = currentLine + 1;
+
+      const events = this.getEventsOfSeverity(warnings);
+      const eventBlocks = this.getEventBlocks(events)
+
+      let event = eventBlocks.find(eb => eb.lineIndex >= startLine);
+
+      if (!event && eventBlocks.length) {
+         event = eventBlocks[0];
+      }
+
+      return event;
+   }
+
    getPrevLogEvent(warnings?: boolean) {
 
       let currentLine = 0;
@@ -121,15 +164,7 @@ class LogHandler {
 
       const startLine = currentLine - 1;
 
-      const events = this.events.filter(e => {
-         if (warnings && e.severity !== EventSeverity.Warning) {
-            return false;
-         }
-         if (!warnings && e.severity !== EventSeverity.Error) {
-            return false;
-         }
-         return true;
-      }).reverse();
+      const events = this.getEventsOfSeverity(warnings).reverse();
 
       if (startLine < 0) {
          if (events.length) {
@@ -147,6 +182,34 @@ class LogHandler {
       return event;
    }
 
+   getPrevLogEventBlock(warnings?: boolean) {
+      
+      let currentLine = -1;
+      if (this.currentLine !== undefined) {
+         currentLine = this.currentLine - 1;
+      }
+
+      const startLine = currentLine - 1;
+
+      const events = this.getEventsOfSeverity(warnings);
+      const eventBlocks = this.getEventBlocks(events).reverse();
+
+      if (startLine < 0) {
+         if (eventBlocks.length) {
+            return eventBlocks[0];
+         }
+         return undefined;
+      }
+
+      let event = eventBlocks.find(eb => eb.lineIndex <= startLine);
+
+      if (!event && eventBlocks.length) {
+         event = eventBlocks[0];
+      }
+
+      return event;
+   }
+
    getLogEvent(line: number | undefined) {
 
       if (line === undefined) {
@@ -1166,6 +1229,21 @@ export const LogList: React.FC<{ logId: string }> = observer(({ logId }) => {
 
    }
 
+   function prevErrorBlock() {
+
+      if (!handler) {
+         return;
+      }
+
+      const eventBlock = handler.getPrevLogEventBlock();
+
+      if (eventBlock) {
+         handler.currentLine = eventBlock.lineIndex + 1;
+      }
+
+      updateEvent();
+   }
+
    function nextError() {
 
       if (!handler) {
@@ -1181,6 +1259,20 @@ export const LogList: React.FC<{ logId: string }> = observer(({ logId }) => {
       updateEvent();
    }
 
+   function nextErrorBlock() {
+
+      if (!handler) {
+         return;
+      }
+
+      const eventBlock = handler.getNextLogEventBlock();
+
+      if (eventBlock) {
+         handler.currentLine = eventBlock.lineIndex + 1;
+      }
+
+      updateEvent();
+   }
 
    function prevWarning() {
 
@@ -1198,6 +1290,22 @@ export const LogList: React.FC<{ logId: string }> = observer(({ logId }) => {
 
    }
 
+   function prevWarningBlock() {
+
+      if (!handler) {
+         return;
+      }
+
+      const eventBlock = handler.getPrevLogEventBlock(true);
+
+      if (eventBlock) {
+         handler.currentLine = eventBlock.lineIndex + 1;
+      }
+
+      updateEvent();
+
+   }
+
    function nextWarning() {
       if (!handler) {
          return;
@@ -1213,6 +1321,21 @@ export const LogList: React.FC<{ logId: string }> = observer(({ logId }) => {
 
    }
 
+   function nextWarningBlock() {
+      if (!handler) {
+         return;
+      }
+
+      const eventBlock = handler.getNextLogEventBlock(true);
+
+      if (eventBlock) {
+         handler.currentLine = eventBlock.lineIndex + 1;
+      }
+
+      updateEvent();
+
+   }
+
    let warningsText = "";
    let errorText = "";
 
@@ -1243,9 +1366,13 @@ export const LogList: React.FC<{ logId: string }> = observer(({ logId }) => {
    }
 
    const nextErrorEnabled = handler?.getNextLogEvent() ? true : false;
-   const prevErrorEnablesd = handler?.getPrevLogEvent() ? true : false;
+   const prevErrorEnabled = handler?.getPrevLogEvent() ? true : false;
+   const nextErrorBlockEnabled = handler?.getNextLogEventBlock() ? true : false;
+   const prevErrorBlockEnabled = handler?.getPrevLogEventBlock() ? true : false;
    const nextWarningEnabled = handler?.getNextLogEvent(true) ? true : false;
    const prevWarningEnabled = handler?.getPrevLogEvent(true) ? true : false;
+   const nextWarningBlockEnabled = handler?.getNextLogEventBlock(true) ? true : false;
+   const prevWarningBlockEnabled = handler?.getPrevLogEventBlock(true) ? true : false;
 
    return <Stack>
       {!!fixme && <IssueModalV2 issueId={query.get("issue")} popHistoryOnClose={issueHistory} />}
@@ -1260,44 +1387,88 @@ export const LogList: React.FC<{ logId: string }> = observer(({ logId }) => {
                <Stack horizontal style={{ paddingBottom: 4 }}>
                   <Stack className={hordeClasses.button} horizontal horizontalAlign={"start"} verticalAlign="center" tokens={{ childrenGap: 8 }}>
                      <Stack horizontal tokens={{ childrenGap: 2 }}>
-                        <DefaultButton disabled={!nextErrorEnabled} className={nextErrorEnabled ? handler.style.errorButton : handler.style.errorButtonDisabled}
-                           text={`${errorText} ${errors.length === 1 ? "Error" : "Errors"}`}
-                           onClick={() => {
-                              nextError();
-                           }}
-                           style={{ color: errors.length ? "#F9F9FB" : buttonNoneText, padding: 15 }} >
-                           {!!errors.length && <Icon style={{ fontSize: 19, paddingLeft: 12 }} iconName='ChevronDown' />}
-                        </DefaultButton>
-
-                        {!!errors.length && <Stack>
-                           <IconButton disabled={!prevErrorEnablesd} className={prevErrorEnablesd ? handler.style.errorButton : handler.style.errorButtonDisabled} style={{ height: 30, fontSize: 19, padding: 8 }} iconProps={{ iconName: 'ChevronUp' }} onClick={(event: any) => {
+
+                        <Stack horizontal verticalAlign='center'>
+                           <Stack horizontal verticalAlign='center' style={{position: 'relative'}}>
+                              <DefaultButton disabled={!nextErrorEnabled} className={nextErrorEnabled ? handler.style.errorButton : handler.style.errorButtonDisabled}
+                                 text={`${errorText} ${errors.length === 1 ? "Error" : "Errors"}`}
+                                 onClick={() => {
+                                    nextError();
+                                 }}
+                                 style={{ color: errors.length ? "#F9F9FB" : buttonNoneText, height: 30, fontSize: 19, paddingLeft: 12, borderRadius: '4px 0 0 4px' }}>
+                                 {!!errors.length && <Icon style={{ fontSize: 19, paddingLeft: 4, paddingRight: 0 }} iconName='ChevronDown' />}
+                              </DefaultButton>
+                              {!!errors.length && <div style={{backgroundColor: "#F9F9FB", height: 18, width: 1, position: 'absolute', right: -0.5, zIndex: 1 }} />}
+                           </Stack>
+                           {!!errors.length && <IconButton disabled={!prevErrorEnabled} className={prevErrorEnabled ? handler.style.errorButton : handler.style.errorButtonDisabled} style={{ color: "#F9F9FB", height: 30, fontSize: 19, padding: 2, borderRadius: '0 4px 4px 0' }} iconProps={{ iconName: 'ChevronUp' }} onClick={(event: any) => {
                               event?.stopPropagation();
                               prevError();
 
-                           }} />
+                           }} />}
+                        </Stack>
 
+                        {!!errors.length && <Stack horizontal verticalAlign='center' style={{paddingLeft: 1}} >
+                           <Stack horizontal verticalAlign='center' style={{position: 'relative'}}>
+                              <IconButton disabled={!nextErrorBlockEnabled} className={nextErrorBlockEnabled ? handler.style.errorButton : handler.style.errorButtonDisabled}
+                                 style={{ color: errors.length ? "#F9F9FB" : buttonNoneText, height: 30, fontSize: 19, padding: 9, borderRadius: '4px 0 0 4px' }}
+                                 styles={{rootPressed: {border: 'none', boxShadow: 'none'}}} 
+                                 iconProps={{ iconName: 'DoubleChevronDown' }}
+                                 onClick={() => {
+                                    nextErrorBlock();
+                                 }} >
+                                 </IconButton>
+                                 <div style={{backgroundColor: "#F9F9FB", height: 18, width: 1, position: 'absolute', right: -0.5, zIndex: 1 }}></div>
+                           </Stack>
+                           <IconButton disabled={!prevErrorBlockEnabled} className={prevErrorBlockEnabled ? handler.style.errorButton : handler.style.errorButtonDisabled}
+                              style={{ color: errors.length ? "#F9F9FB" : buttonNoneText, height: 30, fontSize: 19, padding: 9,  borderRadius: '0 4px 4px 0'}} 
+                              iconProps={{ iconName: 'DoubleChevronUp' }}
+                              onClick={(event: any) => {
+                                 event?.stopPropagation();
+                                 prevErrorBlock();
+                              }} />
                         </Stack>}
 
 
-                        <Stack horizontal style={{ paddingLeft: 12 }} tokens={{ childrenGap: 2 }}>
-                           <DefaultButton disabled={!nextWarningEnabled} className={nextWarningEnabled ? handler.style.warningButton : handler.style.warningButtonDisabled}
-                              text={`${warningsText} ${warnings.length === 1 ? "Warning" : "Warnings"}`}
-                              onClick={() => {
-                                 nextWarning();
-                              }}
-                              style={{ color: (dashboard.darktheme && nextWarningEnabled) ? "#F9F9FB" : buttonNoneText, padding: 15 }} >
-                              {!!warnings.length && <Icon style={{ fontSize: 19, paddingLeft: 12 }} iconName='ChevronDown' />}
-                           </DefaultButton>
+                        <Stack horizontal verticalAlign='center' style={{ paddingLeft: 12 }} >
+                           <Stack horizontal verticalAlign='center' style={{position: 'relative'}}>
+                              <DefaultButton disabled={!nextWarningEnabled} className={nextWarningEnabled ? handler.style.warningButton : handler.style.warningButtonDisabled}
+                                 text={`${warningsText} ${warnings.length === 1 ? "Warning" : "Warnings"}`}
+                                 onClick={() => {
+                                    nextWarning();
+                                 }}
+                                 style={{ color: (dashboard.darktheme && nextWarningEnabled) ? "#F9F9FB" : buttonNoneText, height: 30, fontSize: 19, padding: 12, borderRadius: '4px 0 0 4px'}} >
+                                 {!!warnings.length && <Icon style={{ fontSize: 19, paddingLeft: 4, paddingRight: 0 }} iconName='ChevronDown' />}
+                              </DefaultButton>
+                              {!!warnings.length && <div style={{backgroundColor: (dashboard.darktheme && nextWarningEnabled) ? "#F9F9FB" : buttonNoneText, height: 18, width: 1, position: 'absolute', right: -0.5, zIndex: 1 }} />}
+                           </Stack>
                            {!!warnings.length && <Stack>
-                              <IconButton disabled={!prevWarningEnabled} className={prevWarningEnabled ? handler.style.warningButton : handler.style.warningButtonDisabled} style={{ height: 30, fontSize: 19, padding: 8, color: (dashboard.darktheme && nextWarningEnabled) ? "#F9F9FB" : buttonNoneText }} iconProps={{ iconName: 'ChevronUp' }} onClick={(event: any) => {
+                              <IconButton disabled={!prevWarningEnabled} className={prevWarningEnabled ? handler.style.warningButton : handler.style.warningButtonDisabled} style={{ height: 30, fontSize: 19, padding: 2, color: (dashboard.darktheme && nextWarningEnabled) ? "#F9F9FB" : buttonNoneText, borderRadius: '0 4px 4px 0'}} iconProps={{ iconName: 'ChevronUp' }} onClick={(event: any) => {
                                  event?.stopPropagation();
                                  prevWarning();
-
                               }} />
-
                            </Stack>}
                         </Stack>
 
+                        {!!warnings.length && <Stack horizontal verticalAlign='center' style={{paddingLeft: 1}} >
+                           <Stack horizontal verticalAlign='center' style={{position: 'relative'}}>
+                              <IconButton disabled={!nextWarningBlockEnabled} className={nextWarningBlockEnabled ? handler.style.warningButton : handler.style.warningButtonDisabled}
+                                 style={{ color: (dashboard.darktheme && nextWarningEnabled) ? "#F9F9FB" : buttonNoneText, height: 30, fontSize: 19, padding: 9, borderRadius: '4px 0 0 4px' }}
+                                 styles={{rootPressed: {border: 'none', boxShadow: 'none'}}} 
+                                 iconProps={{ iconName: 'DoubleChevronDown' }}
+                                 onClick={() => {
+                                    nextWarningBlock();
+                                 }} >
+                                 </IconButton>
+                                 <div style={{backgroundColor: (dashboard.darktheme && nextWarningEnabled) ? "#F9F9FB" : buttonNoneText, height: 18, width: 1, position: 'absolute', right: -0.5, zIndex: 1 }}></div>
+                           </Stack>
+                           <IconButton disabled={!prevWarningBlockEnabled} className={prevWarningBlockEnabled ? handler.style.warningButton : handler.style.warningButtonDisabled}
+                              style={{ color: (dashboard.darktheme && nextWarningEnabled) ? "#F9F9FB" : buttonNoneText, height: 30, fontSize: 19, padding: 9,  borderRadius: '0 4px 4px 0'}} 
+                              iconProps={{ iconName: 'DoubleChevronUp' }}
+                              onClick={(event: any) => {
+                                 event?.stopPropagation();
+                                 prevWarningBlock();
+                              }} />
+                        </Stack>}
 
                         {!!menuProps.items.find(i => !i.disabled) && <Stack style={{ paddingLeft: 18 }}><PrimaryButton disabled={!menuProps.items.find(i => !i.disabled)} text="View" menuProps={menuProps} style={{ fontFamily: "Horde Open Sans SemiBold", borderStyle: "hidden", padding: 15 }} /></Stack>}
                         {!menuProps.items.find(i => !i.disabled) && <Stack style={{ paddingLeft: 18 }}><DefaultButton disabled={true} className={handler.style.warningButtonDisabled} text="View" style={{ color: "rgb(97, 110, 133)", padding: 15 }} /> </Stack>}
diff --git a/Engine/Source/Programs/Horde/HordeDashboard/src/styles/Styles.tsx b/Engine/Source/Programs/Horde/HordeDashboard/src/styles/Styles.tsx
--- a/Engine/Source/Programs/Horde/HordeDashboard/src/styles/Styles.tsx
+++ b/Engine/Source/Programs/Horde/HordeDashboard/src/styles/Styles.tsx
@@ -62,6 +62,8 @@ registerIcons({
       "ChevronDown": IconContents.CHEVRON_DOWN,
       "ChevronLeft": IconContents.CHEVRON_LEFT,
       "ChevronRight": IconContents.CHEVRON_RIGHT,
+      "DoubleChevronUp": IconContents.DOUBLE_CHEVRON_UP,
+      "DoubleChevronDown": IconContents.DOUBLE_CHEVRON_DOWN,
       "Refresh": IconContents.REFRESH,
       "Warning": IconContents.WARNING_SIGN,
       "LayoutCircle": IconContents.LAYOUT_CIRCLE,
