Horde: Refactor job handler lease exceptions to run batch cleanup code

diff --git a/Engine/Source/Programs/Horde/HordeAgent/Leases/Handlers/JobHandler.cs b/Engine/Source/Programs/Horde/HordeAgent/Leases/Handlers/JobHandler.cs
--- a/Engine/Source/Programs/Horde/HordeAgent/Leases/Handlers/JobHandler.cs
+++ b/Engine/Source/Programs/Horde/HordeAgent/Leases/Handlers/JobHandler.cs
@@ -34,63 +34,73 @@ namespace HordeAgent.Leases.Handlers
 			span.SetAttribute("horde.job.batch_id", executeTask.BatchId);
 
 			await using IServerLogger logger = session.HordeClient.CreateServerLogger(LogId.Parse(executeTask.LogId)).WithLocalLogger(localLogger);
+
+			int exitCode = 0;
+			Exception? exception = null;
+
+			executeTask.JobOptions ??= new RpcJobOptions();
+
+			List<string> arguments = new List<string>();
+			arguments.Add("execute");
+			arguments.Add("job");
+			arguments.Add($"-AgentId={session.AgentId}");
+			arguments.Add($"-SessionId={session.SessionId}");
+			arguments.Add($"-LeaseId={leaseId}");
+			arguments.Add($"-WorkingDir={session.WorkingDir}");
+			arguments.Add($"-Task={Convert.ToBase64String(executeTask.ToByteArray())}");
+
+			string driverName = String.IsNullOrEmpty(executeTask.JobOptions.Driver) ? "JobDriver" : executeTask.JobOptions.Driver;
+			FileReference driverAssembly = FileReference.Combine(new DirectoryReference(AppContext.BaseDirectory), driverName, $"{driverName}.dll");			
+			span.SetAttribute("horde.job.driver_name", driverName);
+
+			Dictionary<string, string> environment = ManagedProcess.GetCurrentEnvVars();
+
 			try
 			{
-				executeTask.JobOptions ??= new RpcJobOptions();
-
-				List<string> arguments = new List<string>();
-				arguments.Add("execute");
-				arguments.Add("job");
-				arguments.Add($"-AgentId={session.AgentId}");
-				arguments.Add($"-SessionId={session.SessionId}");
-				arguments.Add($"-LeaseId={leaseId}");
-				arguments.Add($"-WorkingDir={session.WorkingDir}");
-				arguments.Add($"-Task={Convert.ToBase64String(executeTask.ToByteArray())}");
-
-				string driverName = String.IsNullOrEmpty(executeTask.JobOptions.Driver) ? "JobDriver" : executeTask.JobOptions.Driver;
-				FileReference driverAssembly = FileReference.Combine(new DirectoryReference(AppContext.BaseDirectory), driverName, $"{driverName}.dll");
-				span.SetAttribute("horde.job.driver_name", driverName);
-
-				Dictionary<string, string> environment = ManagedProcess.GetCurrentEnvVars();
 				environment[HordeHttpClient.HordeUrlEnvVarName] = session.HordeClient.ServerUrl.ToString();
 				environment[HordeHttpClient.HordeTokenEnvVarName] = executeTask.Token;
 				environment["UE_LOG_JSON_TO_STDOUT"] = "1";
 				environment["UE_HORDE_OTEL_SETTINGS"] = OpenTelemetrySettingsExtensions.Serialize(_settings.CurrentValue.OpenTelemetry, true);
 
-				int exitCode = await RunDotNetProcessAsync(driverAssembly, arguments, environment, AgentApp.IsSelfContained, logger, cancellationToken);
-				logger.LogInformation("Driver finished with exit code {ExitCode}", exitCode);
-
-				if (exitCode != 0)
-				{
-					try
-					{
-						logger.LogInformation("Driver initiating job batch cleanup");
-						// Run job batch cleanup
-						environment["UE_HORDE_JOB_BATCH_CLEANUP"] = "1";
-						int cleanupCode = await RunDotNetProcessAsync(driverAssembly, arguments, environment, AgentApp.IsSelfContained, logger, cancellationToken);
-						if (cleanupCode != 0)
-						{
-							logger.LogInformation("Driver job batch cleanup finished with exit code {ExitCode}", cleanupCode);
-						}
-					}
-					catch (Exception ex)
-					{
-						logger.LogError(ex, "Driver job batch cleanup error ({Message})", ex.Message);
-					}
-				}
-
-				return (exitCode == 0) ? LeaseResult.Success : LeaseResult.Failed;
+				exitCode = await RunDotNetProcessAsync(driverAssembly, arguments, environment, AgentApp.IsSelfContained, logger, cancellationToken);
+				logger.LogInformation("Driver finished with exit code {ExitCode}", exitCode);				
 			}
 			catch (OperationCanceledException ex)
 			{
 				logger.LogError(ex, "Lease was cancelled ({Reason})", CancellationReason);
-				throw;
+				exception = ex;
 			}
 			catch (Exception ex)
 			{
 				logger.LogError(ex, "Unhandled exception: {Message}", ex.Message);
-				throw;
+				exception = ex;
 			}
+
+			if (exception != null || exitCode != 0)
+			{
+				try
+				{
+					logger.LogInformation("Driver initiating job batch cleanup");
+					// Run job batch cleanup
+					environment["UE_HORDE_JOB_BATCH_CLEANUP"] = "1";
+					int cleanupCode = await RunDotNetProcessAsync(driverAssembly, arguments, environment, AgentApp.IsSelfContained, logger, cancellationToken);
+					if (cleanupCode != 0)
+					{
+						logger.LogInformation("Driver job batch cleanup finished with exit code {ExitCode}", cleanupCode);
+					}
+				}
+				catch (Exception ex)
+				{
+					logger.LogError(ex, "Driver job batch cleanup error ({Message})", ex.Message);
+				}
+
+				if (exception != null)
+				{
+					throw exception;
+				}
+			}
+
+			return (exitCode == 0) ? LeaseResult.Success : LeaseResult.Failed;
 		}
 	}
 
