Horde: Add breadcrumbs to templates main menu, to indicate all or partial selection in submenus

diff --git a/Engine/Source/Programs/Horde/HordeDashboard/src/components/JobSearchSimple.tsx b/Engine/Source/Programs/Horde/HordeDashboard/src/components/JobSearchSimple.tsx
--- a/Engine/Source/Programs/Horde/HordeDashboard/src/components/JobSearchSimple.tsx
+++ b/Engine/Source/Programs/Horde/HordeDashboard/src/components/JobSearchSimple.tsx
@@ -169,7 +169,12 @@ export const JobSearchSimpleModal: React.FC<{ streamIdIn: string, changeListIn?:
          });
       }
 
-      templateOptions.push({ key: `${cat}_category`, text: cat, subMenuProps: { items: subItems } });
+      templateOptions.push({ 
+         key: `${cat}_category`,
+         text: cat,
+         iconProps: { iconName: allCatTemplatesSelected ? 'Tick' : someCatTemplatesSelected ? 'Form' : '', style: { fontSize: '14px' } },
+         subMenuProps: { items: subItems }
+      });
 
    })
 
@@ -196,7 +201,7 @@ export const JobSearchSimpleModal: React.FC<{ streamIdIn: string, changeListIn?:
             return false;
          }} >
          <Checkbox indeterminate={someTemplatesSelected} checked={allTemplatesSelected} />
-         <Text style={{ padding: "6px 0px"}}>Select All</Text></Stack>
+         <Text style={{ padding: "6px 2px"}}>Select All</Text></Stack>
       }
    });
 
diff --git a/Engine/Source/Programs/Horde/HordeDashboard/src/styles/Styles.tsx b/Engine/Source/Programs/Horde/HordeDashboard/src/styles/Styles.tsx
--- a/Engine/Source/Programs/Horde/HordeDashboard/src/styles/Styles.tsx
+++ b/Engine/Source/Programs/Horde/HordeDashboard/src/styles/Styles.tsx
@@ -108,7 +108,8 @@ registerIcons({
       "Maximize" : IconContents.MAXIMIZE,
       "Star": IconContents.STAR,
       "Clean": IconContents.CLEAN,
-      "Info" : IconContents.INFO_SIGN
+      "Info" : IconContents.INFO_SIGN,
+      "Form": IconContents.FORM
    }
 });
 
