Horde: Add specialized status message support and render as Markdown

diff --git a/Engine/Source/Programs/Horde/HordeDashboard/src/components/Breadcrumbs.tsx b/Engine/Source/Programs/Horde/HordeDashboard/src/components/Breadcrumbs.tsx
--- a/Engine/Source/Programs/Horde/HordeDashboard/src/components/Breadcrumbs.tsx
+++ b/Engine/Source/Programs/Horde/HordeDashboard/src/components/Breadcrumbs.tsx
@@ -5,6 +5,7 @@ import { observer } from 'mobx-react-lite';
 import React, { useState } from 'react';
 import { Link } from 'react-router-dom';
 import dashboard from '../backend/Dashboard';
+import { Markdown } from "../base/components/Markdown";
 import notices from '../backend/Notices';
 import { getHordeTheme } from '../styles/theme';
 import { getHordeStyling } from '../styles/Styles';
@@ -131,7 +132,7 @@ export const Breadcrumbs: React.FC<{ items: BreadcrumbItem[], title?: string, su
                <Stack style={{pointerEvents: "auto"}}>
                   <MessageBar onDismiss={() => { dashboard.alertSquelch = alert; setHideAlert(true) }}
                      messageBarType={MessageBarType.severeWarning} isMultiline={false} >
-                     <Text variant={"small"} style={{ fontFamily: "Horde Open Sans Bold"}}>{alert}</Text>
+                     <Markdown styles={{ root: { fontSize: 12, fontFamily: "Horde Open Sans Bold" } }}>{alert}</Markdown>
                   </MessageBar>
                </Stack>
                <Stack grow />
diff --git a/Engine/Source/Programs/Horde/HordeServer/Server/Notices/NoticesController.cs b/Engine/Source/Programs/Horde/HordeServer/Server/Notices/NoticesController.cs
--- a/Engine/Source/Programs/Horde/HordeServer/Server/Notices/NoticesController.cs
+++ b/Engine/Source/Programs/Horde/HordeServer/Server/Notices/NoticesController.cs
@@ -143,7 +143,18 @@ namespace HordeServer.Server.Notices
 					{
 						if (update.Result != HealthStatus.Healthy)
 						{
-							string message = $"Server is reporting a {status.Name} issue. See status page for more info.";
+							string message = $"Server is reporting a {status.Name} issue.";
+
+							// Add links to relevant status pages within Horde
+							if (status.Name == "Perforce")
+							{
+								message += " See[Perforce servers page](/perforce/servers) for more info.";
+							}
+							else
+							{
+								message += " See[status page](/serverstatus) for more info.";
+							}
+
 							messages.Add(new GetNoticeResponse { Active = true, Message = message });
 						}
 					}
