Add StepId link to Audit log
+ Allow multiline message to display as intended

#rnx

diff --git a/Engine/Source/Programs/Horde/HordeDashboard/src/components/AuditLog.tsx b/Engine/Source/Programs/Horde/HordeDashboard/src/components/AuditLog.tsx
--- a/Engine/Source/Programs/Horde/HordeDashboard/src/components/AuditLog.tsx
+++ b/Engine/Source/Programs/Horde/HordeDashboard/src/components/AuditLog.tsx
@@ -38,7 +38,15 @@ const getAuditStyleNormal = () => {
                '&:hover': { background: theme.palette.neutralLight }
             },
          }
-      ]
+      ],
+      container: {
+         selectors: {
+            '.ms-List-cell': {
+               height: "auto",
+            }
+         }
+      }
+
    });
 
    _auditStyleNormal = auditStyleNormal;
@@ -588,7 +596,7 @@ export const AuditLogPanel: React.FC<{ agentId?: string, deviceId?: string, jobI
       let timestamp = `[${tm.format(format)}]`;
 
       return (
-         <Stack className={auditStyleNormal.logLine} key={`key_log_line_${item?.entry.time}`} style={{ width: "100%", height: logMetricNormal.lineHeight }}>
+         <Stack className={auditStyleNormal.logLine} key={`key_log_line_${item?.entry.time}`} style={{ width: "100%" }}>
             <div style={{ position: "relative" }}>
                <Stack tokens={{ childrenGap: 8 }} horizontal disableShrink={true}>
                   <Stack horizontal disableShrink={true} >
@@ -965,6 +973,19 @@ class AuditLine {
 
          }
 
+         // StepId
+         if (tag === "StepId") {
+
+            const stepId = property.value as string;
+            const jobId = this.properties.get("JobId")?.value;
+
+            if (stepId && jobId) {
+               let to = `/job/${jobId}?step=${stepId}`;
+               element = <Link to={to} >{this.highlight(key, stepId)}</Link>
+            }
+
+         }
+
          // BatchId
          if (tag === "BatchId") {
 
@@ -990,6 +1011,14 @@ class AuditLine {
 
          }
 
+         // Message
+         if (tag === "Message") {
+            const message = property.value?.toString();
+            if (message && message.includes("\n")) {
+               element = <pre style={{margin: 0}}>{this.highlight(key, message)}</pre>;
+            }
+         }
+
          if (element) {
             return element;
          }
