Add Device Audit link from step log (using structure logging)

#rnx

diff --git a/Engine/Source/Programs/Horde/HordeDashboard/src/components/LogRender.tsx b/Engine/Source/Programs/Horde/HordeDashboard/src/components/LogRender.tsx
--- a/Engine/Source/Programs/Horde/HordeDashboard/src/components/LogRender.tsx
+++ b/Engine/Source/Programs/Horde/HordeDashboard/src/components/LogRender.tsx
@@ -18,6 +18,7 @@ enum TagType {
    LeaseId,
    ProjectId,
    StreamId,
+   DeviceId,
 }
 
 export type LogItem = {
@@ -182,6 +183,10 @@ const renderTags = (navigate: NavigateFunction, line: LogLine, lineNumber: numbe
             tagType = TagType.ArtifactId;
          }
 
+         if (type === "DeviceId") {
+            tagType = TagType.DeviceId
+         }
+
       }
 
       if (tagType === TagType.None || !record) {
@@ -203,6 +208,9 @@ const renderTags = (navigate: NavigateFunction, line: LogLine, lineNumber: numbe
       } else if (tagType === TagType.PoolId) {
          const url = `/pools?poolid=${text}`
          return <a key={key} href={url} onClick={(ev) => { ev.stopPropagation(); ev.preventDefault(); navigate(url) }}><Highlight search={search ? search : ""} className={logStyle.logLine}>{text}</Highlight></a>;
+      } else if (tagType === TagType.DeviceId) {
+         const url = `/audit/device/${text}`
+         return <a key={key} href={url} onClick={(ev) => { ev.stopPropagation(); ev.preventDefault(); navigate(url) }}><Highlight search={search ? search : ""} className={logStyle.logLine}>{text}</Highlight></a>;
       } else if (tagType === TagType.LeaseId) {
 
          const navigateToLeaseLog = async (toplevel: boolean) => {
