Horde: Do not forward cancellation token to cleanup job driver as this will interrupt it

diff --git a/Engine/Source/Programs/Horde/HordeAgent/Leases/Handlers/JobHandler.cs b/Engine/Source/Programs/Horde/HordeAgent/Leases/Handlers/JobHandler.cs
--- a/Engine/Source/Programs/Horde/HordeAgent/Leases/Handlers/JobHandler.cs
+++ b/Engine/Source/Programs/Horde/HordeAgent/Leases/Handlers/JobHandler.cs
@@ -49,7 +49,7 @@ namespace HordeAgent.Leases.Handlers
 			arguments.Add($"-WorkingDir={session.WorkingDir}");
 			arguments.Add($"-Task={Convert.ToBase64String(executeTask.ToByteArray())}");
 
-			string driverName = String.IsNullOrEmpty(executeTask.JobOptions.Driver) ? "JobDriver" : executeTask.JobOptions.Driver;
+			string driverName = String.IsNullOrEmpty(executeTask.JobOptions.Driver) ? "JobDriver" : executeTask.JobOptions.Driver;			
 			FileReference driverAssembly = FileReference.Combine(new DirectoryReference(AppContext.BaseDirectory), driverName, $"{driverName}.dll");			
 			span.SetAttribute("horde.job.driver_name", driverName);
 
@@ -83,7 +83,7 @@ namespace HordeAgent.Leases.Handlers
 					logger.LogInformation("Driver initiating job batch cleanup");
 					// Run job batch cleanup
 					environment["UE_HORDE_JOB_BATCH_CLEANUP"] = "1";
-					int cleanupCode = await RunDotNetProcessAsync(driverAssembly, arguments, environment, AgentApp.IsSelfContained, logger, cancellationToken);
+					int cleanupCode = await RunDotNetProcessAsync(driverAssembly, arguments, environment, AgentApp.IsSelfContained, logger, CancellationToken.None);
 					if (cleanupCode != 0)
 					{
 						logger.LogInformation("Driver job batch cleanup finished with exit code {ExitCode}", cleanupCode);
