Add minutes to device pool distribution time

#rnx

diff --git a/Engine/Source/Programs/Horde/HordeDashboard/src/components/devices/DeviceTelemetryGraphs.tsx b/Engine/Source/Programs/Horde/HordeDashboard/src/components/devices/DeviceTelemetryGraphs.tsx
--- a/Engine/Source/Programs/Horde/HordeDashboard/src/components/devices/DeviceTelemetryGraphs.tsx
+++ b/Engine/Source/Programs/Horde/HordeDashboard/src/components/devices/DeviceTelemetryGraphs.tsx
@@ -755,7 +755,7 @@ export class PoolDistributionGraph {
          .attr("dy", "1em")
          .attr("fill", modeColors.text)
          .style("text-anchor", "middle")
-         .text(msecToElapsed((root.data as any).durationMS, false, false));
+         .text(msecToElapsed((root.data as any).durationMS, true, false));
 
       // labels
       const labels = svg.append('g')
@@ -804,7 +804,7 @@ export class PoolDistributionGraph {
          .enter()
          .append("text")
          .attr("class", (d: any) => `arcStepLabel ${d.data.id}`)
-         .text((d: any) => msecToElapsed((d.data as any).durationMS, false, false))
+         .text((d: any) => msecToElapsed((d.data as any).durationMS, true, false))
          .style("fill", modeColors.text)
          .attr("font-weight", "bold")
          .attr("font-size", "smaller")
