Horde: Handle AggregateException when silencing cancellations in RpcService.UpdateSession

Cancellations occur naturally in this method, most likely due to load balancers or reverse proxies aborting requests as a result of timeouts.

diff --git a/Engine/Source/Programs/Horde/Plugins/Compute/HordeServer.Compute/Server/RpcService.cs b/Engine/Source/Programs/Horde/Plugins/Compute/HordeServer.Compute/Server/RpcService.cs
--- a/Engine/Source/Programs/Horde/Plugins/Compute/HordeServer.Compute/Server/RpcService.cs
+++ b/Engine/Source/Programs/Horde/Plugins/Compute/HordeServer.Compute/Server/RpcService.cs
@@ -341,9 +341,10 @@ namespace HordeServer.Server
 					{
 						await nextRequestTask;
 					}
-					catch (OperationCanceledException)
+					catch (Exception ex) when (IsCancellationException(ex))
 					{
-						// Ignore
+						// Ignore cancellation exceptions
+						// A more advanced check is used as it's sometimes wrapped as an AggregateException
 					}
 				}
 			}
@@ -411,5 +412,15 @@ namespace HordeServer.Server
 
 			return Task.FromResult(new Empty());
 		}
+		
+		private static bool IsCancellationException(Exception ex)
+		{
+			return ex switch
+			{
+				OperationCanceledException => true,
+				AggregateException aggEx => aggEx.InnerExceptions.All(IsCancellationException),
+				_ => ex.InnerException != null && IsCancellationException(ex.InnerException)
+			};
+		}
 	}
 }
