Horde: Update release notes
#rnx

diff --git a/Engine/Source/Programs/Horde/Docs/ReleaseNotes.md b/Engine/Source/Programs/Horde/Docs/ReleaseNotes.md
--- a/Engine/Source/Programs/Horde/Docs/ReleaseNotes.md
+++ b/Engine/Source/Programs/Horde/Docs/ReleaseNotes.md
@@ -1,5 +1,16 @@
 # Release Notes
 
+## 2025-02-06
+
+* Do not forward cancellation token to cleanup job driver as this will interrupt it (39755347)
+* Fix Linux runs. (39727551)
+* Add specialized status message support and render as Markdown (39721461)
+* Add breadcrumbs to templates main menu, to indicate all or partial selection in submenus (39633740)
+* Refactor job handler lease exceptions to run batch cleanup code (39633282)
+* Add UX to log view to skip to next block of issue lines (39633254)
+* make C# code compilation not require a .csproj in the project being tested.
+  - UAT will now also recurse in Builds directories of plugins (if they happen to have them). (39601473)
+
 ## 2025-01-30
 
 * Validate ServerUrl for auth modes and load balancers in settings
