Add Usage peak and saturation duration to device pool health graph

#rnx

diff --git a/Engine/Source/Programs/Horde/HordeDashboard/src/components/devices/DeviceHealthWidgets.tsx b/Engine/Source/Programs/Horde/HordeDashboard/src/components/devices/DeviceHealthWidgets.tsx
--- a/Engine/Source/Programs/Horde/HordeDashboard/src/components/devices/DeviceHealthWidgets.tsx
+++ b/Engine/Source/Programs/Horde/HordeDashboard/src/components/devices/DeviceHealthWidgets.tsx
@@ -1,9 +1,9 @@
-import { Stack, Label, Text, TooltipHost, DirectionalHint, IGroup } from "@fluentui/react";
+import { Stack, Label, Text, TooltipHost, ITooltipProps, DirectionalHint, IGroup } from "@fluentui/react";
 import { useState, useRef } from "react";
 import dashboard, { StatusColor } from "../../backend/Dashboard";
 import { getHordeStyling } from "../../styles/Styles";
 import { DeviceHandler, DeviceHealthStatus, DeviceHealthLabels, FilteredDeviceTelemetry } from "../../backend/DeviceHandler";
-import { getHumanTime } from "../../base/utilities/timeUtils";
+import { getHumanTime, msecToElapsed } from "../../base/utilities/timeUtils";
 
 import * as d3 from "d3";
 
@@ -95,6 +95,9 @@ export class PoolTimelineGraph {
       const dashboardColors = dashboard.getStatusColors();
 
       const maxReservations = data.reduce((a, d) => Math.max(a, d.reservations + d.available), 1);
+      const peakReservations = data.reduce((a, d) => Math.max(a, d.reservations), 0);
+      this.peakReservationsPerc = Math.ceil(peakReservations / maxReservations * 100);
+      this.totalSaturatedReservation = peakReservations == maxReservations? data.reduce((a, d) => a + (d.reservations + d.problems == maxReservations? 1 : 0), 0) : 0;
 
       const svg = d3.select(container)
          .append("svg")
@@ -151,6 +154,25 @@ export class PoolTimelineGraph {
          .attr("fill", colorProblem)
          .attr("d", areaProblem as any)
          .style("pointer-events", "none");
+
+      if (maxReservations > 1)
+      {
+         // saturation
+         const lineSaturation = d3.line()
+            .defined((d: any) => d.reservations == maxReservations)
+            .x((d: any) => xPlot(d.date))
+            .y((d: any) => yPlot(d.reservations))
+            .curve(d3.curveStep);
+
+         svg.append("path")
+            .datum(data)
+            .attr("fill", "transparent")
+            .attr("stroke", colorProblem).attr("stroke-width", 4)
+            .attr("opacity", 0.6)
+            .attr("d", lineSaturation as any)
+            .style("pointer-events", "none");
+      }
+
    }
 
    private clear() {
@@ -196,6 +218,9 @@ export class PoolTimelineGraph {
 
    private hasRendered = false;
    forceRender = false;
+
+   peakReservationsPerc: number = 0;
+   totalSaturatedReservation: number = 0;
 }
 
 export const DeviceGroupHeader: React.FC<{ group: IGroup, handler: DeviceHandler, devices?: string[] }> = ({ group, handler, devices }) => {
@@ -207,11 +232,20 @@ export const DeviceGroupHeader: React.FC<{ group: IGroup, handler: DeviceHandler
       graph.current.draw(container, 300, 38, devices);
    }
 
+   const tooltipProps: ITooltipProps = {
+      onRenderContent: () => (
+         <ul style={{margin: 0, padding: 0}}>
+            <li>Peak usage: {graph.current.peakReservationsPerc}%</li>
+            {graph.current.totalSaturatedReservation > 0 && <li>Saturated for: {msecToElapsed(graph.current.totalSaturatedReservation * 10 * 60 * 1000, true, false)}</li>}
+         </ul>
+      )
+   }
+
    const backgroundColor = (dashboard.darktheme? d3.gray(10) : d3.gray(80)).toString();
 
    const title = `${group.name} (${group.count})`;
 
-   return <TooltipHost content={`Device usage < ${handler.telemetryHoursRange}h`} delay={0} closeDelay={250}
+   return <TooltipHost tooltipProps={tooltipProps} delay={0} closeDelay={250}
          styles={{root: {display: 'inline-block'}}} directionalHint={DirectionalHint.rightCenter}> 
       <Stack style={{ position: "relative" }} horizontal verticalAlign="center">
          <Text style={{width: 120, fontWeight: 'bold'}}>{title}</Text>
