Horde: Fix incorrect IMDSv2 detection in agent

The agent was failing to properly detect when IMDSv2 was required on AWS instances, leading to spot interruptions not being acted on when IMDSv2 is used.

diff --git a/Engine/Source/Programs/Horde/HordeAgent.Tests/Services/AwsInstanceLifecycleServiceTests.cs b/Engine/Source/Programs/Horde/HordeAgent.Tests/Services/AwsInstanceLifecycleServiceTests.cs
--- a/Engine/Source/Programs/Horde/HordeAgent.Tests/Services/AwsInstanceLifecycleServiceTests.cs
+++ b/Engine/Source/Programs/Horde/HordeAgent.Tests/Services/AwsInstanceLifecycleServiceTests.cs
@@ -52,7 +52,14 @@ public class FakeAwsImds : HttpMessageHandler
 	public FakeAwsImds(IClock clock)
 	{
 		_clock = clock;
-		_responses["/"] = req => Res(HttpStatusCode.OK, "Metadata server");
+		_responses["/"] = req =>
+		{
+			if (!IsValidImdsV2Request(req))
+			{
+				return new HttpResponseMessage(HttpStatusCode.Unauthorized);
+			}
+			return Res(HttpStatusCode.OK, "Metadata server");
+		};
 		_responses[AwsInstanceLifecycleService.ImdsTokenPath] = req =>
 		{
 			if (req.Method != HttpMethod.Put)
diff --git a/Engine/Source/Programs/Horde/HordeAgent/Services/AwsInstanceLifecycleService.cs b/Engine/Source/Programs/Horde/HordeAgent/Services/AwsInstanceLifecycleService.cs
--- a/Engine/Source/Programs/Horde/HordeAgent/Services/AwsInstanceLifecycleService.cs
+++ b/Engine/Source/Programs/Horde/HordeAgent/Services/AwsInstanceLifecycleService.cs
@@ -56,8 +56,8 @@ class AwsInstanceLifecycleService : BackgroundService
 	/// Name of the HTTP client used for requests to IMDS
 	/// </summary>
 	public const string HttpClientName = "Horde.HttpAwsInstanceClient";
-
 	
+	public const string ImdsBaseUri = "http://169.254.169.254";
 	public const string ImdsTokenPath = "/latest/api/token";
 	public const string ImdsInstanceLifeCyclePath = "/latest/meta-data/instance-life-cycle";
 	public const string ImdsSpotActionPath = "/latest/meta-data/spot/instance-action";
@@ -66,7 +66,6 @@ class AwsInstanceLifecycleService : BackgroundService
 	public const string ImdsTokenTtlHeader = "X-aws-ec2-metadata-token-ttl-seconds";
 	public const string ImdsTokenTtl = "21600"; // 6 hours
 	
-	private const string BaseUri = "http://169.254.169.254"; // TODO: remove prefix
 	private readonly IWorkerService _workerService;
 	private readonly HttpClient _httpClient;
 	private readonly IClock _clock;
@@ -119,7 +118,7 @@ class AwsInstanceLifecycleService : BackgroundService
 		if (_imdsToken == null || _clock.UtcNow >= _imdsTokenExpireTime)
 		{
 			// Get a token
-			using HttpRequestMessage tokenReq = new (HttpMethod.Put, new Uri($"{BaseUri}{ImdsTokenPath}"));
+			using HttpRequestMessage tokenReq = new (HttpMethod.Put, new Uri($"{ImdsBaseUri}{ImdsTokenPath}"));
 			tokenReq.Headers.Add(ImdsTokenTtlHeader, ImdsTokenTtl);
 			using HttpResponseMessage tokenRes = await _httpClient.SendAsync(tokenReq, cancellationToken);
 			tokenRes.EnsureSuccessStatusCode();
@@ -128,7 +127,7 @@ class AwsInstanceLifecycleService : BackgroundService
 		}
 		
 		// Perform actual metadata request
-		using HttpRequestMessage request = new (HttpMethod.Get, $"{BaseUri}{path}");
+		using HttpRequestMessage request = new (HttpMethod.Get, $"{ImdsBaseUri}{path}");
 		request.Headers.Add(ImdsTokenHeader, _imdsToken);
 		return await _httpClient.SendAsync(request, cancellationToken);
 	}
@@ -174,8 +173,8 @@ class AwsInstanceLifecycleService : BackgroundService
 	{
 		try
 		{
-			HttpResponseMessage res = await _httpClient.GetAsync(new Uri(BaseUri + "/"), cancellationToken);
-			return res.StatusCode == HttpStatusCode.OK;
+			HttpResponseMessage res = await _httpClient.GetAsync(new Uri(ImdsBaseUri + "/"), cancellationToken);
+			return res.StatusCode is HttpStatusCode.OK or HttpStatusCode.Unauthorized;
 		}
 		catch (Exception)
 		{
