Horde: Automation summaries, use latest test data fime for a CL

diff --git a/Engine/Source/Programs/Horde/HordeDashboard/src/backend/AutomationTestData.ts b/Engine/Source/Programs/Horde/HordeDashboard/src/backend/AutomationTestData.ts
--- a/Engine/Source/Programs/Horde/HordeDashboard/src/backend/AutomationTestData.ts
+++ b/Engine/Source/Programs/Horde/HordeDashboard/src/backend/AutomationTestData.ts
@@ -980,7 +980,8 @@ export class TestDataHandler {
 
             const date = this.changeDates.get(r.buildChangeList);
             const rdate = this.idDates.get(r.id)!;
-            if (!date || date.getTime() > rdate.getTime()) {
+            // Use the latest date for summary, otherwise gets confusing when a CL spans days
+            if (!date || date.getTime() < rdate.getTime()) {
                 this.changeDates.set(r.buildChangeList, rdate);
             }
 
