Horde: AutomationHub, use ref time instead of CL time in summaries when available

diff --git a/Engine/Source/Programs/Horde/HordeDashboard/src/components/AutomationViewSummary.tsx b/Engine/Source/Programs/Horde/HordeDashboard/src/components/AutomationViewSummary.tsx
--- a/Engine/Source/Programs/Horde/HordeDashboard/src/components/AutomationViewSummary.tsx
+++ b/Engine/Source/Programs/Horde/HordeDashboard/src/components/AutomationViewSummary.tsx
@@ -227,7 +227,12 @@ const TestSummaryButton: React.FC<{ test: GetTestResponse, handler: TestDataHand
                }
 
                let dateString = "";
-               const date = handler.changeDates.get(last.buildChangeList);
+
+               let date = handler.idDates.get(last.id);
+               
+               if (!date) {
+                  date = handler.changeDates.get(last.buildChangeList);
+               }               
                if (date) {
                   dateString = getShortNiceTime(date, false, false)
                }
@@ -699,7 +704,12 @@ const SuiteSummaryButton: React.FC<{ suite: GetTestSuiteResponse, handler: TestD
                const metaName = `${elements.join(" / ")}`;
 
                let dateString = "";
-               const date = handler.changeDates.get(last.buildChangeList);
+               
+               let date = handler.idDates.get(last.id);
+               if (!date) {
+                  date = handler.changeDates.get(last.buildChangeList);
+               }               
+
                if (date) {
                   dateString = getShortNiceTime(date, false, false)
                }
