Horde: Handle gRPC client disconnects more gracefully for log tail streams

ConnectionAbortedException may get generated in such cases and is not caught.

diff --git a/Engine/Source/Programs/Horde/Plugins/Compute/HordeServer.Compute/Logs/LogRpcService.cs b/Engine/Source/Programs/Horde/Plugins/Compute/HordeServer.Compute/Logs/LogRpcService.cs
--- a/Engine/Source/Programs/Horde/Plugins/Compute/HordeServer.Compute/Logs/LogRpcService.cs
+++ b/Engine/Source/Programs/Horde/Plugins/Compute/HordeServer.Compute/Logs/LogRpcService.cs
@@ -74,9 +74,8 @@ namespace HordeServer.Logs
 		/// <inheritdoc/>
 		public override async Task UpdateLogTail(IAsyncStreamReader<RpcUpdateLogTailRequest> requestStream, IServerStreamWriter<RpcUpdateLogTailResponse> responseStream, ServerCallContext context)
 		{
-			RpcUpdateLogTailResponse response = new RpcUpdateLogTailResponse();
-			response.TailNext = -1;
-
+			RpcUpdateLogTailResponse response = new () { TailNext = -1 };
+			
 			Task<bool> moveNextTask = requestStream.MoveNext();
 			while (await moveNextTask)
 			{
@@ -94,24 +93,25 @@ namespace HordeServer.Logs
 				response.TailNext = await _logTailService.GetTailNextAsync(logId, context.CancellationToken);
 				if (response.TailNext == -1)
 				{
-					using (CancellationTokenSource cancellationSource = new CancellationTokenSource())
+					using CancellationTokenSource cancellationSource = new ();
+					_logger.LogDebug("Waiting for tail next on log {LogId}", logId);
+					Task<int> waitTask = _logTailService.WaitForTailNextAsync(logId, cancellationSource.Token);
+					Task completeTask = await Task.WhenAny(waitTask, moveNextTask);
+					await cancellationSource.CancelAsync();
+
+					try
+					{
+						response.TailNext = await waitTask;
+						await completeTask; // Allow moveNextTask to throw if it finished due to cancellation, so we don't write to
+					}
+					catch (OperationCanceledException)
+					{
+						_logger.LogDebug("Log tail stream cancelled by client");
+						break;
+					}
+					catch (Exception ex)
 					{
-						_logger.LogDebug("Waiting for tail next on log {LogId}", logId);
-						Task<int> waitTask = _logTailService.WaitForTailNextAsync(logId, cancellationSource.Token);
-
-						Task completeTask = await Task.WhenAny(waitTask, moveNextTask);
-						await cancellationSource.CancelAsync();
-
-						try
-						{
-							response.TailNext = await waitTask;
-						}
-						catch (Exception ex)
-						{
-							_logger.LogWarning(ex, "Exception while waiting for tail next");
-						}
-
-						await completeTask; // Allow moveNextTask to throw if it finished due to cancellation, so we don't write to 
+						_logger.LogWarning(ex, "Exception while waiting for tail next");
 					}
 				}
 
