Horde: Return explicit 500 status for VCS operation failures in create job REST endpoint

- To prevent excessive logging at the error level when failure is potentially caused by misconfigured streams or templates.
- Raise exceptions in PerforceService as appropriate class
- Wrap such exceptions in CommitCollectionException for operations from ICommitCollection

diff --git a/Engine/Source/Programs/Horde/Plugins/Build/HordeServer.Build/Jobs/JobsController.cs b/Engine/Source/Programs/Horde/Plugins/Build/HordeServer.Build/Jobs/JobsController.cs
--- a/Engine/Source/Programs/Horde/Plugins/Build/HordeServer.Build/Jobs/JobsController.cs
+++ b/Engine/Source/Programs/Horde/Plugins/Build/HordeServer.Build/Jobs/JobsController.cs
@@ -24,6 +24,7 @@ using HordeServer.Streams;
 using HordeServer.Users;
 using HordeServer.Utilities;
 using Microsoft.AspNetCore.Authorization;
+using Microsoft.AspNetCore.Http;
 using Microsoft.AspNetCore.Mvc;
 using Microsoft.Extensions.Logging;
 using Microsoft.Extensions.Options;
@@ -168,13 +169,26 @@ namespace HordeServer.Jobs
 
 			// Get the commits for this stream
 			ICommitCollection commits = _commitService.GetCollection(streamConfig);
-
-			// Get the change to build
-			CommitIdWithOrder commitId = await GetChangeToBuildAsync(create, streamConfig.Id, template, shelfInfo, commits, HttpContext.RequestAborted);
-
-			// And get the matching code changelist
-			ICommit? lastCodeCommit = await commits.GetLastCodeChangeAsync(commitId, HttpContext.RequestAborted);
-			CommitIdWithOrder? codeCommitId = lastCodeCommit?.Id;
+			
+			CommitIdWithOrder commitId;
+			CommitIdWithOrder? codeCommitId;
+			try
+			{
+				// Get the change to build
+				commitId = await GetChangeToBuildAsync(create, streamConfig.Id, template, shelfInfo, commits, HttpContext.RequestAborted);
+				
+				// And get the matching code changelist
+				ICommit? lastCodeCommit = await commits.GetLastCodeChangeAsync(commitId, HttpContext.RequestAborted);
+				codeCommitId = lastCodeCommit?.Id;
+			}
+			catch (CommitCollectionException e)
+			{
+				// Prevent exceptions from bubbling up when VCS operations fail due to misconfigured streams or templates.
+				// Log these at warning level since they're often configuration issues rather than true system failures.
+				// While some failures may be due to user input or configuration, we conservatively return 500 to indicate service unavailability.
+				_logger.LogWarning(e, "Error performing VCS operation(s): {Message}", e.Message);
+				return StatusCode(StatusCodes.Status500InternalServerError, e.Message);
+			}
 
 			// New properties for the job
 			bool? updateIssues = null;
diff --git a/Engine/Source/Programs/Horde/Plugins/Build/HordeServer.Build/Perforce/PerforceService.cs b/Engine/Source/Programs/Horde/Plugins/Build/HordeServer.Build/Perforce/PerforceService.cs
--- a/Engine/Source/Programs/Horde/Plugins/Build/HordeServer.Build/Perforce/PerforceService.cs
+++ b/Engine/Source/Programs/Horde/Plugins/Build/HordeServer.Build/Perforce/PerforceService.cs
@@ -22,6 +22,11 @@ using OpenTelemetry.Trace;
 
 namespace HordeServer.Perforce
 {
+	/// <summary>
+	/// Exception raised by errors in <see cref="PerforceService" />
+	/// </summary>
+	public class PerforceServiceException(string? message) : Exception(message);
+
 	/// <summary>
 	/// P4API implementation of the Perforce service
 	/// </summary>
@@ -398,7 +403,7 @@ namespace HordeServer.Perforce
 				PerforceCredentials? credentials = cluster.Credentials.FirstOrDefault(x => x.UserName.Equals(cluster.ServiceAccount, StringComparison.OrdinalIgnoreCase));
 				if (credentials == null)
 				{
-					throw new Exception($"No credentials defined for {cluster.ServiceAccount} on {cluster.Name}");
+					throw new PerforceServiceException($"No credentials defined for {cluster.ServiceAccount} on {cluster.Name}");
 				}
 				return new Credentials(credentials.UserName, credentials.Password, credentials.Ticket, null);
 			}
@@ -648,7 +653,7 @@ namespace HordeServer.Perforce
 			IPerforceServer? server = await _loadBalancer.SelectServerAsync(cluster, cancellationToken);
 			if (server == null)
 			{
-				throw new Exception($"Unable to select server from '{cluster.Name}'");
+				throw new PerforceServiceException($"Unable to select server from '{cluster.Name}'");
 			}
 			return server;
 		}
@@ -955,7 +960,7 @@ namespace HordeServer.Perforce
 
 							if (String.IsNullOrEmpty(directoryName))
 							{
-								throw new Exception($"Invalid directory name for local client file, disk file path: {diskFilePath}");
+								throw new PerforceServiceException($"Invalid directory name for local client file, disk file path: {diskFilePath}");
 							}
 
 							// Create the directory
@@ -965,7 +970,7 @@ namespace HordeServer.Perforce
 
 								if (!Directory.Exists(directoryName))
 								{
-									throw new Exception($"Unable to create directrory: {directoryName}");
+									throw new PerforceServiceException($"Unable to create directrory: {directoryName}");
 								}
 							}
 
@@ -979,14 +984,14 @@ namespace HordeServer.Perforce
 
 								if (!File.Exists(diskFilePath))
 								{
-									throw new Exception($"Unable to create local change file: {diskFilePath}");
+									throw new PerforceServiceException($"Unable to create local change file: {diskFilePath}");
 								}
 							}
 
 							List<AddRecord> addFiles = await perforce.AddAsync(-1, workspaceFilePath, cancellationToken);
 							if (addFiles.Count != 1)
 							{
-								throw new Exception($"Unable to add local change file,  local: {diskFilePath} : workspace: {workspaceFilePath}");
+								throw new PerforceServiceException($"Unable to add local change file,  local: {diskFilePath} : workspace: {workspaceFilePath}");
 							}
 
 							depotPath = addFiles[0].DepotFile;
@@ -996,13 +1001,13 @@ namespace HordeServer.Perforce
 							List<SyncRecord> syncResults = await perforce.SyncAsync(SyncOptions.Force, -1, workspaceFilePath, cancellationToken).ToListAsync(cancellationToken);
 							if (syncResults == null || syncResults.Count != 1)
 							{
-								throw new Exception($"Unable to sync file, workspace: {workspaceFilePath}");
+								throw new PerforceServiceException($"Unable to sync file, workspace: {workspaceFilePath}");
 							}
 
 							List<EditRecord> editResults = await perforce.EditAsync(-1, syncResults[0].DepotFile.ToString(), cancellationToken);
 							if (editResults == null || editResults.Count != 1)
 							{
-								throw new Exception($"Unable to edit file, workspace: {workspaceFilePath}");
+								throw new PerforceServiceException($"Unable to edit file, workspace: {workspaceFilePath}");
 							}
 
 							depotPath = editResults[0].DepotFile;
@@ -1010,7 +1015,7 @@ namespace HordeServer.Perforce
 
 						if (String.IsNullOrEmpty(depotPath))
 						{
-							throw new Exception($"Unable to get depot path for: {workspaceFilePath}");
+							throw new PerforceServiceException($"Unable to get depot path for: {workspaceFilePath}");
 						}
 
 						// create a new change
@@ -1191,17 +1196,25 @@ namespace HordeServer.Perforce
 				StreamConfig = streamConfig;
 				Logger = logger;
 			}
-
+			
+			/// <inheritdoc/>
 			public async Task<CommitIdWithOrder> CreateNewAsync(string path, string description, CancellationToken cancellationToken = default)
 			{
-				Match match = Regex.Match(path, @"^(//[^/]+/[^/]+)/(.+)$");
-				if (match.Success)
+				try
 				{
-					return CommitIdWithOrder.FromPerforceChange(await PerforceService.CreateNewChangeAsync(StreamConfig.ClusterName, match.Groups[1].Value, match.Groups[2].Value, description, cancellationToken));
+					Match match = Regex.Match(path, @"^(//[^/]+/[^/]+)/(.+)$");
+					if (match.Success)
+					{
+						return CommitIdWithOrder.FromPerforceChange(await PerforceService.CreateNewChangeAsync(StreamConfig.ClusterName, match.Groups[1].Value, match.Groups[2].Value, description, cancellationToken));
+					}
+					else
+					{
+						return CommitIdWithOrder.FromPerforceChange(await PerforceService.CreateNewChangeAsync(StreamConfig.ClusterName, StreamConfig.Name, path, description, cancellationToken));
+					}
 				}
-				else
+				catch (PerforceServiceException e)
 				{
-					return CommitIdWithOrder.FromPerforceChange(await PerforceService.CreateNewChangeAsync(StreamConfig.ClusterName, StreamConfig.Name, path, description, cancellationToken));
+					throw new CommitCollectionException("Failed creating new VCS change", e);
 				}
 			}
 
diff --git a/Engine/Source/Programs/Shared/EpicGames.Horde/Commits/ICommitCollection.cs b/Engine/Source/Programs/Shared/EpicGames.Horde/Commits/ICommitCollection.cs
--- a/Engine/Source/Programs/Shared/EpicGames.Horde/Commits/ICommitCollection.cs
+++ b/Engine/Source/Programs/Shared/EpicGames.Horde/Commits/ICommitCollection.cs
@@ -9,6 +9,14 @@ using EpicGames.Horde.Jobs.Templates;
 
 namespace EpicGames.Horde.Commits
 {
+	/// <summary>
+	/// Represents errors that occur during commit collection operations.
+	/// This exception is thrown when operations on a <see cref="ICommitCollection"/> fail.
+	///
+	/// Serves as a wrapper for underlying VCS-specific exceptions, providing consistent errors for each VCS implementation
+	/// </summary>
+	public class CommitCollectionException(string? message, Exception? innerException) : Exception(message, innerException);
+	
 	/// <summary>
 	/// VCS abstraction. Provides information about commits to a particular stream.
 	/// </summary>
