Horde: Update release notes
#rnx

diff --git a/Engine/Source/Programs/Horde/Docs/ReleaseNotes.md b/Engine/Source/Programs/Horde/Docs/ReleaseNotes.md
--- a/Engine/Source/Programs/Horde/Docs/ReleaseNotes.md
+++ b/Engine/Source/Programs/Horde/Docs/ReleaseNotes.md
@@ -1,5 +1,14 @@
 # Release Notes
 
+## 2025-02-13
+
+* Return explicit 500 status for VCS operation failures in create job REST endpoint (39941708)
+* Downgrade Slack message failure from error to warning log level (39940773)
+* Handle gRPC client disconnects more gracefully for log tail streams (39940732)
+* AutomationHub, use ref time instead of CL time in summaries when available (39912687)
+* Automation summaries, use latest test data fime for a CL (39911090)
+* Fix incorrect IMDSv2 detection in agent (39903273)
+
 ## 2025-02-06
 
 * Do not forward cancellation token to cleanup job driver as this will interrupt it (39755347)
