Horde: Updating schema documentation

#rnx

diff --git a/Engine/Source/Programs/Horde/Docs/Config/Schema/AclActions.md b/Engine/Source/Programs/Horde/Docs/Config/Schema/AclActions.md
--- a/Engine/Source/Programs/Horde/Docs/Config/Schema/AclActions.md
+++ b/Engine/Source/Programs/Horde/Docs/Config/Schema/AclActions.md
@@ -158,16 +158,6 @@
 | `DeleteAccount` |  |
 | `ViewAccount` |  |
 
-## Agents
-
-| Name | Description |
-| ---- | ----------- |
-| `CreateAgent` | Ability to create an agent. This may be done explicitly, or granted to agents to allow them to self-register. |
-| `UpdateAgent` | Update an agent's name, pools, etc... |
-| `DeleteAgent` | Soft-delete an agent |
-| `ViewAgent` | View an agent |
-| `ListAgents` | List the available agents |
-
 ## Artifacts
 
 | Name | Description |
@@ -186,12 +176,18 @@
 | `UpdateBisectTask` |  |
 | `ViewBisectTask` |  |
 
-## Compute
+## Ddc
 
 | Name | Description |
 | ---- | ----------- |
-| `AddComputeTasks` | User can add tasks to the compute cluster |
-| `GetComputeTasks` | User can get and list tasks from the compute cluster |
+| `DdcReadObject` | General read access to refs / blobs and so on |
+| `DdcWriteObject` | General write access to upload refs / blobs etc |
+| `DdcDeleteObject` | Access to delete blobs / refs etc |
+| `DdcDeleteBucket` | Access to delete a particular bucket |
+| `DdcDeleteNamespace` | Access to delete a whole namespace |
+| `DdcReadTransactionLog` | Access to read the transaction log |
+| `DdcWriteTransactionLog` | Access to write the transaction log |
+| `DdcAdminAction` | Access to perform administrative task |
 
 ## Devices
 
@@ -211,24 +207,6 @@
 | `RetryJobStep` |  |
 | `ViewJob` |  |
 
-## Leases
-
-| Name | Description |
-| ---- | ----------- |
-| `ViewLeases` | View all the leases that an agent has worked on |
-| `ViewLeaseTasks` | View the task data for a lease |
-
-## Logs
-
-| Name | Description |
-| ---- | ----------- |
-| `CreateLog` | Ability to create a log. Implicitly granted to agents. |
-| `UpdateLog` | Ability to update log metadata |
-| `ViewLog` | Ability to view a log contents |
-| `WriteLogData` | Ability to write log data |
-| `CreateEvent` | Ability to create events |
-| `ViewEvent` | Ability to view events |
-
 ## Notices
 
 | Name | Description |
@@ -243,16 +221,6 @@
 | ---- | ----------- |
 | `CreateSubscription` |  |
 
-## Pools
-
-| Name | Description |
-| ---- | ----------- |
-| `CreatePool` | Create a global pool of agents |
-| `UpdatePool` | Modify an agent pool |
-| `DeletePool` | Delete an agent pool |
-| `ViewPool` | Ability to view a pool |
-| `ListPools` | View all the available agent pools |
-
 ## Projects
 
 | Name | Description |
@@ -278,21 +246,6 @@
 | `DeleteAccount` |  |
 | `ViewAccount` |  |
 
-## Sessions
-
-| Name | Description |
-| ---- | ----------- |
-| `CreateSession` | Granted to agents to call CreateSession, which returns a bearer token identifying themselves valid to call UpdateSesssion via gRPC. |
-| `ViewSession` | Allows viewing information about an agent session |
-
-## Software
-
-| Name | Description |
-| ---- | ----------- |
-| `UploadSoftware` | Ability to upload new versions of the agent software |
-| `DownloadSoftware` | Ability to download the agent software |
-| `DeleteSoftware` | Ability to delete agent software |
-
 ## Streams
 
 | Name | Description |
@@ -317,11 +270,11 @@
 
 | Name | Description |
 | ---- | ----------- |
-| `CreateAgent` |  |
-| `UpdateAgent` |  |
-| `DeleteAgent` |  |
-| `ViewAgent` |  |
-| `ListAgents` |  |
+| `CreateAgent` | Ability to create an agent. This may be done explicitly, or granted to agents to allow them to self-register. |
+| `UpdateAgent` | Update an agent's name, pools, etc... |
+| `DeleteAgent` | Soft-delete an agent |
+| `ViewAgent` | View an agent |
+| `ListAgents` | List the available agents |
 
 ## Artifacts
 
@@ -345,21 +298,21 @@
 
 | Name | Description |
 | ---- | ----------- |
-| `AddComputeTasks` |  |
-| `GetComputeTasks` |  |
+| `AddComputeTasks` | User can add tasks to the compute cluster |
+| `GetComputeTasks` | User can get and list tasks from the compute cluster |
 
 ## Ddc
 
 | Name | Description |
 | ---- | ----------- |
-| `DdcReadObject` | General read access to refs / blobs and so on |
-| `DdcWriteObject` | General write access to upload refs / blobs etc |
-| `DdcDeleteObject` | Access to delete blobs / refs etc |
-| `DdcDeleteBucket` | Access to delete a particular bucket |
-| `DdcDeleteNamespace` | Access to delete a whole namespace |
-| `DdcReadTransactionLog` | Access to read the transaction log |
-| `DdcWriteTransactionLog` | Access to write the transaction log |
-| `DdcAdminAction` | Access to perform administrative task |
+| `DdcReadObject` |  |
+| `DdcWriteObject` |  |
+| `DdcDeleteObject` |  |
+| `DdcDeleteBucket` |  |
+| `DdcDeleteNamespace` |  |
+| `DdcReadTransactionLog` |  |
+| `DdcWriteTransactionLog` |  |
+| `DdcAdminAction` |  |
 
 ## Devices
 
@@ -383,19 +336,19 @@
 
 | Name | Description |
 | ---- | ----------- |
-| `ViewLeases` |  |
-| `ViewLeaseTasks` |  |
+| `ViewLeases` | View all the leases that an agent has worked on |
+| `ViewLeaseTasks` | View the task data for a lease |
 
 ## Logs
 
 | Name | Description |
 | ---- | ----------- |
-| `CreateLog` |  |
-| `UpdateLog` |  |
-| `ViewLog` |  |
-| `WriteLogData` |  |
-| `CreateEvent` |  |
-| `ViewEvent` |  |
+| `CreateLog` | Ability to create a log. Implicitly granted to agents. |
+| `UpdateLog` | Ability to update log metadata |
+| `ViewLog` | Ability to view a log contents |
+| `WriteLogData` | Ability to write log data |
+| `CreateEvent` | Ability to create events |
+| `ViewEvent` | Ability to view events |
 
 ## Notices
 
@@ -415,11 +368,11 @@
 
 | Name | Description |
 | ---- | ----------- |
-| `CreatePool` |  |
-| `UpdatePool` |  |
-| `DeletePool` |  |
-| `ViewPool` |  |
-| `ListPools` |  |
+| `CreatePool` | Create a global pool of agents |
+| `UpdatePool` | Modify an agent pool |
+| `DeletePool` | Delete an agent pool |
+| `ViewPool` | Ability to view a pool |
+| `ListPools` | View all the available agent pools |
 
 ## Projects
 
@@ -450,16 +403,16 @@
 
 | Name | Description |
 | ---- | ----------- |
-| `CreateSession` |  |
-| `ViewSession` |  |
+| `CreateSession` | Granted to agents to call CreateSession, which returns a bearer token identifying themselves valid to call UpdateSesssion via gRPC. |
+| `ViewSession` | Allows viewing information about an agent session |
 
 ## Software
 
 | Name | Description |
 | ---- | ----------- |
-| `UploadSoftware` |  |
-| `DownloadSoftware` |  |
-| `DeleteSoftware` |  |
+| `UploadSoftware` | Ability to upload new versions of the agent software |
+| `DownloadSoftware` | Ability to download the agent software |
+| `DeleteSoftware` | Ability to delete agent software |
 
 ## Streams
 
@@ -605,12 +558,6 @@
 | `UpdateReplicator` |  |
 | `ViewReplicator` |  |
 
-## Secrets
-
-| Name | Description |
-| ---- | ----------- |
-| `ViewSecret` | View a credential |
-
 ## ServiceAccounts
 
 | Name | Description |
@@ -783,7 +730,7 @@
 
 | Name | Description |
 | ---- | ----------- |
-| `ViewSecret` |  |
+| `ViewSecret` | View a credential |
 
 ## ServiceAccounts
 
@@ -983,19 +930,6 @@
 | `DownloadSoftware` |  |
 | `DeleteSoftware` |  |
 
-## Storage
-
-| Name | Description |
-| ---- | ----------- |
-| `ReadBlobs` | Ability to read blobs from the storage service |
-| `WriteBlobs` | Ability to write blobs to the storage service |
-| `ReadAliases` | Ability to read aliases from the storage service |
-| `WriteAliases` | Ability to write aliases to the storage service |
-| `DeleteAliases` | Ability to write aliases to the storage service |
-| `ReadRefs` | Ability to read refs from the storage service |
-| `WriteRefs` | Ability to write refs to the storage service |
-| `DeleteRefs` | Ability to delete refs |
-
 ## Streams
 
 | Name | Description |
@@ -1007,6 +941,12 @@
 | `ViewChanges` |  |
 | `ViewTemplate` |  |
 
+## Symbols
+
+| Name | Description |
+| ---- | ----------- |
+| `ReadSymbols` | Ability to download symbols |
+
 ## Accounts
 
 | Name | Description |
@@ -1174,14 +1114,14 @@
 
 | Name | Description |
 | ---- | ----------- |
-| `ReadBlobs` |  |
-| `WriteBlobs` |  |
-| `ReadAliases` |  |
-| `WriteAliases` |  |
-| `DeleteAliases` |  |
-| `ReadRefs` |  |
-| `WriteRefs` |  |
-| `DeleteRefs` |  |
+| `ReadBlobs` | Ability to read blobs from the storage service |
+| `WriteBlobs` | Ability to write blobs to the storage service |
+| `ReadAliases` | Ability to read aliases from the storage service |
+| `WriteAliases` | Ability to write aliases to the storage service |
+| `DeleteAliases` | Ability to write aliases to the storage service |
+| `ReadRefs` | Ability to read refs from the storage service |
+| `WriteRefs` | Ability to write refs to the storage service |
+| `DeleteRefs` | Ability to delete refs |
 
 ## Streams
 
@@ -1194,12 +1134,11 @@
 | `ViewChanges` |  |
 | `ViewTemplate` |  |
 
-## Tools
+## Symbols
 
 | Name | Description |
 | ---- | ----------- |
-| `DownloadTool` | Ability to download a tool |
-| `UploadTool` | Ability to upload new tool versions |
+| `ReadSymbols` |  |
 
 ## Accounts
 
@@ -1392,14 +1331,14 @@
 
 | Name | Description |
 | ---- | ----------- |
-| `ReadSymbols` | Ability to download symbols |
+| `ReadSymbols` |  |
 
 ## Tools
 
 | Name | Description |
 | ---- | ----------- |
-| `DownloadTool` |  |
-| `UploadTool` |  |
+| `DownloadTool` | Ability to download a tool |
+| `UploadTool` | Ability to upload new tool versions |
 
 ## Accounts
 
diff --git a/Engine/Source/Programs/Horde/Docs/Config/Schema/Globals.md b/Engine/Source/Programs/Horde/Docs/Config/Schema/Globals.md
--- a/Engine/Source/Programs/Horde/Docs/Config/Schema/Globals.md
+++ b/Engine/Source/Programs/Horde/Docs/Config/Schema/Globals.md
@@ -58,13 +58,13 @@
 ---- | -----------
 `analytics` | [AnalyticsConfig](#analyticsconfig)<br>Configuration for the analytics plugin
 `build` | [BuildConfig](#buildconfig)<br>Configuration for the build plugin
-`compute` | [ComputeConfig](#computeconfig)<br>Configuration for the compute plugin
 `ddc` | [EmptyPluginConfig](#emptypluginconfig)<br>Configuration for the ddc plugin
-`secrets` | [SecretsConfig](#secretsconfig)<br>Configuration for the secrets plugin
+`compute` | [ComputeConfig](#computeconfig)<br>Configuration for the compute plugin
 `experimental` | [ExperimentalConfig](#experimentalconfig)<br>Configuration for the experimental plugin
+`secrets` | [SecretsConfig](#secretsconfig)<br>Configuration for the secrets plugin
+`symbols` | [SymbolsConfig](#symbolsconfig)<br>Configuration for the symbols plugin
 `storage` | [StorageConfig](#storageconfig)<br>Configuration for the storage plugin
 `tools` | [ToolsConfig](#toolsconfig)<br>Configuration for the tools plugin
-`symbols` | [SymbolsConfig](#symbolsconfig)<br>Configuration for the symbols plugin
 
 ## AnalyticsConfig
 
@@ -232,6 +232,11 @@
 `excludeActions` | `string[]`<br>Actions to exclude from the inherited actions
 `extends` | `string[]`<br>Other profiles to extend from
 
+## EmptyPluginConfig
+
+Empty implementation of
+
+
 ## ComputeConfig
 
 Configuration for the compute system
@@ -416,10 +421,44 @@
 `description` | `string`<br>Human-readable description
 `computeId` | `string`<br>Compute ID for this network (used when allocating compute resources)
 
-## EmptyPluginConfig
+## ExperimentalConfig
 
-Empty implementation of
+Configuration for the experimental plugin
 
+Name | Description
+---- | -----------
+`slack` | [SlackConfig](#slackconfig)<br>Configuration for sending experimental Job Summary Slack notifications
 
+## SlackConfig
+
+Configuration for the experimental Slack notifications
+
+Name | Description
+---- | -----------
+`jobNotifications` | [JobNotificationConfig](#jobnotificationconfig)`[]`<br>List of configurations for sending out job summary notifications
+
+## JobNotificationConfig
+
+Configuration for the job notification
+
+Name | Description
+---- | -----------
+`enabled` | `boolean`<br>Whether or not the current configuration is enabled
+`template` | `string`<br>Identifier of the template this notification applies to
+`platformPattern` | `string`<br>Regex pattern of the platforms this notification will include
+`notificationFormats` | [NotificationFormatConfig](#notificationformatconfig)`[]`<br>List of formats to be applied for this notification
+`channels` | `string[]`<br>List of Slack channel or user identifiers the notification will be sent to
+`streams` | `string[]`<br>List of streams this notification will be used for
+
+## NotificationFormatConfig
+
+Configuration for the format of the notification group
+
+Name | Description
+---- | -----------
+`group` | `string`<br>Name of the group
+`stepPattern` | `string`<br>Regex pattern to determine if a job step belongs within the group
+`alternativePlatformPattern` | `string`<br>Optional regex pattern to override the
+
 ## SecretsConfig
 
@@ -460,44 +499,24 @@
 `Text` | Secret is a plain text value which will be stored using the external secret key
 `Json` | Secret is a JSON formatted string containing key/value pairs
 
-## ExperimentalConfig
+## SymbolsConfig
 
-Configuration for the experimental plugin
+Configuration for the tools system
 
 Name | Description
 ---- | -----------
-`slack` | [SlackConfig](#slackconfig)<br>Configuration for sending experimental Job Summary Slack notifications
+`stores` | [SymbolStoreConfig](#symbolstoreconfig)`[]`<br>List of symbol stores
 
-## SlackConfig
+## SymbolStoreConfig
 
-Configuration for the experimental Slack notifications
+Configuration for a symbol store
 
 Name | Description
 ---- | -----------
-`jobNotifications` | [JobNotificationConfig](#jobnotificationconfig)`[]`<br>List of configurations for sending out job summary notifications
+`id` | `string`<br>Identifier for this store
+`namespaceId` | `string`<br>Configuration for the symbol store backend
+`public` | `boolean`<br>Whether to make this store available without auth
+`acl` | [AclConfig](#aclconfig)<br>Access to the symbol store
 
-## JobNotificationConfig
-
-Configuration for the job notification
-
-Name | Description
----- | -----------
-`enabled` | `boolean`<br>Whether or not the current configuration is enabled
-`template` | `string`<br>Identifier of the template this notification applies to
-`platformPattern` | `string`<br>Regex pattern of the platforms this notification will include
-`notificationFormats` | [NotificationFormatConfig](#notificationformatconfig)`[]`<br>List of formats to be applied for this notification
-`channels` | `string[]`<br>List of Slack channel or user identifiers the notification will be sent to
-`streams` | `string[]`<br>List of streams this notification will be used for
-
-## NotificationFormatConfig
-
-Configuration for the format of the notification group
-
-Name | Description
----- | -----------
-`group` | `string`<br>Name of the group
-`stepPattern` | `string`<br>Regex pattern to determine if a job step belongs within the group
-`alternativePlatformPattern` | `string`<br>Optional regex pattern to override the
-
 ## StorageConfig
 
@@ -598,22 +617,3 @@
 `metadata` | `string` `->` `string`<br>Metadata for this tool
 `namespaceId` | `string`<br>Default namespace for new deployments of this tool
 `acl` | [AclConfig](#aclconfig)<br>Permissions for the tool
-
-## SymbolsConfig
-
-Configuration for the tools system
-
-Name | Description
----- | -----------
-`stores` | [SymbolStoreConfig](#symbolstoreconfig)`[]`<br>List of symbol stores
-
-## SymbolStoreConfig
-
-Configuration for a symbol store
-
-Name | Description
----- | -----------
-`id` | `string`<br>Identifier for this store
-`namespaceId` | `string`<br>Configuration for the symbol store backend
-`public` | `boolean`<br>Whether to make this store available without auth
-`acl` | [AclConfig](#aclconfig)<br>Access to the symbol store
