Horde: Warn when OIDC is configured with confidential client as this will prevent command-line auth

Horde's C# client requires a public OAuth client to be configured and these auth settings are auto-detected through this endpoint.

diff --git a/Engine/Source/Programs/Horde/HordeServer/Server/ServerController.cs b/Engine/Source/Programs/Horde/HordeServer/Server/ServerController.cs
--- a/Engine/Source/Programs/Horde/HordeServer/Server/ServerController.cs
+++ b/Engine/Source/Programs/Horde/HordeServer/Server/ServerController.cs
@@ -13,8 +13,8 @@ using HordeServer.Plugins;
 using HordeServer.Utilities;
 using Microsoft.AspNetCore.Authorization;
 using Microsoft.AspNetCore.Mvc;
+using Microsoft.Extensions.Logging;
 using Microsoft.Extensions.Options;
-using MongoDB.Driver;
 
 namespace HordeServer.Server
 {
@@ -30,15 +30,17 @@ namespace HordeServer.Server
 		readonly IAgentVersionProvider? _agentVersionProvider;
 		readonly IPluginCollection _pluginCollection;
 		readonly IOptionsSnapshot<GlobalConfig> _globalConfig;
+		readonly ILogger<ServerController> _logger;
 
 		/// <summary>
 		/// Constructor
 		/// </summary>
-		public ServerController(IEnumerable<IAgentVersionProvider> agentVersionProviders, IPluginCollection pluginCollection, IOptionsSnapshot<GlobalConfig> globalConfig)
+		public ServerController(IEnumerable<IAgentVersionProvider> agentVersionProviders, IPluginCollection pluginCollection, IOptionsSnapshot<GlobalConfig> globalConfig, ILogger<ServerController> logger)
 		{
 			_agentVersionProvider = agentVersionProviders.FirstOrDefault();
 			_pluginCollection = pluginCollection;
 			_globalConfig = globalConfig;
+			_logger = logger;
 		}
 
 		/// <summary>
@@ -112,12 +114,14 @@ namespace HordeServer.Server
 			return response;
 		}
 
+		private const string ServerAuthRoute = "/api/v1/server/auth";
+		
 		/// <summary>
 		/// Returns settings for automating auth against this server
 		/// </summary>
 		[HttpGet]
 		[AllowAnonymous]
-		[Route("/api/v1/server/auth")]
+		[Route(ServerAuthRoute)]
 		public ActionResult<GetAuthConfigResponse> GetAuthConfig()
 		{
 			ServerSettings settings = _globalConfig.Value.ServerSettings;
@@ -132,6 +136,15 @@ namespace HordeServer.Server
 			}
 			else
 			{
+				if (!String.IsNullOrEmpty(settings.OidcClientSecret))
+				{
+					_logger.LogWarning(
+						"OIDC config mismatch: Command-line auth requires a public OAuth/OIDC client, but a confidential client is configured ({ClientSecret} is set). " +
+						"This will prevent Horde's C# client from signing in and block usage of Unreal Build Accelerator. " +
+						"To fix: Configure your OAuth/OIDC client as public (SPA/mobile/desktop) and remove the client secret",
+						nameof(ServerSettings.OidcClientSecret)
+					);
+				}
 				response.ServerUrl = settings.OidcAuthority;
 				response.ClientId = settings.OidcClientId;
 				response.Scopes = settings.OidcApiRequestedScopes;
