Promotion of the Experimental plugin configuration to reside outside of the global configuration.

#rnx

diff --git a/Engine/Source/Programs/Horde/Docs/Config.md b/Engine/Source/Programs/Horde/Docs/Config.md
--- a/Engine/Source/Programs/Horde/Docs/Config.md
+++ b/Engine/Source/Programs/Horde/Docs/Config.md
@@ -57,6 +57,14 @@ cases.
 * [Agents](Config/Agents.md): Configuring worker machines to connect
   to the Horde server for CI and remote execution workspaces.
 
+### Experimental
+
+Experimental research and development functionality within Horde.
+While functional, this plugin is subject to change without notice.
+
+* [Experimental](Config/Experimental.md): Configuring the current
+  functionality of this plugin.
+
 ### Secrets
 
 Managed access to sensitive data integrated with Horde's
diff --git a/Engine/Source/Programs/Horde/Docs/Config/Experimental.md b/Engine/Source/Programs/Horde/Docs/Config/Experimental.md
new file mode 100644
--- /dev/null
+++ b/Engine/Source/Programs/Horde/Docs/Config/Experimental.md
@@ -0,0 +1,24 @@
+[Horde](../../README.md) > [Configuration](../Config.md) > Experimental
+
+# Experimental
+
+Experimental is a plugin used to research, implement, and test functionality within Horde without directly making changes or impacting the
+core functionality. Use of the functionality within this plugin does not guarantee backward compatibility, the APIs for these features
+are subject to change, and we may remove entire Experimental features or specific functionality at our discretion. The plugin is
+Disabled by default and can be enabled by specifying the following snippet in the appsettings.json file within the `Horde` section:
+
+```
+"Plugins": {
+	"Experimental": {
+		"Enabled": true
+	}
+}
+```
+
+The [Getting Started > Experimental](../Tutorials/Experimental.md) guide explains the functionality currently for use within Horde.
+
+## Notifications
+
+Horde already provides support for various Slack notifications. The notification work within this plugin is testing a new grouped
+layout which is based on a set of defined rules. The notifications can be configured through the `Notifications` object within
+the global.json file (see [NotificationConfig](Schema/Notification.md)).
diff --git a/Engine/Source/Programs/Horde/Docs/Config/Schema/Notification.md b/Engine/Source/Programs/Horde/Docs/Config/Schema/Notification.md
new file mode 100644
--- /dev/null
+++ b/Engine/Source/Programs/Horde/Docs/Config/Schema/Notification.md
@@ -0,0 +1,50 @@
+[Horde](../../../README.md) > [Configuration](../../Config.md) > *.notification.json
+
+# *.notification.json
+
+Config for notifications
+
+Name | Description
+---- | -----------
+`id` | `string`<br>Identifier for this store
+`notificationStreams` | [NotificationStreamConfig](#notificationstreamconfig)`[]`<br>Streams for the notifications to be sent on the Horde server
+`include` | [ConfigInclude](#configinclude)`[]`<br>Includes for other configuration files
+
+## NotificationStreamConfig
+
+Config for the notifications based on the streams
+
+Name | Description
+---- | -----------
+`streams` | `string[]`<br>List of streams for this notification
+`jobNotifications` | [JobNotificationConfig](#jobnotificationconfig)`[]`<br>Configuration for sending experimental Job Summary Slack notifications
+
+## JobNotificationConfig
+
+Configuration for the job notification
+
+Name | Description
+---- | -----------
+`isEnabled` | `boolean`<br>Whether or not the current configuration is enabled
+`template` | `string`<br>Identifier of the template this notification applies to
+`namePattern` | `string`<br>Regex pattern used to extract name of the step this notification will use for display
+`notificationFormats` | [NotificationFormatConfig](#notificationformatconfig)`[]`<br>List of formats to be applied for this notification
+`channels` | `string[]`<br>List of Slack channel or user identifiers the notification will be sent to
+
+## NotificationFormatConfig
+
+Configuration for the format of the notification group
+
+Name | Description
+---- | -----------
+`group` | `string`<br>Name of the group
+`stepPattern` | `string`<br>Regex pattern to determine if a job step belongs within the group
+`alternativeNamePattern` | `string`<br>Optional regex pattern to override the
+
+## ConfigInclude
+
+Directive to merge config data from another source
+
+Name | Description
+---- | -----------
+`path` | `string`<br>Path to the config data to be included. May be relative to the including file's location.
diff --git a/Engine/Source/Programs/Horde/Docs/Features.md b/Engine/Source/Programs/Horde/Docs/Features.md
--- a/Engine/Source/Programs/Horde/Docs/Features.md
+++ b/Engine/Source/Programs/Horde/Docs/Features.md
@@ -10,6 +10,7 @@
 | Build Automation | 🟩&nbsp;Production Ready | Migrate jobs to a more lightweight data model.<br>Programmable state machine for jobs. |
 | Build Health | 🟩&nbsp;Production Ready | - |
 | DDC | 🟥&nbsp;Early | Recommend using UnrealCloudDDC rather than Horde. |
+| Experimental | 🟧&nbsp;Experimental | - |
 | Remote Execution | 🟩&nbsp;Production Ready | - |
 | UnrealGameSync&nbsp;Server | 🟩&nbsp;Production Ready | - |
 
diff --git a/Engine/Source/Programs/Horde/Docs/Tutorials/Experimental.md b/Engine/Source/Programs/Horde/Docs/Tutorials/Experimental.md
new file mode 100644
--- /dev/null
+++ b/Engine/Source/Programs/Horde/Docs/Tutorials/Experimental.md
@@ -0,0 +1,45 @@
+[Horde](../../README.md) > Getting Started: Experimental
+
+# Getting Started: Experimental
+
+## Introduction
+
+**Horde** provides an Experimental plugin used to research, develop and test new fucntionality. The functionality
+specified within this plugin, while functionally stable, should not be used in a production environment.
+
+The information specified below shows how to configure the Experimental plugin for use with the current
+features and functionality.
+
+## Prerequisites
+
+* Horde Server installation (see [Getting Started: Install Horde](InstallHorde.md)).
+* A working Slack token which will be used to send out notifications (see [Slack Integration](../Deployment/Integrations/Slack.md)).
+
+## Steps
+
+1. Configure the plugin to provide the path where the notifications will reside, which can be added by
+   the following snippet to your [globals.json](../Config/Orientation.md) file:
+
+    ```json
+    // Define the location to the available notification configurations.
+    "experimental": {
+        "notifications": [
+            {
+                "id": "default",
+                "include": [
+                    {
+                        "path": "$(HordeDir)/Defaults/default.notification.json"
+                    }
+                ]
+            }
+        ]
+    },
+    ```
+
+2. Create the configuration file as specified by the path in the previous step. This config
+   will include information regarding the streams the notifications will be applied to as well
+   as the template and formats which will define the style and job information of the notification.
+
+## See Also
+
+* [Configuration > Experimental](../Config/Experimental.md)
diff --git a/Engine/Source/Programs/Horde/Plugins/Experimental/HordeServer.Experimental/ExperimentalConfig.cs b/Engine/Source/Programs/Horde/Plugins/Experimental/HordeServer.Experimental/ExperimentalConfig.cs
--- a/Engine/Source/Programs/Horde/Plugins/Experimental/HordeServer.Experimental/ExperimentalConfig.cs
+++ b/Engine/Source/Programs/Horde/Plugins/Experimental/HordeServer.Experimental/ExperimentalConfig.cs
@@ -1,6 +1,9 @@
 // Copyright Epic Games, Inc. All Rights Reserved.
 
+using EpicGames.Horde.Jobs.Templates;
+using EpicGames.Horde.Streams;
 using HordeServer.Plugins;
+using System.Diagnostics.CodeAnalysis;
 
 namespace HordeServer
 {
@@ -10,14 +13,45 @@ namespace HordeServer
 	public class ExperimentalConfig : IPluginConfig
 	{
 		/// <summary>
-		/// Configuration for sending experimental Job Summary Slack notifications
+		/// Notifications to be sent on the Horde server
 		/// </summary>
-		public SlackConfig Slack { get; set; } = new SlackConfig();
+		public List<NotificationConfig> Notifications { get; set; } = new List<NotificationConfig>();
+
+		readonly Dictionary<NotificationConfigId, NotificationConfig> _notificationConfigLookup = new Dictionary<NotificationConfigId, NotificationConfig>();
 
 		/// <inheritdoc/>
 		public void PostLoad(PluginConfigOptions configOptions)
 		{
-			Slack.PostLoad(configOptions);
+			_notificationConfigLookup.Clear();
+			foreach (NotificationConfig config in Notifications)
+			{
+				config.PostLoad();
+				_notificationConfigLookup.Add(config.Id, config);
+			}
+		}
+
+		/// <summary>
+		/// Attempt to get config for a given stream and template
+		/// </summary>
+		/// <param name="streamId">Stream id</param>
+		/// <param name="templateId">Template id</param>
+		/// <param name="jobNotificationConfigs">Receives the list of config objects on success</param>
+		/// <param name="includeDisabledConfigs">Flag on whether or not disabled configs should be included</param>
+		/// <returns>True if the job notification configs were found</returns>
+		public bool TryGetJobNofications(StreamId streamId, TemplateId templateId, [NotNullWhen(true)] out HashSet<JobNotificationConfig>? jobNotificationConfigs, bool includeDisabledConfigs = false)
+		{
+			jobNotificationConfigs = new HashSet<JobNotificationConfig>();
+
+			foreach (NotificationConfig config in Notifications)
+			{
+				HashSet<JobNotificationConfig>? localJobNotificationConfigs;
+				if (config.TryGetJobNofications(streamId, templateId, out localJobNotificationConfigs, includeDisabledConfigs))
+				{
+					jobNotificationConfigs.UnionWith(localJobNotificationConfigs);
+				}
+			}
+
+			return jobNotificationConfigs.Count != 0;
 		}
 	}
 }
diff --git a/Engine/Source/Programs/Horde/Plugins/Experimental/HordeServer.Experimental/ExperimentalPlugin.cs b/Engine/Source/Programs/Horde/Plugins/Experimental/HordeServer.Experimental/ExperimentalPlugin.cs
--- a/Engine/Source/Programs/Horde/Plugins/Experimental/HordeServer.Experimental/ExperimentalPlugin.cs
+++ b/Engine/Source/Programs/Horde/Plugins/Experimental/HordeServer.Experimental/ExperimentalPlugin.cs
@@ -1,11 +1,11 @@
 // Copyright Epic Games, Inc. All Rights Reserved.
 
-using System.Diagnostics.CodeAnalysis;
 using HordeServer.Notifications;
 using HordeServer.Notifications.Sinks;
 using HordeServer.Plugins;
 using Microsoft.AspNetCore.Builder;
 using Microsoft.Extensions.DependencyInjection;
+using System.Diagnostics.CodeAnalysis;
 
 namespace HordeServer
 {
@@ -40,19 +40,19 @@ namespace HordeServer
 	public static class ExperimentalPluginExtensions
 	{
 		/// <summary>
-		/// Configures the build plugin
+		/// Configures the experimental plugin
 		/// </summary>
 		public static void AddExperimentalConfig(this IDictionary<PluginName, IPluginConfig> dictionary, ExperimentalConfig experimentalConfig)
 			=> dictionary[new PluginName("Experimental")] = experimentalConfig;
 
 		/// <summary>
-		/// Gets configuration for the build plugin
+		/// Gets configuration for the experimental plugin
 		/// </summary>
 		public static ExperimentalConfig GetExperimentalConfig(this IDictionary<PluginName, IPluginConfig> dictionary)
 			=> (ExperimentalConfig)dictionary[new PluginName("Experimental")];
 
 		/// <summary>
-		/// Gets configuration for the build plugin
+		/// Gets configuration for the experimental plugin
 		/// </summary>
 		public static bool TryGetExperimentalConfig(this IDictionary<PluginName, IPluginConfig> dictionary, [NotNullWhen(true)] out ExperimentalConfig? experimentalConfig)
 		{
diff --git a/Engine/Source/Programs/Horde/Plugins/Experimental/HordeServer.Experimental/Notifications/NotificationConfig.cs b/Engine/Source/Programs/Horde/Plugins/Experimental/HordeServer.Experimental/Notifications/NotificationConfig.cs
new file mode 100644
--- /dev/null
+++ b/Engine/Source/Programs/Horde/Plugins/Experimental/HordeServer.Experimental/Notifications/NotificationConfig.cs
@@ -0,0 +1,231 @@
+// Copyright Epic Games, Inc. All Rights Reserved.
+
+using EpicGames.Core;
+using EpicGames.Horde.Jobs.Templates;
+using EpicGames.Horde.Streams;
+using HordeServer.Configuration;
+using HordeServer.Utilities;
+using System.Diagnostics.CodeAnalysis;
+using System.Text.RegularExpressions;
+
+namespace HordeServer
+{
+	/// <summary>
+	/// Configuration for the format of the notification group
+	/// </summary>
+	public class NotificationFormatConfig
+	{
+		/// <summary>
+		/// Name of the group
+		/// </summary>
+		public string Group { get; set; } = String.Empty;
+
+		/// <summary>
+		/// Regex pattern to determine if a job step belongs within the group
+		/// </summary>
+		public string StepPattern { get; set; } = String.Empty;
+
+		/// <summary>
+		/// Optional regex pattern to override the <see cref="JobNotificationConfig.NamePattern"/>
+		/// </summary>
+		public string? AlternativeNamePattern { get; set; }
+	}
+
+	/// <summary>
+	/// Configuration for the job notification
+	/// </summary>
+	public class JobNotificationConfig
+	{
+		/// <summary>
+		/// Whether or not the current configuration is enabled
+		/// </summary>
+		public bool IsEnabled { get; set; } = true;
+
+		/// <summary>
+		/// Identifier of the template this notification applies to
+		/// </summary>
+		public TemplateId Template { get; set; }
+
+		/// <summary>
+		/// Regex pattern used to extract name of the step this notification will use for display
+		/// </summary>
+		public string NamePattern { get; set; } = String.Empty;
+
+		/// <summary>
+		/// List of formats to be applied for this notification
+		/// </summary>
+		public List<NotificationFormatConfig> NotificationFormats { get; set; } = new List<NotificationFormatConfig>();
+
+		/// <summary>
+		/// List of Slack channel or user identifiers the notification will be sent to
+		/// </summary>
+		public List<string> Channels { get; set; } = new List<string>();
+
+		/// <summary>
+		/// Gets the group and platform information of the corresponding step within the list of notification formats
+		/// </summary>
+		/// <param name="stepName">Name of the job step to test against</param>
+		/// <param name="group">Name of the notification group the job step belongs to or null if not found</param>
+		/// <param name="name">Name associated with the job step or null if not found</param>
+		/// <returns>True if the job step is to be included in the notification or false otherwise</returns>
+		public bool GetStepDetails(string stepName, out string? group, out string? name)
+		{
+			group = null;
+			name = null;
+
+			foreach (NotificationFormatConfig formatConfig in NotificationFormats)
+			{
+				Match monitoredPattern = Regex.Match(stepName, formatConfig.StepPattern);
+				if (monitoredPattern.Success)
+				{
+					// Fetch our platform from the step name
+					string namePattern = !String.IsNullOrEmpty(formatConfig.AlternativeNamePattern) ? formatConfig.AlternativeNamePattern : NamePattern;
+					Match match = Regex.Match(stepName, namePattern);
+					if (match.Success)
+					{
+						group = formatConfig.Group;
+						name = match.Groups[0].Value;
+
+						return true;
+					}
+				}
+			}
+
+			return false;
+		}
+	}
+
+	/// <summary>
+	/// Config for the notifications based on the streams
+	/// </summary>
+	public class NotificationStreamConfig
+	{
+		/// <summary>
+		/// List of streams for this notification
+		/// </summary>
+		public List<StreamId> Streams { get; set; } = new List<StreamId>();
+
+		/// <summary>
+		/// Configuration for sending experimental Job Summary Slack notifications
+		/// </summary>
+		public List<JobNotificationConfig> JobNotifications { get; set; } = new List<JobNotificationConfig>();
+
+		readonly Dictionary<TemplateId, List<JobNotificationConfig>> _notificationLookup = new Dictionary<TemplateId, List<JobNotificationConfig>>();
+
+		/// <summary>
+		/// Called after the store has been deserialized to compute cached values
+		/// </summary>
+		public void PostLoad()
+		{
+			_notificationLookup.Clear();
+			foreach (JobNotificationConfig notificationConfig in JobNotifications)
+			{
+				List<JobNotificationConfig>? configs;
+				if (!_notificationLookup.TryGetValue(notificationConfig.Template, out configs))
+				{
+					configs = new List<JobNotificationConfig>();
+				}
+
+				configs.Add(notificationConfig);
+				_notificationLookup[notificationConfig.Template] = configs;
+			}
+		}
+
+		/// <summary>
+		/// Attempt to get config for a given template
+		/// </summary>
+		/// <param name="templateId">Template id</param>
+		/// <param name="notificationConfigList">Receives the list of config objects on success</param>
+		/// <returns>True if the stream was found</returns>
+		public bool TryGetTemplateConfigs(TemplateId templateId, [NotNullWhen(true)] out List<JobNotificationConfig>? notificationConfigList)
+			=> _notificationLookup.TryGetValue(templateId, out notificationConfigList);
+	}
+
+	/// <summary>
+	/// Config for notifications
+	/// </summary>
+	[JsonSchema("https://unrealengine.com/horde/notification")]
+	[JsonSchemaCatalog("Horde Notifications", "Horde notification configuration file", new[] { "*.notification.json", "Notifications/*.json" })]
+	[ConfigDoc("*.notification.json", "[Horde](../../../README.md) > [Configuration](../../Config.md)", "Config/Schema/Notification.md")]
+	[ConfigIncludeRoot]
+	[ConfigMacroScope]
+	public class NotificationConfig
+	{
+		/// <summary>
+		/// Identifier for this store
+		/// </summary>
+		public NotificationConfigId Id { get; set; }
+
+		/// <summary>
+		/// Streams for the notifications to be sent on the Horde server
+		/// </summary>
+		public List<NotificationStreamConfig> NotificationStreams { get; set; } = new List<NotificationStreamConfig>();
+
+		/// <summary>
+		/// Includes for other configuration files
+		/// </summary>
+		public List<ConfigInclude> Include { get; set; } = new List<ConfigInclude>();
+
+		readonly Dictionary<StreamId, HashSet<NotificationStreamConfig>> _streamLookup = new Dictionary<StreamId, HashSet<NotificationStreamConfig>>();
+
+		/// <summary>
+		/// Called after the store has been deserialized to compute cached values
+		/// </summary>
+		public void PostLoad()
+		{
+			_streamLookup.Clear();
+			foreach (NotificationStreamConfig notificationStream in NotificationStreams)
+			{
+				notificationStream.PostLoad();
+
+				foreach (StreamId streamId in notificationStream.Streams)
+				{
+					HashSet<NotificationStreamConfig>? configs;
+					if (!_streamLookup.TryGetValue(streamId, out configs))
+					{
+						configs = new HashSet<NotificationStreamConfig>();
+					}
+
+					configs.Add(notificationStream);
+					_streamLookup[streamId] = configs;
+				}
+			}
+		}
+
+		/// <summary>
+		/// Attempt to get config for a given stream and template
+		/// </summary>
+		/// <param name="streamId">Stream id</param>
+		/// <param name="templateId">Template id</param>
+		/// <param name="jobNotificationConfigs">Receives the list of config objects on success</param>
+		/// <param name="includeDisabledConfigs">Flag on whether or not disabled configs should be included</param>
+		/// <returns>True if the job notification configs were found</returns>
+		public bool TryGetJobNofications(StreamId streamId, TemplateId templateId, [NotNullWhen(true)] out HashSet<JobNotificationConfig>? jobNotificationConfigs, bool includeDisabledConfigs = false)
+		{
+			jobNotificationConfigs = new HashSet<JobNotificationConfig>();
+
+			HashSet<NotificationStreamConfig>? streamConfigList;
+			if (!_streamLookup.TryGetValue(streamId, out streamConfigList))
+			{
+				return false;
+			}
+
+			foreach (NotificationStreamConfig streamConfig in streamConfigList)
+			{
+				List<JobNotificationConfig>? notificationConfigList;
+				if (streamConfig.TryGetTemplateConfigs(templateId, out notificationConfigList))
+				{
+					foreach (JobNotificationConfig notificationConfig in notificationConfigList)
+					{
+						if (notificationConfig.IsEnabled || includeDisabledConfigs)
+						{
+							jobNotificationConfigs.Add(notificationConfig);
+						}
+					}
+				}
+			}
+
+			return jobNotificationConfigs.Count != 0;
+		}
+	}
+}
diff --git a/Engine/Source/Programs/Horde/Plugins/Experimental/HordeServer.Experimental/Notifications/NotificationConfigId.cs b/Engine/Source/Programs/Horde/Plugins/Experimental/HordeServer.Experimental/Notifications/NotificationConfigId.cs
new file mode 100644
--- /dev/null
+++ b/Engine/Source/Programs/Horde/Plugins/Experimental/HordeServer.Experimental/Notifications/NotificationConfigId.cs
@@ -0,0 +1,51 @@
+// Copyright Epic Games, Inc. All Rights Reserved.
+
+using EpicGames.Core;
+using EpicGames.Horde;
+using EpicGames.Serialization;
+using System.ComponentModel;
+
+namespace HordeServer
+{
+	/// <summary>
+	/// Identifier for a notification configuration
+	/// </summary>
+	/// <param name="Id">Id to construct from</param>
+	[LogValueType]
+	[JsonSchemaString]
+	[TypeConverter(typeof(StringIdTypeConverter<NotificationConfigId, NotificationStoreIdConverter>))]
+	[StringIdConverter(typeof(NotificationStoreIdConverter))]
+	[CbConverter(typeof(StringIdCbConverter<NotificationConfigId, NotificationStoreIdConverter>))]
+	public readonly record struct NotificationConfigId(StringId Id)
+	{
+		/// <summary>
+		/// Default notification configuration identifier
+		/// </summary>
+		public static NotificationConfigId Default { get; } = new NotificationConfigId("default");
+
+		/// <summary>
+		/// Constructor
+		/// </summary>
+		public NotificationConfigId(string id) : this(new StringId(id))
+		{
+		}
+
+		/// <inheritdoc cref="StringId.IsEmpty"/>
+		public bool IsEmpty => Id.IsEmpty;
+
+		/// <inheritdoc/>
+		public override string ToString() => Id.ToString();
+	}
+
+	/// <summary>
+	/// Converter to and from <see cref="StringId"/> instances.
+	/// </summary>
+	public class NotificationStoreIdConverter : StringIdConverter<NotificationConfigId>
+	{
+		/// <inheritdoc/>
+		public override NotificationConfigId FromStringId(StringId id) => new NotificationConfigId(id);
+
+		/// <inheritdoc/>
+		public override StringId ToStringId(NotificationConfigId value) => value.Id;
+	}
+}
diff --git a/Engine/Source/Programs/Horde/Plugins/Experimental/HordeServer.Experimental/Notifications/Sinks/ExperimentalSlackNotificationSink.cs b/Engine/Source/Programs/Horde/Plugins/Experimental/HordeServer.Experimental/Notifications/Sinks/ExperimentalSlackNotificationSink.cs
--- a/Engine/Source/Programs/Horde/Plugins/Experimental/HordeServer.Experimental/Notifications/Sinks/ExperimentalSlackNotificationSink.cs
+++ b/Engine/Source/Programs/Horde/Plugins/Experimental/HordeServer.Experimental/Notifications/Sinks/ExperimentalSlackNotificationSink.cs
@@ -155,7 +155,13 @@ namespace HordeServer.Notifications.Sinks
 		/// <param name="userCollection">Collection of users</param>
 		/// <param name="serverInfo">Information of the current server instance</param>
 		/// <param name="logger">Logger for output</param>
-		public ExperimentalSlackNotificationSink(IMongoService mongoService, IOptions<BuildServerConfig> buildServerConfig, IOptionsMonitor<BuildConfig> buildConfig, IOptionsMonitor<ExperimentalConfig> experimentalConfig, IUserCollection userCollection, IServerInfo serverInfo, ILogger<ExperimentalSlackNotificationSink> logger)
+		public ExperimentalSlackNotificationSink(IMongoService mongoService,
+			IOptions<BuildServerConfig> buildServerConfig,
+			IOptionsMonitor<BuildConfig> buildConfig,
+			IOptionsMonitor<ExperimentalConfig> experimentalConfig,
+			IUserCollection userCollection,
+			IServerInfo serverInfo,
+			ILogger<ExperimentalSlackNotificationSink> logger)
 		{
 			_buildConfig = buildConfig;
 			_experimentalConfig = experimentalConfig;
@@ -230,7 +236,7 @@ namespace HordeServer.Notifications.Sinks
 				// Update our steps to reflect the canceled job
 				_ = await UpdatePendingJobStepNotificationsAsync(job, jobNotificationState, cancellationToken);
 
-				await SendSlackSummaryMessageAsync(job, branch, hordeJobURL, jobNotificationState.TemplateId, jobNotificationState.Recipient, abortedByUserName, cancellationToken);
+				await SendSlackSummaryMessageAsync(job, branch, hordeJobURL, jobNotificationState.Recipient, abortedByUserName, cancellationToken);
 
 				// No further notifications will be sent out and can be cleaned up
 				_ = await DeleteJobStepNotificationStatesAsync(jobNotificationState.JobId, jobNotificationState.TemplateId, jobNotificationState.Recipient, cancellationToken);
@@ -336,13 +342,12 @@ namespace HordeServer.Notifications.Sinks
 		/// <param name="cancellationToken">Cancellation token for the operation</param>
 		async Task ProcessJobStepCompleteNotificationAsync(IJob job, IJobStep step, List<LogData> jobStepEventData, CancellationToken cancellationToken = default)
 		{
-			SlackConfig currentConfig = _experimentalConfig.CurrentValue.Slack;
+			ExperimentalConfig currentExperimentalConfig = _experimentalConfig.CurrentValue;
 
-			List<JobNotificationConfig>? configs = currentConfig.GetNotificationConfigs(job.TemplateId, job.StreamId);
-			if (configs is null || !configs.Any())
+			HashSet<JobNotificationConfig>? jobNotificationConfigs;
+			if (!currentExperimentalConfig.TryGetJobNofications(job.StreamId, job.TemplateId, out jobNotificationConfigs))
 			{
-				// Log that there is nothing to do here
-				_logger.LogInformation("No notification config found for '{JobName}' in stream '{StreamId}'", job.Name, job.StreamId);
+				_logger.LogInformation("No config found for '{JobName}' in stream '{StreamId}' and for template '{TemplateId}'", job.Name, job.StreamId, job.TemplateId);
 				return;
 			}
 
@@ -361,16 +366,16 @@ namespace HordeServer.Notifications.Sinks
 				bool wasSetupSuccess = (step.Outcome == JobStepOutcome.Success || step.Outcome == JobStepOutcome.Warnings);
 				if (!wasSetupSuccess)
 				{
-					await SendSetupFailedSlackMessageAsync(job, step, branch, jobStepEventData, configs, cancellationToken);
+					await SendSetupFailedSlackMessageAsync(job, step, branch, jobStepEventData, jobNotificationConfigs, cancellationToken);
 					return;
 				}
 
-				foreach (JobNotificationConfig config in configs)
+				foreach (JobNotificationConfig config in jobNotificationConfigs)
 				{
 					await ProcessJobAsync(job, config, cancellationToken);
 					foreach (string recipient in config.Channels)
 					{
-						await SendSlackSummaryMessageAsync(job, branch, hordeJobURL, config.Template, recipient, null, cancellationToken);
+						await SendSlackSummaryMessageAsync(job, branch, hordeJobURL, recipient, null, cancellationToken);
 					}
 				}
 				return;
@@ -380,7 +385,7 @@ namespace HordeServer.Notifications.Sinks
 			bool needsSummaryUpdate = await UpdateAndSendSlackNotificationsForStepAsync(job.Id, step, jobStepEventData, null, cancellationToken);
 
 			// Fetch all dependent steps if there was an error and send notifications with this step as the cause for the dependent steps being skipped
-			if (step.Outcome != JobStepOutcome.Success)
+			if (step.Outcome != JobStepOutcome.Success && step.Outcome != JobStepOutcome.Warnings)
 			{
 				HashSet<JobStepId>? dependentSteps = await FindDependentJobStepIdsAsync(job, step.Id, cancellationToken);
 				if (dependentSteps is not null)
@@ -421,7 +426,7 @@ namespace HordeServer.Notifications.Sinks
 
 			foreach (JobNotificationStateDocument jobNotificationState in jobNotificationStates)
 			{
-				await SendSlackSummaryMessageAsync(job, branch, hordeJobURL, jobNotificationState.TemplateId, jobNotificationState.Recipient, null, cancellationToken);
+				await SendSlackSummaryMessageAsync(job, branch, hordeJobURL, jobNotificationState.Recipient, null, cancellationToken);
 
 				// Check to see if all notifications for this job were sent to then delete all documents related to this job as no other notifications will be sent
 				await TryNotificationCleanupAsync(jobNotificationState, cancellationToken);
@@ -847,16 +852,15 @@ namespace HordeServer.Notifications.Sinks
 		/// <param name="job">Information about the job</param>
 		/// <param name="branch">Name of the branch from the stream</param>
 		/// <param name="hordeJobURL">URL of the job in Horde</param>
-		/// <param name="templateId">Identifier of the stream's template</param>
 		/// <param name="recipient">Identifier of the channel or user in Slack or user's email address</param>
 		/// <param name="abortedByUserName">Optional username provided when cancelling a job</param>
 		/// <param name="cancellationToken">Cancellation token for the operation</param>
 		/// <returns>Generated <see cref="SlackMessage"/> if successful, otherwise returns null</returns>
-		async Task<SlackMessage?> GenerateSlackSummaryMessageAsync(IJob job, string branch, string hordeJobURL, TemplateId templateId, string recipient, string? abortedByUserName = null, CancellationToken cancellationToken = default)
+		async Task<SlackMessage?> GenerateSlackSummaryMessageAsync(IJob job, string branch, string hordeJobURL, string recipient, string? abortedByUserName = null, CancellationToken cancellationToken = default)
 		{
 			// Fetch all of our job steps needed to send out the summary
 			FilterDefinition<JobStepNotificationStateDocument> filter = Builders<JobStepNotificationStateDocument>.Filter.Eq(x => x.JobId, job.Id) &
-																		Builders<JobStepNotificationStateDocument>.Filter.Eq(x => x.TemplateId, templateId) &
+																		Builders<JobStepNotificationStateDocument>.Filter.Eq(x => x.TemplateId, job.TemplateId) &
 																		Builders<JobStepNotificationStateDocument>.Filter.Eq(x => x.Recipient, recipient);
 			List<JobStepNotificationStateDocument>? jobStepNotificationStateDocuments = await GetJobStepNotificationStatesAsync(filter, cancellationToken);
 			if (jobStepNotificationStateDocuments is null)
@@ -874,10 +878,10 @@ namespace HordeServer.Notifications.Sinks
 		/// <param name="step">Information about the step</param>
 		/// <param name="branch">Name of the branch from the stream</param>
 		/// <param name="jobStepEventData">List of the step's logged events</param>
-		/// <param name="jobNotificationConfigs">Information about the configuration tied to his job</param>
+		/// <param name="jobNotificationConfigs">Collection of configurations tied to this job</param>
 		/// <param name="cancellationToken">Cancellation token for the operation</param>
 		/// <returns></returns>
-		async Task SendSetupFailedSlackMessageAsync(IJob job, IJobStep step, string branch, List<LogData> jobStepEventData, List<JobNotificationConfig> jobNotificationConfigs, CancellationToken cancellationToken = default)
+		async Task SendSetupFailedSlackMessageAsync(IJob job, IJobStep step, string branch, List<LogData> jobStepEventData, HashSet<JobNotificationConfig> jobNotificationConfigs, CancellationToken cancellationToken = default)
 		{
 			SlackMessage message = new SlackMessage();
 
@@ -925,21 +929,20 @@ namespace HordeServer.Notifications.Sinks
 		/// <param name="job">Information about the job</param>
 		/// <param name="branch">Name of the branch from the stream</param>
 		/// <param name="hordeJobURL">Optional URL of the job in Horde</param>
-		/// <param name="templateId">Identifier of the stream's template</param>
 		/// <param name="recipient">Identifier of the channel or user in Slack or user's email address</param>
 		/// <param name="abortedByUserName">Optional username provided when cancelling a job</param>
 		/// <param name="cancellationToken">Cancellation token for the operation</param>
 		/// <returns>True if a notification needed to be sent, otherwise returns false</returns>
-		async Task SendSlackSummaryMessageAsync(IJob job, string branch, string hordeJobURL, TemplateId templateId, string recipient, string? abortedByUserName = null, CancellationToken cancellationToken = default)
+		async Task SendSlackSummaryMessageAsync(IJob job, string branch, string hordeJobURL, string recipient, string? abortedByUserName = null, CancellationToken cancellationToken = default)
 		{
-			SlackMessage? message = await GenerateSlackSummaryMessageAsync(job, branch, hordeJobURL, templateId, recipient, abortedByUserName, cancellationToken);
+			SlackMessage? message = await GenerateSlackSummaryMessageAsync(job, branch, hordeJobURL, recipient, abortedByUserName, cancellationToken);
 			if (message is null)
 			{
-				_logger.LogWarning("Unable to generate Job summary notification for job '{JobId}' for the template '{TemplateId}'", job.Id, templateId);
+				_logger.LogWarning("Unable to generate Job summary notification for job '{JobId}' for the template '{TemplateId}'", job.Id, job.TemplateId);
 				return;
 			}
 
-			JobNotificationStateDocument? jobNotificationState = await GetJobNotificationStateAsync(job.Id, templateId, recipient, cancellationToken);
+			JobNotificationStateDocument? jobNotificationState = await GetJobNotificationStateAsync(job.Id, job.TemplateId, recipient, cancellationToken);
 			if (jobNotificationState is not null)
 			{
 				_ = _slackClient.UpdateMessageAsync(jobNotificationState.MessageId, message, cancellationToken);
@@ -955,7 +958,7 @@ namespace HordeServer.Notifications.Sinks
 				}
 
 				SlackMessageId messageId = await _slackClient.PostMessageAsync(channel, message, cancellationToken);
-				_ = await AddOrUpdateJobNotificationStateAsync(job.Id, templateId, recipient, messageId, cancellationToken);
+				_ = await AddOrUpdateJobNotificationStateAsync(job.Id, job.TemplateId, recipient, messageId, cancellationToken);
 			}
 		}
 
diff --git a/Engine/Source/Programs/Horde/Plugins/Experimental/HordeServer.Experimental/Slack/SlackConfig.cs b/Engine/Source/Programs/Horde/Plugins/Experimental/HordeServer.Experimental/Slack/SlackConfig.cs
deleted file mode 100644
