Horde: Updating schema documentation

#rnx

diff --git a/Engine/Source/Programs/Horde/Docs/Config/Schema/Globals.md b/Engine/Source/Programs/Horde/Docs/Config/Schema/Globals.md
--- a/Engine/Source/Programs/Horde/Docs/Config/Schema/Globals.md
+++ b/Engine/Source/Programs/Horde/Docs/Config/Schema/Globals.md
@@ -427,38 +427,7 @@ Configuration for the experimental plugin
 
 Name | Description
 ---- | -----------
-`slack` | [SlackConfig](#slackconfig)<br>Configuration for sending experimental Job Summary Slack notifications
-
-## SlackConfig
-
-Configuration for the experimental Slack notifications
-
-Name | Description
----- | -----------
-`jobNotifications` | [JobNotificationConfig](#jobnotificationconfig)`[]`<br>List of configurations for sending out job summary notifications
-
-## JobNotificationConfig
-
-Configuration for the job notification
-
-Name | Description
----- | -----------
-`enabled` | `boolean`<br>Whether or not the current configuration is enabled
-`template` | `string`<br>Identifier of the template this notification applies to
-`platformPattern` | `string`<br>Regex pattern of the platforms this notification will include
-`notificationFormats` | [NotificationFormatConfig](#notificationformatconfig)`[]`<br>List of formats to be applied for this notification
-`channels` | `string[]`<br>List of Slack channel or user identifiers the notification will be sent to
-`streams` | `string[]`<br>List of streams this notification will be used for
-
-## NotificationFormatConfig
-
-Configuration for the format of the notification group
-
-Name | Description
----- | -----------
-`group` | `string`<br>Name of the group
-`stepPattern` | `string`<br>Regex pattern to determine if a job step belongs within the group
-`alternativePlatformPattern` | `string`<br>Optional regex pattern to override the
+`notifications` | [NotificationConfig](Notification.md)`[]`<br>Notifications to be sent on the Horde server
 
 ## SecretsConfig
 
