Horde: Allow service accounts to modify devices

diff --git a/Engine/Source/Programs/Horde/Plugins/Build/HordeServer.Build/Devices/DevicesController.cs b/Engine/Source/Programs/Horde/Plugins/Build/HordeServer.Build/Devices/DevicesController.cs
--- a/Engine/Source/Programs/Horde/Plugins/Build/HordeServer.Build/Devices/DevicesController.cs
+++ b/Engine/Source/Programs/Horde/Plugins/Build/HordeServer.Build/Devices/DevicesController.cs
@@ -199,10 +199,6 @@ namespace HordeServer.Devices
 		public async Task<ActionResult> UpdateDeviceAsync(string deviceId, [FromBody] UpdateDeviceRequest update, CancellationToken cancellationToken = default)
 		{
 			IUser? internalUser = await _userCollection.GetUserAsync(User, cancellationToken);
-			if (internalUser == null)
-			{
-				return NotFound();
-			}
 
 			DeviceId deviceIdValue = new DeviceId(deviceId);
 
@@ -249,7 +245,7 @@ namespace HordeServer.Devices
 			string? name = update.Name?.Trim();
 			string? address = update.Address?.Trim();
 
-			await _deviceService.UpdateDeviceAsync(deviceIdValue, poolIdValue, name, address, modelIdValue, update.Notes, update.Enabled, update.Problem, update.Maintenance, internalUser.Id, update.Message);
+			await _deviceService.UpdateDeviceAsync(deviceIdValue, poolIdValue, name, address, modelIdValue, update.Notes, update.Enabled, update.Problem, update.Maintenance, internalUser?.Id, update.Message);
 
 			return Ok();
 		}
