Horde: Add support for automated device cleaning on CI and shared kits, server side changes

diff --git a/Engine/Source/Programs/Horde/Plugins/Build/HordeServer.Build/Devices/DeviceCollection.cs b/Engine/Source/Programs/Horde/Plugins/Build/HordeServer.Build/Devices/DeviceCollection.cs
--- a/Engine/Source/Programs/Horde/Plugins/Build/HordeServer.Build/Devices/DeviceCollection.cs
+++ b/Engine/Source/Programs/Horde/Plugins/Build/HordeServer.Build/Devices/DeviceCollection.cs
@@ -158,6 +158,24 @@ namespace HordeServer.Devices
 			[BsonIgnoreIfNull]
 			public DateTime? MaintenanceTimeUtc { get; set; }
 
+			/// <summary>
+			/// The time the last clean started,
+			/// </summary>
+			[BsonIgnoreIfNull]
+			public DateTime? CleanStartTimeUtc { get; set; }
+
+			/// <summary>
+			/// The time the least clean finished
+			/// </summary>
+			[BsonIgnoreIfNull]
+			public DateTime? CleanFinishTimeUtc { get; set; }
+
+			/// <summary>
+			/// Whether a clean is required
+			/// </summary>
+			[BsonIgnoreIfNull]
+			public bool? CleanRequired { get; set; }
+
 			/// <summary>
 			///  The last time this device was reserved, used to cycle devices for reservations
 			/// </summary>
@@ -519,6 +537,7 @@ namespace HordeServer.Devices
 			{
 				updates.Add(updateBuilder.Set(x => x.CheckOutTime, DateTime.UtcNow));
 				updates.Add(updateBuilder.Set(x => x.CheckoutExpiringNotificationSent, false));
+				updates.Add(updateBuilder.Set(x => x.CleanRequired, false));
 			}
 			else
 			{
@@ -755,6 +774,52 @@ namespace HordeServer.Devices
 			}
 		}
 
+		public async Task UpdateDeviceCleaningAsync(DeviceId deviceId, bool? cleaning = null, bool? cleanRequired = null)
+		{
+			UpdateDefinitionBuilder<DeviceDocument> updateBuilder = Builders<DeviceDocument>.Update;
+			List<UpdateDefinition<DeviceDocument>> updates = new List<UpdateDefinition<DeviceDocument>>();
+						
+			IDevice? current = await _devices.Find<DeviceDocument>(x => x.Id == deviceId).FirstOrDefaultAsync();
+
+			if (current == null)
+			{
+				_logger.LogInformation("Device cleaning updated for non-existant {DeviceId}", deviceId.ToString());
+				return;
+			}
+
+			DateTime utcNow = DateTime.UtcNow;
+
+			IAuditLogChannel<DeviceId> logger = GetDeviceLogger(deviceId);
+
+			if (cleaning.HasValue)
+			{
+				if (cleaning.Value)
+				{
+					logger.LogInformation("Starting Clean");
+					updates.Add(updateBuilder.Set(x => x.CleanStartTimeUtc, utcNow));
+					updates.Add(updateBuilder.Set(x => x.CleanFinishTimeUtc, null));
+					updates.Add(updateBuilder.Set(x => x.CleanRequired, false));
+				}				
+				else
+				{
+					logger.LogInformation("Finished Clean");
+					updates.Add(updateBuilder.Set(x => x.CleanFinishTimeUtc, utcNow));
+					updates.Add(updateBuilder.Set(x => x.CleanRequired, false));
+				}
+			}
+
+			if (cleanRequired.HasValue)
+			{
+				logger.LogInformation("Marking as requiring a clean: {CleanRequired}", cleanRequired.Value);
+				updates.Add(updateBuilder.Set(x => x.CleanRequired, cleanRequired.Value));
+			}
+
+			if (updates.Count > 0)
+			{
+				await _devices.FindOneAndUpdateAsync<DeviceDocument>(x => x.Id == deviceId, updateBuilder.Combine(updates));
+			}
+		}
+
 		/// <inheritdoc/>
 		public async Task<bool> DeleteDeviceAsync(DeviceId deviceId)
 		{
@@ -1049,6 +1114,7 @@ namespace HordeServer.Devices
 				deviceUpdates.Add(deviceBuilder.Set(x => x.CheckedOutByUser, null));
 				deviceUpdates.Add(deviceBuilder.Set(x => x.CheckOutTime, null));
 				deviceUpdates.Add(deviceBuilder.Set(x => x.CheckoutExpiringNotificationSent, true));
+				deviceUpdates.Add(deviceBuilder.Set(x => x.CleanRequired, true));
 
 				UpdateResult result = await _devices.UpdateManyAsync(Builders<DeviceDocument>.Filter.In(x => x.Id, expiredDevices.Select(y => y.Id)), deviceBuilder.Combine(deviceUpdates));
 
diff --git a/Engine/Source/Programs/Horde/Plugins/Build/HordeServer.Build/Devices/DeviceService.cs b/Engine/Source/Programs/Horde/Plugins/Build/HordeServer.Build/Devices/DeviceService.cs
--- a/Engine/Source/Programs/Horde/Plugins/Build/HordeServer.Build/Devices/DeviceService.cs
+++ b/Engine/Source/Programs/Horde/Plugins/Build/HordeServer.Build/Devices/DeviceService.cs
@@ -426,6 +426,14 @@ namespace HordeServer.Devices
 			await _devices.CheckoutDeviceAsync(deviceId, userId);
 		}
 
+		/// <summary>
+		/// Clean a device
+		/// </summary>
+		public async Task CleanDeviceAsync(DeviceId deviceId, bool cleaning)
+		{
+			await _devices.UpdateDeviceCleaningAsync(deviceId, cleaning);
+		}
+
 		/// <summary>
 		/// Try to create a reservation satisfying the specified device platforms and models
 		/// </summary>
diff --git a/Engine/Source/Programs/Horde/Plugins/Build/HordeServer.Build/Devices/DevicesController.cs b/Engine/Source/Programs/Horde/Plugins/Build/HordeServer.Build/Devices/DevicesController.cs
--- a/Engine/Source/Programs/Horde/Plugins/Build/HordeServer.Build/Devices/DevicesController.cs
+++ b/Engine/Source/Programs/Horde/Plugins/Build/HordeServer.Build/Devices/DevicesController.cs
@@ -186,7 +186,7 @@ namespace HordeServer.Devices
 				ReservationFinishUtc = u.ReservationFinishUtc
 			}).ToList();
 
-			return new GetDeviceResponse(device.Id.ToString(), device.PlatformId.ToString(), device.PoolId.ToString(), device.Name, device.Enabled, device.Address, device.ModelId?.ToString(), device.ModifiedByUser?.ToString(), device.Notes, device.ProblemTimeUtc, device.MaintenanceTimeUtc, utilization, device.CheckedOutByUser, device.CheckOutTime, checkoutExpiration);
+			return new GetDeviceResponse(device.Id.ToString(), device.PlatformId.ToString(), device.PoolId.ToString(), device.Name, device.Enabled, device.Address, device.ModelId?.ToString(), device.ModifiedByUser?.ToString(), device.Notes, device.ProblemTimeUtc, device.MaintenanceTimeUtc, utilization, device.CheckedOutByUser, device.CheckOutTime, checkoutExpiration, device.CleanRequired, device.CleanStartTimeUtc, device.CleanFinishTimeUtc);
 		}
 
 		/// <summary>
@@ -256,7 +256,6 @@ namespace HordeServer.Devices
 		[HttpPut]
 		[Authorize]
 		[Route("/api/v2/devices/{deviceId}/checkout")]
-		[ProducesResponseType(typeof(List<GetDeviceResponse>), 200)]
 		public async Task<ActionResult> CheckoutDeviceAsync(string deviceId, [FromBody] CheckoutDeviceRequest request, CancellationToken cancellationToken)
 		{
 
@@ -300,6 +299,36 @@ namespace HordeServer.Devices
 			return Ok();
 		}
 
+		/// <summary>
+		/// Clean a device
+		/// </summary>
+		[HttpPut]
+		[Authorize]
+		[Route("/api/v2/devices/{deviceId}/clean")]		
+		public async Task<ActionResult> CleanDeviceAsync(string deviceId, [FromBody] CleanDeviceRequest request)
+		{
+
+			DeviceId deviceIdValue = new DeviceId(deviceId);
+
+			IDevice? device = await _deviceService.GetDeviceAsync(deviceIdValue);
+
+			if (device == null)
+			{
+				return BadRequest($"Device with id ${deviceId} does not exist");
+			}
+
+			DevicePoolAuthorization? poolAuth = _deviceService.GetUserPoolAuthorization(device.PoolId, User, _buildConfig.CurrentValue);
+
+			if (poolAuth == null || !poolAuth.Write)
+			{
+				return Forbid();
+			}
+
+			await _deviceService.CleanDeviceAsync(device.Id, request.Clean);
+
+			return Ok();
+		}
+
 		/// <summary>
 		/// Delete a specific device
 		/// </summary>
diff --git a/Engine/Source/Programs/Horde/Plugins/Build/HordeServer.Build/Devices/IDevice.cs b/Engine/Source/Programs/Horde/Plugins/Build/HordeServer.Build/Devices/IDevice.cs
--- a/Engine/Source/Programs/Horde/Plugins/Build/HordeServer.Build/Devices/IDevice.cs
+++ b/Engine/Source/Programs/Horde/Plugins/Build/HordeServer.Build/Devices/IDevice.cs
@@ -332,6 +332,21 @@ namespace HordeServer.Devices
 		/// </summary>
 		public string? Notes { get; }
 
+		/// <summary>
+		/// Whether the device requires a clean
+		/// </summary>
+		public bool? CleanRequired { get; }
+
+		/// <summary>
+		/// The time device the device cleaning started
+		/// </summary>
+		public DateTime? CleanStartTimeUtc { get; }
+
+		/// <summary>
+		/// The time the device cleaning finished
+		/// </summary>
+		public DateTime? CleanFinishTimeUtc { get; }
+
 		/// <summary>
 		/// Device job utilization history 
 		/// </summary>
diff --git a/Engine/Source/Programs/Horde/Plugins/Build/HordeServer.Build/Devices/IDeviceCollection.cs b/Engine/Source/Programs/Horde/Plugins/Build/HordeServer.Build/Devices/IDeviceCollection.cs
--- a/Engine/Source/Programs/Horde/Plugins/Build/HordeServer.Build/Devices/IDeviceCollection.cs
+++ b/Engine/Source/Programs/Horde/Plugins/Build/HordeServer.Build/Devices/IDeviceCollection.cs
@@ -104,6 +104,15 @@ namespace HordeServer.Devices
 		/// <param name="message">Attached a message to the log of the device</param>
 		Task UpdateDeviceAsync(DeviceId deviceId, DevicePoolId? newPoolId, string? newName, string? newAddress, string? newModelId, string? newNotes, bool? newEnabled, bool? newProblem, bool? newMaintenance, UserId? modifiedByUserId = null, string? message = null);
 
+		/// <summary>
+		/// Updates a device's cleaning status
+		/// </summary>
+		/// <param name="deviceId">The device id to update</param>
+		/// <param name="cleaning">Whether the device is actively being cleaned</param>
+		/// <param name="cleanRequired">Whetner the device requires cleaning, for example after being checked in</param>
+		/// <returns></returns>
+		Task UpdateDeviceCleaningAsync(DeviceId deviceId, bool? cleaning = null, bool? cleanRequired = null);
+
 		/// <summary>
 		/// Delete a device from the collection
 		/// </summary>
diff --git a/Engine/Source/Programs/Shared/EpicGames.Horde/Devices/DeviceMessages.cs b/Engine/Source/Programs/Shared/EpicGames.Horde/Devices/DeviceMessages.cs
--- a/Engine/Source/Programs/Shared/EpicGames.Horde/Devices/DeviceMessages.cs
+++ b/Engine/Source/Programs/Shared/EpicGames.Horde/Devices/DeviceMessages.cs
@@ -357,6 +357,21 @@ namespace EpicGames.Horde.Devices
 		/// </summary>
 		public string? ModifiedByUser { get; set; }
 
+		/// <summary>
+		/// When the clean pass started
+		/// </summary>
+		public DateTime? CleanStartTime { get; set; }
+
+		/// <summary>
+		/// When the clean pass finished
+		/// </summary>
+		public DateTime? CleanFinishTime { get; set; }
+
+		/// <summary>
+		/// Whether the device requires a cleaning
+		/// </summary>
+		public bool? CleanRequired { get; set; }
+
 		/// <summary>
 		///  Device Utilization data
 		/// </summary>
@@ -365,7 +380,7 @@ namespace EpicGames.Horde.Devices
 		/// <summary>
 		/// Device response constructor
 		/// </summary>
-		public GetDeviceResponse(string id, string platformId, string poolId, string name, bool enabled, string? address, string? modelId, string? modifiedByUser, string? notes, DateTime? problemTime, DateTime? maintenanceTime, List<GetDeviceUtilizationResponse>? utilization, string? checkedOutByUser = null, DateTime? checkOutTime = null, DateTime? checkOutExpirationTime = null)
+		public GetDeviceResponse(string id, string platformId, string poolId, string name, bool enabled, string? address, string? modelId, string? modifiedByUser, string? notes, DateTime? problemTime, DateTime? maintenanceTime, List<GetDeviceUtilizationResponse>? utilization, string? checkedOutByUser = null, DateTime? checkOutTime = null, DateTime? checkOutExpirationTime = null, bool? cleanRequired = null, DateTime? cleanStartTime = null, DateTime? cleanFinishTime = null)
 		{
 			Id = id;
 			Name = name;
@@ -382,6 +397,9 @@ namespace EpicGames.Horde.Devices
 			CheckedOutByUserId = checkedOutByUser;
 			CheckOutTime = checkOutTime;
 			CheckOutExpirationTime = checkOutExpirationTime;
+			CleanRequired = cleanRequired;
+			CleanStartTime = cleanStartTime;
+			CleanFinishTime = cleanFinishTime;
 		}
 	}
 
@@ -458,6 +476,17 @@ namespace EpicGames.Horde.Devices
 		public bool Checkout { get; set; }
 	}
 
+	/// <summary>
+	/// Device clean request object
+	/// </summary>
+	public class CleanDeviceRequest
+	{
+		/// <summary>
+		/// Whether to set the device as being cleaned, or not
+		/// </summary>
+		public bool Clean { get; set; }
+	}
+
 	/// <summary>
 	/// Device reservation request object
 	/// </summary>
