Horde: Auto assign Horde issues to Jira owner when possible

diff --git a/Engine/Source/Programs/Horde/Plugins/Build/HordeServer.Build/Devices/DeviceCollection.cs b/Engine/Source/Programs/Horde/Plugins/Build/HordeServer.Build/Devices/DeviceCollection.cs
--- a/Engine/Source/Programs/Horde/Plugins/Build/HordeServer.Build/Devices/DeviceCollection.cs
+++ b/Engine/Source/Programs/Horde/Plugins/Build/HordeServer.Build/Devices/DeviceCollection.cs
@@ -774,6 +774,7 @@ namespace HordeServer.Devices
 			}
 		}
 
+		/// <inheritdoc/>
 		public async Task UpdateDeviceCleaningAsync(DeviceId deviceId, bool? cleaning = null, bool? cleanRequired = null)
 		{
 			UpdateDefinitionBuilder<DeviceDocument> updateBuilder = Builders<DeviceDocument>.Update;
diff --git a/Engine/Source/Programs/Horde/Plugins/Build/HordeServer.Build/Issues/External/JiraService.cs b/Engine/Source/Programs/Horde/Plugins/Build/HordeServer.Build/Issues/External/JiraService.cs
--- a/Engine/Source/Programs/Horde/Plugins/Build/HordeServer.Build/Issues/External/JiraService.cs
+++ b/Engine/Source/Programs/Horde/Plugins/Build/HordeServer.Build/Issues/External/JiraService.cs
@@ -135,7 +135,7 @@ namespace HordeServer.Issues.External
 			_client = new HttpClient();
 			byte[] authBytes = Encoding.ASCII.GetBytes($"{_settings.JiraUsername}:{_settings.JiraApiToken}");
 			_client.DefaultRequestHeaders.Authorization = new AuthenticationHeaderValue("Basic", Convert.ToBase64String(authBytes));
-			_client.Timeout = TimeSpan.FromSeconds(30.0);
+			_client.Timeout = TimeSpan.FromSeconds(20.0);
 			_retryPolicy = HttpPolicyExtensions.HandleTransientHttpError().WaitAndRetryAsync(3, attempt => TimeSpan.FromSeconds(Math.Pow(2.0, attempt)));
 		}
 
@@ -440,7 +440,7 @@ namespace HordeServer.Issues.External
 			{
 				Uri uri = new Uri(_jiraUrl, $"/rest/api/2/search?jql=issueKey%20in%20({String.Join(",", queryJiras)})&fields=assignee,status,resolution,priority&maxResults={queryJiras.Count}");
 
-				HttpResponseMessage response = await _retryPolicy.ExecuteAsync(ctx => _client.GetAsync(uri, ctx), cancellationToken);
+				HttpResponseMessage response = await _client.GetAsync(uri, cancellationToken);
 				if (!response.IsSuccessStatusCode)
 				{
 					_logger.LogWarning("GET to {Uri} returned {Code} ({Response})", uri, response.StatusCode, await response.Content.ReadAsStringAsync(cancellationToken));
diff --git a/Engine/Source/Programs/Horde/Plugins/Build/HordeServer.Build/Issues/IssueService.cs b/Engine/Source/Programs/Horde/Plugins/Build/HordeServer.Build/Issues/IssueService.cs
--- a/Engine/Source/Programs/Horde/Plugins/Build/HordeServer.Build/Issues/IssueService.cs
+++ b/Engine/Source/Programs/Horde/Plugins/Build/HordeServer.Build/Issues/IssueService.cs
@@ -15,6 +15,7 @@ using EpicGames.Horde.Logs;
 using EpicGames.Horde.Streams;
 using EpicGames.Horde.Users;
 using HordeServer.Commits;
+using HordeServer.Issues.External;
 using HordeServer.Jobs;
 using HordeServer.Logs;
 using HordeServer.Streams;
@@ -198,6 +199,7 @@ namespace HordeServer.Issues
 		readonly JobService _jobService;
 		readonly IJobStepRefCollection _jobStepRefs;
 		readonly IIssueCollection _issueCollection;
+		readonly IExternalIssueService _externalIssueService;
 		readonly ICommitService _commitService;
 		readonly IStreamCollection _streams;
 		readonly IUserCollection _userCollection;
@@ -255,7 +257,7 @@ namespace HordeServer.Issues
 		/// <summary>
 		/// Constructor
 		/// </summary>
-		public IssueService(JobService jobService, IIssueCollection issueCollection, ICommitService commitService, IJobStepRefCollection jobStepRefs, IStreamCollection streams, IUserCollection userCollection, ILogCollection logCollection, IClock clock, IOptionsMonitor<BuildConfig> buildConfig, Tracer tracer, ILogger<IssueService> logger)
+		public IssueService(JobService jobService, IIssueCollection issueCollection, IExternalIssueService externalIssueService, ICommitService commitService, IJobStepRefCollection jobStepRefs, IStreamCollection streams, IUserCollection userCollection, ILogCollection logCollection, IClock clock, IOptionsMonitor<BuildConfig> buildConfig, Tracer tracer, ILogger<IssueService> logger)
 		{
 			Type[] issueTypes = Assembly.GetExecutingAssembly().GetTypes().Where(x => !x.IsAbstract && typeof(IIssue).IsAssignableFrom(x)).ToArray();
 			foreach (Type issueType in issueTypes)
@@ -266,6 +268,7 @@ namespace HordeServer.Issues
 			// Get all the collections
 			_jobService = jobService;
 			_issueCollection = issueCollection;
+			_externalIssueService = externalIssueService;
 			_commitService = commitService;
 			_jobStepRefs = jobStepRefs;
 			_streams = streams;
@@ -498,6 +501,7 @@ namespace HordeServer.Issues
 			using TelemetrySpan tracerSpan = _tracer.StartActiveSpan($"{nameof(IssueService)}.{nameof(UpdateIssueAsync)}");
 			
 			IIssue? issue;
+			bool autoAssign = ownerId == null && _externalIssueService != null && !String.IsNullOrEmpty(externalIssueKey);
 			for (; ; )
 			{
 				issue = await _issueCollection.GetIssueAsync(id, cancellationToken);
@@ -506,6 +510,38 @@ namespace HordeServer.Issues
 					return false;
 				}
 
+				if (issue.OwnerId != null || issue.ExternalIssueKey == externalIssueKey)
+				{
+					autoAssign = false;
+				}
+
+				try
+				{
+					// Auto assign to external issue owner if possible
+					if (autoAssign)
+					{
+						// Get the external issue and auto assign					
+						List<IExternalIssue> externalIssues = await _externalIssueService!.GetIssuesAsync(new string[] { externalIssueKey! }, cancellationToken);
+						if (externalIssues.Count > 0 && !String.IsNullOrEmpty(externalIssues[0].AssigneeEmailAddress))
+						{
+							IUser? matchedUser = await _userCollection.FindUserByEmailAsync(externalIssues[0].AssigneeEmailAddress!, cancellationToken);
+							if (matchedUser != null)
+							{
+								ownerId = matchedUser.Id;
+								_logger.LogInformation("Assigning issue owner to {OwnerId} for external issue {ExternalIssueKey} by matching with {UserEmail}", ownerId, externalIssueKey, externalIssues[0].AssigneeEmailAddress);
+							}
+						}
+					}
+				}
+				catch (Exception ex)
+				{
+					_logger.LogWarning(ex, "Exception while auto-assigning issue {IssueId} to {ExternalId}: {ExceptionMessage}", issue.Id, externalIssueKey, ex.Message);
+				}
+				finally
+				{
+					autoAssign = false;
+				}
+
 				issue = await _issueCollection.TryUpdateIssueAsync(issue, initiatedById, new UpdateIssueOptions { UserSummary = summary, Description = description, Promoted = promoted, OwnerId = ownerId ?? resolvedById, NominatedById = nominatedById, DeclinedById = declinedById, Acknowledged = acknowledged, FixCommitId = fixCommitId, FixSystemic = fixSystemic, ResolvedById = resolvedById, ExcludeSpanIds = removeSpanIds, ExternalIssueKey = externalIssueKey, QuarantinedById = quarantinedById, ForceClosedById = forceClosedById, WorkflowThreadUrl = workflowThreadUrl }, cancellationToken: cancellationToken);
 				if (issue != null)
 				{
