Horde: Add some logging to debug Horde issue => jira auto assigning

diff --git a/Engine/Source/Programs/Horde/Plugins/Build/HordeServer.Build/Issues/IssueService.cs b/Engine/Source/Programs/Horde/Plugins/Build/HordeServer.Build/Issues/IssueService.cs
--- a/Engine/Source/Programs/Horde/Plugins/Build/HordeServer.Build/Issues/IssueService.cs
+++ b/Engine/Source/Programs/Horde/Plugins/Build/HordeServer.Build/Issues/IssueService.cs
@@ -530,6 +530,21 @@ namespace HordeServer.Issues
 								ownerId = matchedUser.Id;
 								_logger.LogInformation("Assigning issue owner to {OwnerId} for external issue {ExternalIssueKey} by matching with {UserEmail}", ownerId, externalIssueKey, externalIssues[0].AssigneeEmailAddress);
 							}
+							else
+							{
+								_logger.LogInformation("Unable to match owner for external issue {ExternalIssueKey} by matching with {UserEmail}", externalIssueKey, externalIssues[0].AssigneeEmailAddress);
+							}
+						}
+						else
+						{
+							if (externalIssues.Count == 0)
+							{
+								_logger.LogInformation("Unable to get external issue {ExternalIssueKey} for auto assign", externalIssueKey);
+							}
+							else
+							{
+								_logger.LogInformation("Unable to get external issue {ExternalIssueKey} email address for auto assign", externalIssueKey);
+							}							
 						}
 					}
 				}
