Horde: Disable injecting jira issue service until can sort it properly

diff --git a/Engine/Source/Programs/Horde/Plugins/Build/HordeServer.Build/Issues/IssueService.cs b/Engine/Source/Programs/Horde/Plugins/Build/HordeServer.Build/Issues/IssueService.cs
--- a/Engine/Source/Programs/Horde/Plugins/Build/HordeServer.Build/Issues/IssueService.cs
+++ b/Engine/Source/Programs/Horde/Plugins/Build/HordeServer.Build/Issues/IssueService.cs
@@ -199,7 +199,7 @@ namespace HordeServer.Issues
 		readonly JobService _jobService;
 		readonly IJobStepRefCollection _jobStepRefs;
 		readonly IIssueCollection _issueCollection;
-		readonly IExternalIssueService _externalIssueService;
+		//readonly IExternalIssueService _externalIssueService;
 		readonly ICommitService _commitService;
 		readonly IStreamCollection _streams;
 		readonly IUserCollection _userCollection;
@@ -257,7 +257,7 @@ namespace HordeServer.Issues
 		/// <summary>
 		/// Constructor
 		/// </summary>
-		public IssueService(JobService jobService, IIssueCollection issueCollection, IExternalIssueService externalIssueService, ICommitService commitService, IJobStepRefCollection jobStepRefs, IStreamCollection streams, IUserCollection userCollection, ILogCollection logCollection, IClock clock, IOptionsMonitor<BuildConfig> buildConfig, Tracer tracer, ILogger<IssueService> logger)
+		public IssueService(JobService jobService, IIssueCollection issueCollection, ICommitService commitService, IJobStepRefCollection jobStepRefs, IStreamCollection streams, IUserCollection userCollection, ILogCollection logCollection, IClock clock, IOptionsMonitor<BuildConfig> buildConfig, Tracer tracer, ILogger<IssueService> logger)
 		{
 			Type[] issueTypes = Assembly.GetExecutingAssembly().GetTypes().Where(x => !x.IsAbstract && typeof(IIssue).IsAssignableFrom(x)).ToArray();
 			foreach (Type issueType in issueTypes)
@@ -268,7 +268,7 @@ namespace HordeServer.Issues
 			// Get all the collections
 			_jobService = jobService;
 			_issueCollection = issueCollection;
-			_externalIssueService = externalIssueService;
+			//_externalIssueService = externalIssueService;
 			_commitService = commitService;
 			_jobStepRefs = jobStepRefs;
 			_streams = streams;
@@ -501,7 +501,9 @@ namespace HordeServer.Issues
 			using TelemetrySpan tracerSpan = _tracer.StartActiveSpan($"{nameof(IssueService)}.{nameof(UpdateIssueAsync)}");
 			
 			IIssue? issue;
-			bool autoAssign = ownerId == null && _externalIssueService != null && !String.IsNullOrEmpty(externalIssueKey);
+			//bool autoAssign = ownerId == null && _externalIssueService != null && !String.IsNullOrEmpty(externalIssueKey);
+
+			//_logger.LogInformation("Updating issue {ArgOne} {ArgTwo} {ArgThree}", ownerId == null,  _externalIssueService != null , !String.IsNullOrEmpty(externalIssueKey));
 			for (; ; )
 			{
 				issue = await _issueCollection.GetIssueAsync(id, cancellationToken);
@@ -510,9 +512,11 @@ namespace HordeServer.Issues
 					return false;
 				}
 
+				/*
 				if (issue.OwnerId != null || issue.ExternalIssueKey == externalIssueKey)
 				{
 					autoAssign = false;
+					_logger.LogInformation("Updating issue {ArgOne} {ArgTwo} external: {ArgThree}", issue.OwnerId != null, issue.ExternalIssueKey == externalIssueKey, externalIssueKey);
 				}
 
 				try
@@ -556,6 +560,7 @@ namespace HordeServer.Issues
 				{
 					autoAssign = false;
 				}
+				*/
 
 				issue = await _issueCollection.TryUpdateIssueAsync(issue, initiatedById, new UpdateIssueOptions { UserSummary = summary, Description = description, Promoted = promoted, OwnerId = ownerId ?? resolvedById, NominatedById = nominatedById, DeclinedById = declinedById, Acknowledged = acknowledged, FixCommitId = fixCommitId, FixSystemic = fixSystemic, ResolvedById = resolvedById, ExcludeSpanIds = removeSpanIds, ExternalIssueKey = externalIssueKey, QuarantinedById = quarantinedById, ForceClosedById = forceClosedById, WorkflowThreadUrl = workflowThreadUrl }, cancellationToken: cancellationToken);
 				if (issue != null)
