Horde: Fix warning in Issue Service
#rnx

diff --git a/Engine/Source/Programs/Horde/Plugins/Build/HordeServer.Build/Issues/IssueService.cs b/Engine/Source/Programs/Horde/Plugins/Build/HordeServer.Build/Issues/IssueService.cs
--- a/Engine/Source/Programs/Horde/Plugins/Build/HordeServer.Build/Issues/IssueService.cs
+++ b/Engine/Source/Programs/Horde/Plugins/Build/HordeServer.Build/Issues/IssueService.cs
@@ -15,7 +15,6 @@ using EpicGames.Horde.Logs;
 using EpicGames.Horde.Streams;
 using EpicGames.Horde.Users;
 using HordeServer.Commits;
-using HordeServer.Issues.External;
 using HordeServer.Jobs;
 using HordeServer.Logs;
 using HordeServer.Streams;
