Horde: Explain GitHub image registry auth better in docs

diff --git a/Engine/Source/Programs/Horde/Docs/Deployment/Server.md b/Engine/Source/Programs/Horde/Docs/Deployment/Server.md
--- a/Engine/Source/Programs/Horde/Docs/Deployment/Server.md
+++ b/Engine/Source/Programs/Horde/Docs/Deployment/Server.md
@@ -18,14 +18,11 @@ associated with an EpicGames account to follow these links.
 
 * [Horde Server]([REDACTED_PROFILE_URL]/EpicGames/packages/container/package/horde-server)
 
-To download an image, first create a GitHub personal access token (PAT) from the developer section
-of your account settings page and pass it as the password to:
+To download an image, first create a [GitHub personal access token (classic)](https://docs.github.com/en/packages/working-with-a-github-packages-registry/working-with-the-container-registry#authenticating-with-a-personal-access-token-classic) from the developer section
+of your account settings page and use that to login with GitHub via `docker`.
+Follow the linked GitHub documentation for more details.
 
-```bash
-docker login ghcr.io
-```
-
-To download the image:
+Once authenticated, you can verify by downloading an image with:
 
 ```bash
 docker pull ghcr.io/epicgames/horde-server:latest 
@@ -34,8 +31,8 @@ docker pull ghcr.io/epicgames/horde-server:latest
 Note that in this form, an external MongoDB and Redis instance must be configured through a configuration file or
 environment variable (see below).
 
-Running multiple Horde servers behind a load balancer does not require explicit configuration as long as each
-server points to the same MongoDB and Redis instance.
+Running multiple Horde servers behind a load balancer is possible but remember to set `ServerUrl` in settings.
+That setting should reflect the URL your users and agents use to connect with.
 
 Using Docker containers on Linux is [Epic's preferred way of running Horde](../Deployment.md).
 
