Horde: Update release notes
#rnx

diff --git a/Engine/Source/Programs/Horde/Docs/ReleaseNotes.md b/Engine/Source/Programs/Horde/Docs/ReleaseNotes.md
--- a/Engine/Source/Programs/Horde/Docs/ReleaseNotes.md
+++ b/Engine/Source/Programs/Horde/Docs/ReleaseNotes.md
@@ -1,5 +1,17 @@
 # Release Notes
 
+## 2025-02-20
+
+* Explain GitHub image registry auth better in docs (40164690)
+* Add -snapshot parameter to BuildCookRun, which will retrieve the most recent snapshot published to Jupiter and cook incrementally on top of the fetched data. (40156692)
+* Auto assign Horde issues to Jira owner when possible (40147283)
+* Add support for automated device cleaning on CI and shared kits, server side changes (40143179)
+* Extend compute socket test timeouts to eliminate flakiness Tests were intermittently failing when CI resources were under load. (40137698)
+* Allow service accounts to modify devices (40134987)
+* Added verification that interfaces have matching native interface base classes (40091963)
+* Warn when OIDC is configured with confidential client as this will prevent command-line auth (40060901)
+* Cloud DDC - Fixed exception when a object id is encountered as we convert it to json (39986461)
+
 ## 2025-02-13
 
 * Return explicit 500 status for VCS operation failures in create job REST endpoint (39941708)
