Added missing UBT/UBA Cache Rate Metrics

diff --git a/Engine/Source/Programs/Horde/HordeServer/Defaults/default-metrics.telemetry.json b/Engine/Source/Programs/Horde/HordeServer/Defaults/default-metrics.telemetry.json
--- a/Engine/Source/Programs/Horde/HordeServer/Defaults/default-metrics.telemetry.json
+++ b/Engine/Source/Programs/Horde/HordeServer/Defaults/default-metrics.telemetry.json
@@ -804,6 +804,14 @@
 			"function": "Average",
 			"interval": "00:30:00"
 		},
+		{
+			"id": "ubt-uba-cache-rate",
+			"filter": "$[?( $(IsUBAExecutor) )]",
+			"groupBy": "$.Payload.Build_BranchName",
+			"property": "$.Payload.CacheRate",
+			"function": "Average",
+			"interval": "00:30:00"
+		},
 		{
 			"id": "ubt-completed-duration",
 			"filter": "$[?(@.Payload.EventName == 'Core.UBT.Completed' )]",
