Fix for UBA Cache Rate Average and UBA Cache Rate All metrics

diff --git a/Engine/Source/Programs/Horde/HordeServer/Defaults/default-metrics.telemetry.json b/Engine/Source/Programs/Horde/HordeServer/Defaults/default-metrics.telemetry.json
--- a/Engine/Source/Programs/Horde/HordeServer/Defaults/default-metrics.telemetry.json
+++ b/Engine/Source/Programs/Horde/HordeServer/Defaults/default-metrics.telemetry.json
@@ -748,6 +748,14 @@
 			"function": "Average",
 			"interval": "00:30:00"
 		},
+		{
+			"id": "ubt-uba-cache-rate-avg",
+			"filter": "$[?( $(IsUBAExecutor) )]",
+			"groupBy": "$.Payload.Build_BranchName,$.Payload.Target,$.Payload.Mode,$.Payload.Configuration,$.Payload.Platform",
+			"property": "$.Payload.CacheRate",
+			"function": "Average",
+			"interval": "00:30:00"
+		},
 		{
 			"id": "ubt-uba-local-actions-avg",
 			"filter": "$[?( $(IsUBAExecutor) )]",
@@ -805,7 +813,7 @@
 			"interval": "00:30:00"
 		},
 		{
-			"id": "ubt-uba-cache-rate",
+			"id": "ubt-uba-cache-rate-all",
 			"filter": "$[?( $(IsUBAExecutor) )]",
 			"groupBy": "$.Payload.Build_BranchName",
 			"property": "$.Payload.CacheRate",
diff --git a/Engine/Source/Programs/Horde/HordeServer/Defaults/default-views.telemetry.json b/Engine/Source/Programs/Horde/HordeServer/Defaults/default-views.telemetry.json
--- a/Engine/Source/Programs/Horde/HordeServer/Defaults/default-views.telemetry.json
+++ b/Engine/Source/Programs/Horde/HordeServer/Defaults/default-views.telemetry.json
@@ -634,7 +634,7 @@
 							"display": "Ratio",
 							"metrics": [
 								{
-									"id": "ubt-uba-cache-rate"
+									"id": "ubt-uba-cache-rate-avg"
 								}
 							]
 						},
