Horde: Improve logging in agent relay client for port mappings

diff --git a/Engine/Source/Programs/Horde/Plugins/Build/HordeServer.Build.Tests/Agents/Relay/AgentRelayClientTest.cs b/Engine/Source/Programs/Horde/Plugins/Build/HordeServer.Build.Tests/Agents/Relay/AgentRelayClientTest.cs
--- a/Engine/Source/Programs/Horde/Plugins/Build/HordeServer.Build.Tests/Agents/Relay/AgentRelayClientTest.cs
+++ b/Engine/Source/Programs/Horde/Plugins/Build/HordeServer.Build.Tests/Agents/Relay/AgentRelayClientTest.cs
@@ -1,11 +1,12 @@
 // Copyright Epic Games, Inc. All Rights Reserved.
 
+using EpicGames.Core;
 using Grpc.Core;
 using Horde.Common.Rpc;
 using HordeServer.Agents.Relay;
-//using HordeServer.Tests.Server;
 using HordeServer.Utilities;
-using Microsoft.Extensions.Logging.Abstractions;
+using Microsoft.Extensions.Logging;
+using static HordeServer.Tests.Agents.Relay.NftablesTests;
 
 namespace HordeServer.Tests.Agents.Relay;
 
@@ -31,7 +32,9 @@ public class AgentRelayClientTest
 
 	public AgentRelayClientTest()
 	{
-		_relayClient = new AgentRelayClient("myCluster", "myAgent", new List<string> { "192.168.1.99" }, Nftables.CreateNull(), _grpcClient, NullLogger<AgentRelayClient>.Instance);
+		using ILoggerFactory loggerFactory = LoggerFactory.Create(builder => builder.AddConsole());
+		ILogger<AgentRelayClient> logger = loggerFactory.CreateLogger<AgentRelayClient>();
+		_relayClient = new AgentRelayClient("myCluster", "myAgent", new List<string> { "192.168.1.99" }, Nftables.CreateNull(), _grpcClient, new DefaultClock(), logger);
 		_relayClient.CooldownOnException = TimeSpan.Zero;
 	}
 
@@ -39,8 +42,8 @@ public class AgentRelayClientTest
 	public async Task GetPortMappingsLongPoll_Results_Async()
 	{
 		using CancellationTokenSource cts = new(3000);
-		_grpcClient.GetPortMappingsResponses.AddMessage(new GetPortMappingsResponse { PortMappings = { NftablesTests.LeaseMap1 }, RevisionCount = 1 });
-		_grpcClient.GetPortMappingsResponses.AddMessage(new GetPortMappingsResponse { PortMappings = { NftablesTests.LeaseMap2, NftablesTests.LeaseMap1 }, RevisionCount = 2 });
+		_grpcClient.GetPortMappingsResponses.AddMessage(new GetPortMappingsResponse { PortMappings = { LeaseMap1 }, RevisionCount = 1 });
+		_grpcClient.GetPortMappingsResponses.AddMessage(new GetPortMappingsResponse { PortMappings = { LeaseMap2, LeaseMap1 }, RevisionCount = 2 });
 
 		List<PortMapping>? pm1 = await _relayClient.GetPortMappingsLongPollAsync(cts.Token);
 		Assert.AreEqual(1, pm1!.Count);
@@ -68,10 +71,41 @@ public class AgentRelayClientTest
 
 		Task t = _relayClient.ListenForPortMappingsAsync(cts.Token);
 		await Task.Delay(100, cts.Token);
-		_grpcClient.GetPortMappingsResponses.AddMessage(new GetPortMappingsResponse { PortMappings = { NftablesTests.LeaseMap1 } });
+		_grpcClient.GetPortMappingsResponses.AddMessage(new GetPortMappingsResponse { PortMappings = { LeaseMap1 } });
 		await Task.Delay(100, cts.Token);
-		_grpcClient.GetPortMappingsResponses.AddMessage(new GetPortMappingsResponse { PortMappings = { NftablesTests.LeaseMap2 } });
+		_grpcClient.GetPortMappingsResponses.AddMessage(new GetPortMappingsResponse { PortMappings = { LeaseMap2 } });
 
 		await t;
 	}
+	
+	[TestMethod]
+	public void ComparePortMappings()
+	{
+		IReadOnlySet<PortMapping> added;
+		IReadOnlySet<PortMapping> removed;
+		
+		(added, removed) = AgentRelayClient.ComparePortMappings(null, [LeaseMap1]);
+		Assert.AreEqual(1, added.Count);
+		Assert.AreEqual(0, removed.Count);
+		
+		(added, removed) = AgentRelayClient.ComparePortMappings([], [LeaseMap1]);
+		Assert.AreEqual(1, added.Count);
+		Assert.AreEqual(0, removed.Count);
+		
+		(added, removed) = AgentRelayClient.ComparePortMappings([LeaseMap1, LeaseMap2], [LeaseMap1, LeaseMap2]);
+		Assert.AreEqual(0, added.Count);
+		Assert.AreEqual(0, removed.Count);
+		
+		(added, removed) = AgentRelayClient.ComparePortMappings([LeaseMap1, LeaseMap2], [LeaseMap2, LeaseMap3]);
+		Assert.AreEqual(1, added.Count);
+		Assert.AreEqual(1, removed.Count);
+		Assert.IsTrue(added.Contains(LeaseMap3));
+		Assert.IsTrue(removed.Contains(LeaseMap1));
+		
+		(added, removed) = AgentRelayClient.ComparePortMappings([LeaseMap1, LeaseMap1], [LeaseMap2, LeaseMap2]);
+		Assert.AreEqual(1, added.Count);
+		Assert.AreEqual(1, removed.Count);
+		Assert.IsTrue(added.Contains(LeaseMap2));
+		Assert.IsTrue(removed.Contains(LeaseMap1));
+	}
 }
diff --git a/Engine/Source/Programs/Horde/Plugins/Build/HordeServer.Build/Commands/RelayCommand.cs b/Engine/Source/Programs/Horde/Plugins/Build/HordeServer.Build/Commands/RelayCommand.cs
--- a/Engine/Source/Programs/Horde/Plugins/Build/HordeServer.Build/Commands/RelayCommand.cs
+++ b/Engine/Source/Programs/Horde/Plugins/Build/HordeServer.Build/Commands/RelayCommand.cs
@@ -94,7 +94,7 @@ class RelayCommand : Command
 
 		using GrpcChannel channel = GrpcChannel.ForAddress(ServerUrl);
 		RelayRpc.RelayRpcClient relayGrpcClient = new(channel);
-		AgentRelayClient relayClient = new(ClusterId, AgentId, listenIps.ToList(), nftables, relayGrpcClient, logger);
+		AgentRelayClient relayClient = new(ClusterId, AgentId, listenIps.ToList(), nftables, relayGrpcClient, new DefaultClock(), logger);
 
 		logger.LogInformation("Listening for port mappings...");
 		await relayClient.ListenForPortMappingsAsync(cts.Token);
diff --git a/Engine/Source/Programs/Horde/Plugins/Compute/HordeServer.Compute/Agents/Relay/AgentRelayClient.cs b/Engine/Source/Programs/Horde/Plugins/Compute/HordeServer.Compute/Agents/Relay/AgentRelayClient.cs
--- a/Engine/Source/Programs/Horde/Plugins/Compute/HordeServer.Compute/Agents/Relay/AgentRelayClient.cs
+++ b/Engine/Source/Programs/Horde/Plugins/Compute/HordeServer.Compute/Agents/Relay/AgentRelayClient.cs
@@ -1,5 +1,7 @@
 ﻿// Copyright Epic Games, Inc. All Rights Reserved.
 
+using System.Diagnostics.CodeAnalysis;
+using EpicGames.Core;
 using Grpc.Core;
 using Horde.Common.Rpc;
 using Microsoft.Extensions.Logging;
@@ -28,7 +30,9 @@ public class AgentRelayClient
 	private readonly List<string> _ipAddresses;
 	private readonly Nftables _nftables;
 	private readonly RelayRpc.RelayRpcClient _relayRpcClient;
+	private readonly IClock _clock;
 	private readonly ILogger _logger;
+	private List<PortMapping>? _lastPortMappings; 
 
 	/// <summary>
 	/// Constructor
@@ -37,15 +41,17 @@ public class AgentRelayClient
 	/// <param name="agentId">Unique ID of this relay agent</param>
 	/// <param name="ipAddresses">IP addresses this relay agent is listening on</param>
 	/// <param name="nftables"></param>
+	/// <param name="clock"></param>
 	/// <param name="relayRpcClient"></param>
 	/// <param name="logger"></param>
-	public AgentRelayClient(string clusterId, string agentId, List<string> ipAddresses, Nftables nftables, RelayRpc.RelayRpcClient relayRpcClient, ILogger logger)
+	public AgentRelayClient(string clusterId, string agentId, List<string> ipAddresses, Nftables nftables, RelayRpc.RelayRpcClient relayRpcClient, IClock clock, ILogger logger)
 	{
 		_clusterId = clusterId;
 		_agentId = agentId;
 		_ipAddresses = ipAddresses;
 		_nftables = nftables;
 		_relayRpcClient = relayRpcClient;
+		_clock = clock;
 		_logger = logger;
 	}
 
@@ -92,7 +98,9 @@ public class AgentRelayClient
 				if (portMappings != null)
 				{
 					_logger.LogDebug("Received {NumMappings} port mappings...", portMappings.Count);
+					LogPortMappingChanges(portMappings);
 					await _nftables.ApplyPortForwardingAsync(portMappings);
+					_lastPortMappings = portMappings;
 				}
 			}
 			catch (OperationCanceledException) when (cancellationToken.IsCancellationRequested)
@@ -110,4 +118,57 @@ public class AgentRelayClient
 			}
 		}
 	}
+	
+	private void LogPortMappingChanges(List<PortMapping> portMappings)
+	{
+		(IReadOnlySet<PortMapping> added, IReadOnlySet<PortMapping> removed) = ComparePortMappings(_lastPortMappings, portMappings);
+		
+		[SuppressMessage("Usage", "CA2254:Template should be a static expression")]
+		void Log(string prefix, PortMapping pm)
+		{
+			string ports = String.Join(", ", pm.Ports.Select(p => $"{p.Protocol} {p.RelayPort} -> {p.AgentPort}"));
+			string sourceIps = pm.AllowedSourceIps.Count switch
+			{
+				0 => "*",
+				1 => pm.AllowedSourceIps.First(),
+				_ => $"[{String.Join(", ", pm.AllowedSourceIps)}]"
+			};
+			string? createdAt = pm.CreatedAt?.ToDateTime().ToString("o");
+			long? ageMs = null;
+			if (pm.CreatedAt != null)
+			{
+				ageMs = (long)(_clock.UtcNow - pm.CreatedAt.ToDateTime()).TotalMilliseconds;
+			}
+			_logger.LogInformation(prefix + " {LeaseId}: {SourceIps} to {AgentIp} [{Ports}] ({CreatedAt} - {AgeMs} ms old)",
+				pm.LeaseId, sourceIps, pm.AgentIp, ports, createdAt, ageMs);
+		}
+		
+		added.ToList().ForEach(pm => Log("Added", pm));
+		removed.ToList().ForEach(pm => Log("Removed", pm));
+	}
+	
+	private class PortMappingComparer : IEqualityComparer<PortMapping>
+	{
+		public bool Equals(PortMapping? x, PortMapping? y)
+		{
+			return x?.LeaseId == y?.LeaseId;
+		}
+		
+		public int GetHashCode(PortMapping obj)
+		{
+			return obj.LeaseId.GetHashCode(StringComparison.Ordinal);
+		}
+	}
+
+	internal static (IReadOnlySet<PortMapping> added, IReadOnlySet<PortMapping> removed) ComparePortMappings(List<PortMapping>? oldMappings, List<PortMapping> newMappings)
+	{
+		PortMappingComparer comparer = new ();
+		HashSet<PortMapping> oldMappingSet = new(oldMappings ?? [], comparer);
+		HashSet<PortMapping> newMappingSet = new(newMappings, comparer);
+		
+		return (
+			added: new HashSet<PortMapping>(newMappingSet.Except(oldMappingSet, comparer), comparer),
+			removed: new HashSet<PortMapping>(oldMappingSet.Except(newMappingSet, comparer), comparer)
+		);
+	}
 }
\ No newline at end of file
