Horde: Include requester's public IP in trace spans for compute lease allocations

diff --git a/Engine/Source/Programs/Horde/Plugins/Compute/HordeServer.Compute/Compute/ComputeService.cs b/Engine/Source/Programs/Horde/Plugins/Compute/HordeServer.Compute/Compute/ComputeService.cs
--- a/Engine/Source/Programs/Horde/Plugins/Compute/HordeServer.Compute/Compute/ComputeService.cs
+++ b/Engine/Source/Programs/Horde/Plugins/Compute/HordeServer.Compute/Compute/ComputeService.cs
@@ -370,6 +370,7 @@ namespace HordeServer.Compute
 			span.SetAttribute("userId", arp.UserId?.ToString());
 			span.SetAttribute("requestId", arp.RequestId);
 			span.SetAttribute("requestIp", arp.RequesterIp?.ToString());
+			span.SetAttribute("requesterPublicIp", arp.RequesterPublicIp);
 			span.SetAttribute("connectionMode", arp.ConnectionMode?.ToString());
 			span.SetAttribute("usePublicIp", arp.UsePublicIp.ToString());
 			span.SetAttribute("encryption", arp.Encryption.ToString());
