Horde: Fix analytics prefix issue by matching whole words

diff --git a/Engine/Source/Programs/Horde/HordeDashboard/plugins/analytics/telemetryView.tsx b/Engine/Source/Programs/Horde/HordeDashboard/plugins/analytics/telemetryView.tsx
--- a/Engine/Source/Programs/Horde/HordeDashboard/plugins/analytics/telemetryView.tsx
+++ b/Engine/Source/Programs/Horde/HordeDashboard/plugins/analytics/telemetryView.tsx
@@ -205,13 +205,53 @@ class MetricsHandler {
       function getPrefix(words: string[]) {
 
          if (!words[0] || words.length === 1) return words[0] || "";
-         let i = 0;
 
-         // eslint-disable-next-line no-loop-func
-         while (words[0][i] && words.every(w => w[i] === words[0][i]))
-            i++;
+         const elementMap = new Map<number, Set<string>>();
+
+         const length = words[0].split(":").length;
+
+         if (length <= 1)
+         {
+            return "";
+         }
+
+         if (words.find(w => w.split(":").length !== length))
+         {
+            return "";
+         }
+
+         for (let i = 0; i < length; i++)
+         {
+            elementMap.set(i, new Set);
+         }
+
+         words.forEach(w => {
+
+            const elements = w.split(":");
+
+            elements.forEach((element, index) => {
+               elementMap.get(index)!.add(element);
+            })
+         })
+
+         let replace:string[] = [];         
+         for (let count = 0; count < length - 1; count++)
+         {
+            if (elementMap.get(count)!.size > 1)
+            {
+               break;
+            }
+
+            replace.push(Array.from(elementMap.get(count)!)[0])            
+         }         
+
+         let replaceValue = replace.join(":");
+         if (replaceValue) {
+            replaceValue += ":";
+         }
+
+         return replaceValue;
 
-         return words[0].slice(0, i);
       }
 
       // collapse to common value is left most variable only has one value
@@ -224,9 +264,8 @@ class MetricsHandler {
          })
 
          if (v && legend.length > 0 && v.split(";").length === 2) {
-            const display = legend.map(v => v.display);
+            const display = legend.map(v => v.display);            
             const prefix = getPrefix(display);
-
             if (prefix) {
                legend.forEach(v => {
                   const ndisplay = v.display.replace(prefix, "");
