Fixing Slack thread message handling where some callstacks can get quite large causing Slack to throw an exception with messages reaching over 3001 characters.

#rnx

diff --git a/Engine/Source/Programs/Horde/Plugins/Experimental/HordeServer.Experimental/Notifications/Sinks/ExperimentalSlackNotificationSink.cs b/Engine/Source/Programs/Horde/Plugins/Experimental/HordeServer.Experimental/Notifications/Sinks/ExperimentalSlackNotificationSink.cs
--- a/Engine/Source/Programs/Horde/Plugins/Experimental/HordeServer.Experimental/Notifications/Sinks/ExperimentalSlackNotificationSink.cs
+++ b/Engine/Source/Programs/Horde/Plugins/Experimental/HordeServer.Experimental/Notifications/Sinks/ExperimentalSlackNotificationSink.cs
@@ -145,6 +145,11 @@ namespace HordeServer.Notifications.Sinks
 		/// </summary>
 		const int MaxEventMessages = 5;
 
+		/// <summary>
+		/// Length limit for log event to avoid character limit from Slack
+		/// </summary>
+		const int MaxLogEventLength = 2000;
+
 		/// <summary>
 		/// Constructor
 		/// </summary>
@@ -970,7 +975,7 @@ namespace HordeServer.Notifications.Sinks
 		/// <param name="threadMessage"><see cref="SlackMessage"/> to be sent</param>
 		/// <param name="cancellationToken">Cancellation token for the operation</param>
 		/// <returns>The <see cref="SlackMessageId"/> of either the newly sent notification or from the <see cref="JobStepNotificationStateDocument"/></returns>
-		async Task<SlackMessageId> SendSlackThreadMessageAsync(SlackMessageId jobMessageId, JobStepNotificationStateDocument jobStepNotificationState, SlackMessage threadMessage, CancellationToken cancellationToken = default)
+		async Task<SlackMessageId?> SendSlackThreadMessageAsync(SlackMessageId jobMessageId, JobStepNotificationStateDocument jobStepNotificationState, SlackMessage threadMessage, CancellationToken cancellationToken = default)
 		{
 			// Check to see if a message has already been sent out to avoid sending duplicate notifications
 			if (!String.IsNullOrEmpty(jobStepNotificationState.MessageId.ThreadTs))
@@ -978,11 +983,19 @@ namespace HordeServer.Notifications.Sinks
 				return jobStepNotificationState.MessageId;
 			}
 
-			SlackMessageId threadMessageId = await _slackClient.PostMessageToThreadAsync(jobMessageId, threadMessage, cancellationToken);
+			try
+			{
+				SlackMessageId threadMessageId = await _slackClient.PostMessageToThreadAsync(jobMessageId, threadMessage, cancellationToken);
 
-			_logger.LogInformation("Finished sending notification for job step '{JobStepId}' with Job Id '{JobId}' for template '{Template}'", jobStepNotificationState.JobStepId, jobStepNotificationState.JobId, jobStepNotificationState.TemplateId);
+				_logger.LogInformation("Finished sending notification for job step '{JobStepId}' with Job Id '{JobId}' for template '{Template}'", jobStepNotificationState.JobStepId, jobStepNotificationState.JobId, jobStepNotificationState.TemplateId);
+				return threadMessageId;
+			}
+			catch (Exception ex)
+			{
+				_logger.LogWarning(ex, "Failed to send Slack message for job step '{JobStepId}' with Job Id '{JobId}' for template '{Template}'", jobStepNotificationState.JobStepId, jobStepNotificationState.JobId, jobStepNotificationState.TemplateId);
+			}
 
-			return threadMessageId;
+			return null;
 		}
 
 		/// <summary>
@@ -1051,9 +1064,11 @@ namespace HordeServer.Notifications.Sinks
 				slackBlockContainer.Blocks.Add(limitContext);
 			}
 
+			int logEventCount = Math.Min(logData.Count, MaxEventMessages);
+			int maxLogEventLength = logEventCount <= 0 ? MaxLogEventLength : MaxLogEventLength / logEventCount;
 			foreach (LogData log in logData.Take(MaxEventMessages))
 			{
-				SectionBlock logSection = new SectionBlock(new TextObject($"```{log.Message}```"));
+				SectionBlock logSection = new SectionBlock(new TextObject(SanitizeQuoteText(log.Message, maxLogEventLength)));
 				slackBlockContainer.Blocks.Add(logSection);
 			}
 
@@ -1148,6 +1163,20 @@ namespace HordeServer.Notifications.Sinks
 			return $"{branch} - {changelistDetails} - {job.Name}";
 		}
 
+		/// <summary>
+		/// Sanitizes the text for a Slack quote message
+		/// </summary>
+		/// <param name="text">Text to sanitize</param>
+		/// <param name="maxLength">Maximum allowed length for the message</param>
+		static string SanitizeQuoteText(string text, int maxLength = MaxLogEventLength)
+		{
+			if (text.Length > maxLength)
+			{
+				text = text.Substring(0, maxLength - 4) + "...";
+			}
+			return $"```{text}```";
+		}
+
 		/// <summary>
 		/// Adds a context to the Slack message with cancellation details
 		/// </summary>
