Horde: Improve configuration for downtime schedules

- Add Duration property, which is an alternative to FinishTime
- Add TimeZone property that can apply a timezone offset to StartTime and FinishTime using a timezone name, which is an alternative to setting the offset directly in the StartTime and FinishTime value.
- StartTime, FinishTime and Duration can be set using abbervations instead of ISO 8601 for example "Sun Dec 1, 2024, 3AM" and "1:30"
- Add DateTimeOffset JSON Converter
- Add tests for ScheduledDowntime and DowntimeService
- Add PostLoad method to ScheduledDowntime

diff --git a/Engine/Source/Programs/Horde/HordeServer.Shared/Utilities/DateTimeOffsetJsonConverter.cs b/Engine/Source/Programs/Horde/HordeServer.Shared/Utilities/DateTimeOffsetJsonConverter.cs
new file mode 100644
--- /dev/null
+++ b/Engine/Source/Programs/Horde/HordeServer.Shared/Utilities/DateTimeOffsetJsonConverter.cs
@@ -0,0 +1,34 @@
+// Copyright Epic Games, Inc. All Rights Reserved.
+
+using System.Diagnostics;
+using System.Globalization;
+using System.Text.Json;
+using System.Text.Json.Serialization;
+
+namespace HordeServer.Utilities
+{
+	/// <summary>
+	/// Converter for <see cref="DateTimeOffset"/> values to JSON
+	/// </summary>
+	public class DateTimeOffsetJsonConverter : JsonConverter<DateTimeOffset>
+	{
+		/// <inheritdoc/>
+		public override DateTimeOffset Read(ref Utf8JsonReader reader, Type typeToConvert, JsonSerializerOptions options)
+		{
+			Debug.Assert(typeToConvert == typeof(DateTimeOffset));
+
+			string? str = reader.GetString();
+			if (str == null)
+			{
+				throw new InvalidDataException("Unable to parse DateTimeOffset");
+			}
+			return DateTimeOffset.Parse(str, CultureInfo.InvariantCulture);
+		}
+
+		/// <inheritdoc/>
+		public override void Write(Utf8JsonWriter writer, DateTimeOffset dateTimeOffset, JsonSerializerOptions options)
+		{
+			writer.WriteStringValue(dateTimeOffset.ToString("o", CultureInfo.InvariantCulture));
+		}
+	}
+}
diff --git a/Engine/Source/Programs/Horde/HordeServer.Shared/Utilities/JsonUtils.cs b/Engine/Source/Programs/Horde/HordeServer.Shared/Utilities/JsonUtils.cs
--- a/Engine/Source/Programs/Horde/HordeServer.Shared/Utilities/JsonUtils.cs
+++ b/Engine/Source/Programs/Horde/HordeServer.Shared/Utilities/JsonUtils.cs
@@ -37,6 +37,7 @@ namespace HordeServer.Utilities
 			options.Converters.Add(new SubResourceIdJsonConverterFactory());
 			options.Converters.Add(new DateTimeJsonConverter());
 			options.Converters.Add(new TimeSpanJsonConverter());
+			options.Converters.Add(new DateTimeOffsetJsonConverter());
 		}
 	}
 }
diff --git a/Engine/Source/Programs/Horde/HordeServer.Tests/Server/DowntimeServiceTests.cs b/Engine/Source/Programs/Horde/HordeServer.Tests/Server/DowntimeServiceTests.cs
new file mode 100644
--- /dev/null
+++ b/Engine/Source/Programs/Horde/HordeServer.Tests/Server/DowntimeServiceTests.cs
@@ -0,0 +1,333 @@
+// Copyright Epic Games, Inc. All Rights Reserved.
+
+using System;
+using System.Globalization;
+using System.Text.Json;
+using System.Threading.Tasks;
+using HordeCommon;
+using HordeServer.Server;
+using HordeServer.Utilities;
+using Microsoft.VisualStudio.TestTools.UnitTesting;
+using Microsoft.Extensions.Logging.Abstractions;
+
+namespace HordeServer.Tests.Server
+{
+	[TestClass]
+	public class DowntimeServiceTest
+	{
+		private static ScheduledDowntime GetDowntimeFromJson(string json)
+		{
+			ScheduledDowntime scheduledDowntime = JsonSerializer.Deserialize<ScheduledDowntime>(json, JsonUtils.DefaultSerializerOptions)!;
+			GlobalConfig config = new();
+			config.Downtime.Add(scheduledDowntime);
+			config.PostLoad(new ServerSettings(), [], []);
+			return scheduledDowntime;
+		}
+
+		private static DateTimeOffset GetNow(int year, int month, int day, int hour, int minute)
+		{
+			DateTime clock = new(year, month, day, hour, minute, 0, DateTimeKind.Utc);
+			DateTimeOffset now = TimeZoneInfo.ConvertTime(new DateTimeOffset(clock), TimeZoneInfo.Utc);
+			return now;
+		}
+
+		private static DateTimeOffset DateTimeOffsetParse(string input)
+		{
+			return DateTimeOffset.Parse(input, CultureInfo.InvariantCulture);
+		}
+
+		[TestMethod]
+		public async Task AdvanceToDowntimeIsActiveAsync()
+		{
+			FakeClock clock = new();
+			GlobalConfig config = new();
+			ScheduledDowntime downtime = new()
+			{
+				StartTime = clock.UtcNow + TimeSpan.FromSeconds(30),
+				FinishTime = clock.UtcNow + TimeSpan.FromSeconds(90)
+			};
+			config.Downtime.Add(downtime);
+			config.PostLoad(new ServerSettings(), [], []);
+			TestOptionsMonitor<GlobalConfig> optionsMonitor = new(config);
+			await using DowntimeService downtimeService = new(clock, optionsMonitor, NullLogger<DowntimeService>.Instance);
+			Assert.IsFalse(downtimeService.IsDowntimeActive);
+			await clock.AdvanceAsync(TimeSpan.FromSeconds(60));
+			downtimeService.Tick();
+			Assert.IsTrue(downtimeService.IsDowntimeActive);
+		}
+
+		[TestMethod]
+		public void IsNotActive()
+		{
+			ScheduledDowntime downtime = GetDowntimeFromJson("""
+				{
+					"startTime": "2020-09-06T03:00:00+00:00",
+					"finishTime": "2020-09-06T03:45:00+00:00"
+				}
+			""");
+			Assert.IsFalse(downtime.IsActive(GetNow(2024, 12, 8, 2, 45)));
+		}
+
+		[TestMethod]
+		public void IsActive()
+		{
+			ScheduledDowntime downtime = GetDowntimeFromJson("""
+				{
+					"startTime": "2020-09-06T03:00:00+00:00",
+					"finishTime": "2020-09-06T03:45:00+00:00"
+				}
+			""");
+			Assert.IsTrue(downtime.IsActive(GetNow(2020, 9, 6, 3, 15)));
+		}
+
+		[TestMethod]
+		public void IsNotActiveDaily()
+		{
+			ScheduledDowntime downtime = GetDowntimeFromJson("""
+				{
+					"startTime": "2020-09-06T03:00:00+00:00",
+					"finishTime": "2020-09-06T03:45:00+00:00",
+					"frequency": "Daily"
+				}
+			""");
+			Assert.IsFalse(downtime.IsActive(GetNow(2024, 12, 8, 2, 45)));
+		}
+
+		[TestMethod]
+		public void IsActiveDaily()
+		{
+			ScheduledDowntime downtime = GetDowntimeFromJson("""
+				{
+					"startTime": "2020-09-06T03:00:00+00:00",
+					"finishTime": "2020-09-06T03:45:00+00:00",
+					"frequency": "Daily"
+				}
+			""");
+			Assert.IsTrue(downtime.IsActive(GetNow(2024, 12, 8, 3, 15)));
+		}
+
+		[TestMethod]
+		public void IsNotActiveWeekly()
+		{
+			ScheduledDowntime downtime = GetDowntimeFromJson("""
+				{
+					"startTime": "2020-09-06T03:00:00+00:00",
+					"finishTime": "2020-09-06T03:45:00+00:00",
+					"frequency": "Weekly"
+				}
+			""");
+			Assert.IsFalse(downtime.IsActive(GetNow(2024, 12, 8, 2, 45)));
+		}
+
+		[TestMethod]
+		public void IsActiveWeekly()
+		{
+			ScheduledDowntime downtime = GetDowntimeFromJson("""
+				{
+					"startTime": "2020-09-06T03:00:00+00:00",
+					"finishTime": "2020-09-06T03:45:00+00:00",
+					"frequency": "Weekly"
+				}
+			""");
+			Assert.IsTrue(downtime.IsActive(GetNow(2024, 12, 8, 3, 15)));
+		}
+
+		[TestMethod]
+		public void FinishTimeFromDuration()
+		{
+			ScheduledDowntime downtime = GetDowntimeFromJson("""
+				{
+					"startTime": "2020-09-06T03:00:00+00:00",
+					"duration": "0:45"
+				}
+			""");
+			Assert.AreEqual(DateTimeOffsetParse("2020-09-06T03:45+0"), downtime.FinishTime);
+		}
+
+		[TestMethod]
+		public void DurationFromFinishTime()
+		{
+			ScheduledDowntime downtime = GetDowntimeFromJson("""
+				{
+					"startTime": "2020-09-06T03:00:00+00:00",
+					"finishTime": "2020-09-06T03:30:00+00:00",
+				}
+			""");
+			Assert.AreEqual(TimeSpan.FromMinutes(30), downtime.Duration);
+		}
+
+		[TestMethod]
+		public void DurationHigherPrecedence()
+		{
+			ScheduledDowntime downtime = GetDowntimeFromJson("""
+				{
+					"startTime": "2020-09-06T03:00:00+00:00",
+					"finishTime": "2020-09-06T03:30:00+00:00",
+					"duration": "0:45"
+				}
+			""");
+			Assert.AreEqual(DateTimeOffsetParse("2020-09-06T03:45+0"), downtime.FinishTime);
+		}
+
+		[TestMethod]
+		public void StartTimeAndDurationMinuteAbbreviation()
+		{
+			ScheduledDowntime downtime = GetDowntimeFromJson("""
+				{
+					"startTime": "Sun Sep 6, 2020, 3AM",
+					"duration": "0:45"
+				}
+			""");
+			Assert.AreEqual(DateTimeOffsetParse("2020-09-06T03:00"), downtime.StartTime);
+			Assert.AreEqual(DateTimeOffsetParse("2020-09-06T03:45"), downtime.FinishTime);
+		}
+		
+		[TestMethod]
+		public void WrongDayOfWeekNameError()
+		{
+			Assert.ThrowsException<FormatException>(() =>
+			{
+				GetDowntimeFromJson("""
+					{
+						"startTime": "Mon Sep 6, 2020, 3AM",
+						"finishTime": "Sun Sep 6, 2020, 3:30AM",
+					}
+				""");
+			});
+		}
+
+		[TestMethod]
+		public void StartTimeAndDurationHourAbbreviation()
+		{
+			ScheduledDowntime downtime = GetDowntimeFromJson("""
+				{
+					"startTime": "Sun Sep 6, 2020, 3AM",
+					"duration": "1:30"
+				}
+			""");
+			Assert.AreEqual(DateTimeOffsetParse("2020-09-06T03:00"), downtime.StartTime);
+			Assert.AreEqual(DateTimeOffsetParse("2020-09-06T04:30"), downtime.FinishTime);
+		}
+
+		[TestMethod]
+		public void StartTimeAndFinishTimeAbbreviation()
+		{
+			ScheduledDowntime downtime = GetDowntimeFromJson("""
+				{
+					"startTime": "Dec 1, 2024, 3AM",
+					"finishTime": "Sunday December 1 2024, 3:15AM"
+				}
+			""");
+			Assert.AreEqual(DateTimeOffsetParse("2024-12-01T03:00"), downtime.StartTime);
+			Assert.AreEqual(DateTimeOffsetParse("2024-12-01T03:15"), downtime.FinishTime);
+		}
+
+		[TestMethod]
+		public void WithTimeZone()
+		{
+			ScheduledDowntime downtime = GetDowntimeFromJson("""
+				{
+					"startTime": "Jan 1, 2020, 3AM",
+					"duration": "1:30",
+					"timezone": "Eastern Standard Time"
+				}
+			""");
+			Assert.AreEqual(DateTimeOffsetParse("2020-01-01T03:00-05:00"), downtime.StartTime);
+			Assert.AreEqual(DateTimeOffsetParse("2020-01-01T04:30-05:00"), downtime.FinishTime);
+		}
+	
+		[TestMethod]
+		public void WithTimeZoneAndDaylightSaving()
+		{
+			ScheduledDowntime downtime = GetDowntimeFromJson("""
+				{
+					"startTime": "Sun Sep 6, 2020, 3AM",
+					"duration": "1:30",
+					"timezone": "Eastern Standard Time"
+				}
+			""");
+			Assert.AreEqual(DateTimeOffsetParse("2020-09-06T03:00-04:00"), downtime.StartTime);
+			Assert.AreEqual(DateTimeOffsetParse("2020-09-06T04:30-04:00"), downtime.FinishTime);
+		}
+	
+		[TestMethod]
+		public void WithTimeZoneIana()
+		{
+			ScheduledDowntime downtime = GetDowntimeFromJson("""
+				{
+					"startTime": "Jan 1, 2020, 3AM",
+					"duration": "1:30",
+					"timezone": "America/New_York"
+				}
+			""");
+			Assert.AreEqual(DateTimeOffsetParse("2020-01-01T03:00-05:00"), downtime.StartTime);
+			Assert.AreEqual(DateTimeOffsetParse("2020-01-01T04:30-05:00"), downtime.FinishTime);
+		}
+	
+		[TestMethod]
+		public void WithTimeZoneIanaAndDaylightSaving()
+		{
+			ScheduledDowntime downtime = GetDowntimeFromJson("""
+				{
+					"startTime": "Sun Sep 6, 2020, 3AM",
+					"duration": "1:30",
+					"timezone": "America/New_York"
+				}
+			""");
+			Assert.AreEqual(DateTimeOffsetParse("2020-09-06T03:00-04:00"), downtime.StartTime);
+			Assert.AreEqual(DateTimeOffsetParse("2020-09-06T04:30-04:00"), downtime.FinishTime);
+		}
+
+		[TestMethod]
+		public void WrongTimeZoneName()
+		{
+			Assert.ThrowsException<TimeZoneNotFoundException>(() =>
+			{
+				GetDowntimeFromJson("""
+					{
+						"startTime": "Jan 1, 2020, 3AM",
+						"timezone": "ET"
+					}
+				""");
+			});
+		}
+	
+		[TestMethod]
+		public void DowntimeToJson()
+		{
+			// -4 hours is Eastern Daylight Time
+			string json = JsonSerializer.Serialize(GetDowntimeFromJson("""
+					{
+						"startTime": "Sun Sep 6, 2020, 3AM -4",
+						"duration": "0:30"
+					}
+				"""),
+				JsonUtils.DefaultSerializerOptions);
+			ScheduledDowntime downtime = GetDowntimeFromJson(json);
+			Assert.AreEqual(DateTimeOffsetParse("2020-09-06T03:00:00-04:00"), downtime.StartTime);
+			Assert.AreEqual(DateTimeOffsetParse("2020-09-06T03:30:00-04:00"), downtime.FinishTime);
+			Assert.AreEqual(TimeSpan.FromMinutes(30), downtime.Duration);
+			Assert.AreEqual(ScheduledDowntimeFrequency.Once, downtime.Frequency);
+			Assert.IsNull(downtime.TimeZone);
+		}
+
+		[TestMethod]
+		public void DowntimeWithTimeZoneToJson()
+		{
+			string json = JsonSerializer.Serialize(GetDowntimeFromJson("""
+					{
+						"startTime": "Jan 1, 2025, 3AM",
+						"duration": "0:30",
+						"timezone": "Eastern Standard Time"
+					}
+				"""),
+				JsonUtils.DefaultSerializerOptions);
+			ScheduledDowntime downtime = GetDowntimeFromJson(json);
+			Assert.AreEqual(DateTimeOffsetParse("2025-01-01T03:00:00-05:00"), downtime.StartTime);
+			Assert.AreEqual(DateTimeOffsetParse("2025-01-01T03:30:00-05:00"), downtime.FinishTime);
+			Assert.AreEqual(TimeSpan.FromMinutes(30), downtime.Duration);
+			Assert.AreEqual(ScheduledDowntimeFrequency.Once, downtime.Frequency);
+			Assert.AreEqual(downtime.TimeZone, "Eastern Standard Time");
+		}
+	}
+}
diff --git a/Engine/Source/Programs/Horde/HordeServer/Server/DowntimeService.cs b/Engine/Source/Programs/Horde/HordeServer/Server/DowntimeService.cs
--- a/Engine/Source/Programs/Horde/HordeServer/Server/DowntimeService.cs
+++ b/Engine/Source/Programs/Horde/HordeServer/Server/DowntimeService.cs
@@ -66,7 +66,7 @@ namespace HordeServer.Server
 			return new ValueTask();
 		}
 
-		void Tick()
+		internal void Tick()
 		{
 			GlobalConfig globalConfig = _globalConfig.CurrentValue;
 
diff --git a/Engine/Source/Programs/Horde/HordeServer/Server/GlobalConfig.cs b/Engine/Source/Programs/Horde/HordeServer/Server/GlobalConfig.cs
--- a/Engine/Source/Programs/Horde/HordeServer/Server/GlobalConfig.cs
+++ b/Engine/Source/Programs/Horde/HordeServer/Server/GlobalConfig.cs
@@ -135,6 +135,11 @@ namespace HordeServer.Server
 				pluginConfig.PostLoad(pluginConfigOptions);
 			}
 
+			foreach (ScheduledDowntime downtime in Downtime)
+			{
+				downtime.PostLoad();
+			}
+
 			_aclLookup.Clear();
 			BuildAclScopeLookup(Acl, _aclLookup);
 		}
@@ -256,6 +261,69 @@ namespace HordeServer.Server
 
 	/// <summary>
 	/// Settings for the maintenance window
+	/// <para>
+	/// In the JSON configuration file:
+	/// <list type="bullet">
+	/// <item>
+	/// <description>
+	/// The values for <see cref="StartTime"/> and <see cref="FinishTime"/> is a date time string in the 
+	/// <see href="https://learn.microsoft.com/en-us/dotnet/standard/base-types/standard-date-and-time-format-strings#Roundtrip">ISO 8601</see>
+	/// format or abbreviated form. The supported abbreviated strings are the strings supported by the
+	/// <see cref="DateTimeOffset.Parse(String)"/> method.
+	/// </description>
+	/// </item>
+	/// <item>
+	/// <description>
+	/// <see cref="FinishTime"/> can be omitted for <see cref="Duration"/> to set when the maintenance window will be completed.
+	/// </description>
+	/// </item>
+	/// <item>
+	/// <description>
+	/// The UTC offset component in <see cref="StartTime"/> and <see cref="FinishTime"/> strings can be omitted and <see cref="TimeZone"/> used instead.
+	/// </description>
+	/// </item>
+	/// </list>
+	/// </para>
+	/// <example>
+	/// The following example creates a weekly maintenance window using the ISO 8601 format on the 6th September at 3AM and 4 hours behind UTC,
+	/// for 30 minutes and repeats every 7 days.
+	/// <code>
+	/// "downtime": [
+	///     {
+	///         "startTime": "2020-09-06T03:00:00-04:00",
+	///         "finishTime": "2020-09-06T03:30:00-04:00",
+	///         "frequency": "Weekly"
+	///     }
+	/// ]
+	/// </code>
+	/// </example>
+	/// <example>
+	/// The following example creates a weekly maintenance window using the abbreviated format 4 hours behind UTC, using
+	/// <see cref="Duration"/> instead of <see cref="FinishTime"/>. The day of the week name is optional.
+	/// <code>
+	/// "downtime": [
+	///     {
+	///         "startTime": "Sun Sep 6, 2020, 3AM -4",
+	///         "duration": "0:30",
+	///         "frequency": "Weekly"
+	///     }
+	/// ]
+	/// </code>
+	/// </example>
+	/// <example>
+	/// The following example creates a weekly maintenance window using the abbreviated format and <see cref="TimeZone"/>. By using a timezone the
+	/// explicit offset component can be omitted.
+	/// <code>
+	/// "downtime": [
+	///     {
+	///         "startTime": "Sun Sep 6, 2020, 3AM",
+	///         "timeZone": "Eastern Standard Time",
+	///         "duration": "0:30",
+	///         "frequency": "Weekly"
+	///     }
+	/// ]
+	/// </code>
+	/// </example>
 	/// </summary>
 	public class ScheduledDowntime
 	{
@@ -269,11 +337,79 @@ namespace HordeServer.Server
 		/// </summary>
 		public DateTimeOffset FinishTime { get; set; }
 
+		/// <summary>
+		/// Duration of the maintenance window. An alternative to <see cref="FinishTime"/>.
+		/// <example>
+		/// The following example creates non-recurring maintenance windows of 1 hour and 15 minutes
+		/// The supported strings for the duration are the strings supported by the <see cref="TimeSpan.Parse(String)"/> method.
+		/// <code>
+		/// "downtime": [
+		///     {
+		///         "startTime": "Dec 10, 2024, 1AM",
+		///         "duration": "1:15"
+		///     }
+		/// ]
+		/// </code>
+		/// </example>
+		/// <remarks>
+		/// If <see cref="FinishTime"/> is also set in the JSON configuration then the value of Duration will take precedence.
+		/// </remarks>
+		/// </summary>
+		public TimeSpan? Duration { get; set; }
+
+		/// <summary>
+		/// The name of time zone to set the Coordinated Universal Time (UTC) offset for <see cref="StartTime"/> and <see cref="FinishTime"/>.
+		/// When TimeZone is not null will replace the existing offsets in StartTime and FinishTime.
+		/// <example>
+		/// The following example creates a maintenance window on Sunday 6th September 2020 at 3AM, US and Canada Eastern Time.
+		/// <code>
+		/// "downtime": [
+		///     {
+		///         "startTime": "Sun Sep 6, 2020, 3AM",
+		///         "timeZone": "Eastern Standard Time",
+		///         "duration": "0:30"
+		///     }
+		/// ]
+		/// </code>
+		/// </example>
+		/// <remarks>
+		/// The supported TimeZone values are the supported values for <see cref="TimeZoneInfo.FindSystemTimeZoneById(String)"/>, which
+		/// includes the canonical values of the operating system, which are returned by <see cref="TimeZoneInfo.GetSystemTimeZones()"/>
+		/// and the aliases provided by .NET.
+		/// The names of aliases can be checked using <see cref="TimeZoneInfo.TryConvertWindowsIdToIanaId(String, out String?)"/> and
+		/// <see cref="TimeZoneInfo.TryConvertIanaIdToWindowsId(String, out String?)"/>.
+		/// For example, for North Carolina the time zone can be either "Eastern Standard Time" or "America/New_York".
+		/// </remarks>
+		/// </summary>
+		public string? TimeZone { get; set; }
+		
 		/// <summary>
 		/// Frequency that the window repeats
 		/// </summary>
 		public ScheduledDowntimeFrequency Frequency { get; set; } = ScheduledDowntimeFrequency.Once;
 
+		internal void PostLoad()
+		{
+			if (TimeZone != null)
+			{
+				TimeZoneInfo timeZoneInfo = TimeZoneInfo.FindSystemTimeZoneById(TimeZone);
+				TimeSpan offset = timeZoneInfo.GetUtcOffset(StartTime);
+				StartTime = new DateTimeOffset(StartTime.Year, StartTime.Month, StartTime.Day,
+					StartTime.Hour, StartTime.Minute, StartTime.Second, offset);
+				FinishTime = new DateTimeOffset(FinishTime.Year, FinishTime.Month, FinishTime.Day,
+					FinishTime.Hour, FinishTime.Minute, FinishTime.Second, offset);
+			}
+
+			if (Duration != null)
+			{
+				FinishTime = StartTime + Duration.Value;
+			}
+			else
+			{
+				Duration = FinishTime - StartTime;
+			}
+		}
+
 		/// <summary>
 		/// Gets the next scheduled downtime
 		/// </summary>
