Horde: Add slow zip download warning

diff --git a/Engine/Source/Programs/Horde/HordeDashboard/src/components/artifacts/ArtifactButton.tsx b/Engine/Source/Programs/Horde/HordeDashboard/src/components/artifacts/ArtifactButton.tsx
--- a/Engine/Source/Programs/Horde/HordeDashboard/src/components/artifacts/ArtifactButton.tsx
+++ b/Engine/Source/Programs/Horde/HordeDashboard/src/components/artifacts/ArtifactButton.tsx
@@ -1,14 +1,18 @@
-import { IContextualMenuProps, DirectionalHint, PrimaryButton } from "@fluentui/react";
+import { IContextualMenuProps, DirectionalHint, PrimaryButton, Modal, Stack, Text, DefaultButton, Checkbox } from "@fluentui/react";
 import backend from "horde/backend";
 import { ArtifactBrowserItem, ArtifactBrowserType, ArtifactSelection, CreateZipRequest, DownloadFormat, GetArtifactResponse } from "horde/backend/Api";
 import dashboard from "horde/backend/Dashboard";
 import { copyToClipboard } from "horde/base/utilities/clipboard";
 import { formatBytes } from "horde/base/utilities/stringUtills";
+import { getHordeStyling } from "horde/styles/Styles";
+import { useState } from "react";
 import { useNavigate } from "react-router-dom";
 
 
 export const ArtifactButton: React.FC<{ jobId?: string, stepId?: string, artifact?: GetArtifactResponse, pathIn?: string, selection?: ArtifactSelection, openArtifactInfo?: () => void, disabled?: boolean, minWidth?: number }> = ({ jobId, stepId, artifact, pathIn, selection, openArtifactInfo, disabled, minWidth }) => {
 
+   const [state, setState] = useState<{ showZipDownloadWarning?: boolean }>({});
+   
    const navigate = useNavigate();
 
    const format = dashboard.downloadFormat;
@@ -148,7 +152,7 @@ export const ArtifactButton: React.FC<{ jobId?: string, stepId?: string, artifac
             key: 'download_zip',
             text: 'Download Zip',
             onClick: () => {
-               downloadZip();
+               showDownloadWarning() ? setState({ showZipDownloadWarning: true }) : downloadZip()
             }
          }
       )
@@ -184,17 +188,58 @@ export const ArtifactButton: React.FC<{ jobId?: string, stepId?: string, artifac
 
    }
 
-   return <PrimaryButton split menuProps={downloadProps} text={buttonText} styles={{ root: { minWidth: minWidth ?? 104, fontFamily: 'Horde Open Sans SemiBold !important' } }} disabled={disabled} onClick={() => {
-      if (format === DownloadFormat.Zip) {
-         downloadZip();
+   function showDownloadWarning() {
+      let hide_warning = localStorage.getItem("horde_hide_zip_warning")
+      if (hide_warning === "true") {
+         return false
       }
-      else if (format === DownloadFormat.UGS) {
-         downloadUGS();
-      }
-      else if (format === DownloadFormat.UTOOLBOX) {
-         downloadToolbox();
+
+      if (selection && selection.size < 52428800) {
+         return false
       }
 
-   }} />
+      return true
+   }   
 
+   return <Stack>{!!state.showZipDownloadWarning && <ZipDownloadWarning onContinue={() => { downloadZip(); setState({ showZipDownloadWarning: false }); }} onClose={() => { setState({ showZipDownloadWarning: false }) }}/>}
+      <PrimaryButton split menuProps={downloadProps} text={buttonText} styles={{ root: { minWidth: minWidth ?? 104, fontFamily: 'Horde Open Sans SemiBold !important' } }} disabled={disabled} onClick={() => {
+         if (format === DownloadFormat.Zip) {
+            showDownloadWarning() ? setState({ showZipDownloadWarning: true }) : downloadZip()
+         }
+         else if (format === DownloadFormat.UGS) {
+            downloadUGS();
+         }
+         else if (format === DownloadFormat.UTOOLBOX) {
+            downloadToolbox();
+         }
+      }} />
+   </Stack>
+
+}
+
+const ZipDownloadWarning: React.FC<{ onContinue: () => void, onClose: () => void }> = ({ onContinue, onClose }) => {
+   
+   const { hordeClasses } = getHordeStyling();
+
+   return <Modal className={hordeClasses.modal} isOpen={true} isBlocking={true} topOffsetFixed={true} styles={{ main: { padding: 8, height: 200, width: 600, hasBeenOpened: false, top: 0, bottom: 150, left: 0, right: 0, margin: "auto", position: "absolute" } }} >
+      
+      <Stack style={{ padding: 8 }}>
+         <Stack style={{ paddingBottom: 24 }}>
+            <Text variant="mediumPlus" style={{ fontFamily: "Horde Open Sans SemiBold" }}>Warning: Slow Download</Text>
+         </Stack>
+         <Stack style={{ paddingBottom: 38 }}>
+            <Text variant="medium" style={{ fontFamily: "Horde Open Sans SemiBold" }}>Downloading artifacts in zip format is significantly slower than downloading with UGS.</Text>
+         </Stack>
+         <Stack>
+            <Checkbox label="Don't show this warning again" onChange={(ev?: any, checked?: boolean | undefined) => {localStorage.setItem("horde_hide_zip_warning", checked!.toString())}}/>
+         </Stack>
+         <Stack horizontal style={{ paddingTop: 16 }}>
+            <Stack grow />
+            <Stack horizontal tokens={{ childrenGap: 20 }}>
+               <PrimaryButton onClick={() => onContinue()} text="Continue" />
+               <DefaultButton onClick={() => onClose()} text="Cancel" />
+            </Stack>
+         </Stack>
+      </Stack>
+   </Modal>
 }
\ No newline at end of file
