Adjusting 'Enabled' property to match the rest of the properties.

#rnx

diff --git a/Engine/Source/Programs/Horde/Plugins/Experimental/HordeServer.Experimental/Notifications/NotificationConfig.cs b/Engine/Source/Programs/Horde/Plugins/Experimental/HordeServer.Experimental/Notifications/NotificationConfig.cs
--- a/Engine/Source/Programs/Horde/Plugins/Experimental/HordeServer.Experimental/Notifications/NotificationConfig.cs
+++ b/Engine/Source/Programs/Horde/Plugins/Experimental/HordeServer.Experimental/Notifications/NotificationConfig.cs
@@ -39,7 +39,7 @@ namespace HordeServer
 		/// <summary>
 		/// Whether or not the current configuration is enabled
 		/// </summary>
-		public bool IsEnabled { get; set; } = true;
+		public bool Enabled { get; set; } = true;
 
 		/// <summary>
 		/// Identifier of the template this notification applies to
@@ -217,7 +217,7 @@ namespace HordeServer
 				{
 					foreach (JobNotificationConfig notificationConfig in notificationConfigList)
 					{
-						if (notificationConfig.IsEnabled || includeDisabledConfigs)
+						if (notificationConfig.Enabled || includeDisabledConfigs)
 						{
 							jobNotificationConfigs.Add(notificationConfig);
 						}
