Horde: Updating schema documentation

#rnx

diff --git a/Engine/Source/Programs/Horde/Docs/Config/Schema/Globals.md b/Engine/Source/Programs/Horde/Docs/Config/Schema/Globals.md
--- a/Engine/Source/Programs/Horde/Docs/Config/Schema/Globals.md
+++ b/Engine/Source/Programs/Horde/Docs/Config/Schema/Globals.md
@@ -35,11 +35,22 @@ Name | Description
 ## ScheduledDowntime
 
 Settings for the maintenance window
+In the JSON configuration file:
+The values for  and  is a date time string in the ISO 8601 format or abbreviated form. The supported abbreviated strings are the strings supported by the method. can be omitted for  to set when the maintenance window will be completed.
+The UTC offset component in  and  strings can be omitted and  used instead.
+The following example creates a weekly maintenance window using the ISO 8601 format on the 6th September at 3AM and 4 hours behind UTC, for 30 minutes and repeats every 7 days.
+"downtime": [ { "startTime": "2020-09-06T03:00:00-04:00", "finishTime": "2020-09-06T03:30:00-04:00", "frequency": "Weekly" } ]
+The following example creates a weekly maintenance window using the abbreviated format 4 hours behind UTC, using instead of . The day of the week name is optional.
+"downtime": [ { "startTime": "Sun Sep 6, 2020, 3AM -4", "duration": "0:30", "frequency": "Weekly" } ]
+The following example creates a weekly maintenance window using the abbreviated format and . By using a timezone the explicit offset component can be omitted.
+"downtime": [ { "startTime": "Sun Sep 6, 2020, 3AM", "timeZone": "Eastern Standard Time", "duration": "0:30", "frequency": "Weekly" } ]
 
 Name | Description
 ---- | -----------
 `startTime` | `string`<br>Start time
 `finishTime` | `string`<br>Finish time
+`duration` | `string`<br>Duration of the maintenance window. An alternative to .<br>The following example creates non-recurring maintenance windows of 1 hour and 15 minutes The supported strings for the duration are the strings supported by the  method.<br>"downtime": [ { "startTime": "Dec 10, 2024, 1AM", "duration": "1:15" } ]<br>If  is also set in the JSON configuration then the value of Duration will take precedence.
+`timeZone` | `string`<br>The name of time zone to set the Coordinated Universal Time (UTC) offset for  and . When TimeZone is not null will replace the existing offsets in StartTime and FinishTime.<br>The following example creates a maintenance window on Sunday 6th September 2020 at 3AM, US and Canada Eastern Time.<br>"downtime": [ { "startTime": "Sun Sep 6, 2020, 3AM", "timeZone": "Eastern Standard Time", "duration": "0:30" } ]<br>The supported TimeZone values are the supported values for , which includes the canonical values of the operating system, which are returned by and the aliases provided by .NET. The names of aliases can be checked using  and . For example, for North Carolina the time zone can be either "Eastern Standard Time" or "America/New_York".
 `frequency` | [ScheduledDowntimeFrequency](#scheduleddowntimefrequency-enum)<br>Frequency that the window repeats
 
 ## ScheduledDowntimeFrequency (Enum)
diff --git a/Engine/Source/Programs/Horde/Docs/Config/Schema/Notification.md b/Engine/Source/Programs/Horde/Docs/Config/Schema/Notification.md
--- a/Engine/Source/Programs/Horde/Docs/Config/Schema/Notification.md
+++ b/Engine/Source/Programs/Horde/Docs/Config/Schema/Notification.md
@@ -25,7 +25,7 @@ Configuration for the job notification
 
 Name | Description
 ---- | -----------
-`isEnabled` | `boolean`<br>Whether or not the current configuration is enabled
+`enabled` | `boolean`<br>Whether or not the current configuration is enabled
 `template` | `string`<br>Identifier of the template this notification applies to
 `namePattern` | `string`<br>Regex pattern used to extract name of the step this notification will use for display
 `notificationFormats` | [NotificationFormatConfig](#notificationformatconfig)`[]`<br>List of formats to be applied for this notification
