Horde: Log view search improvements
Fixed missing line indicator on build health issues
Preserve current line after search is cleared
Search now shows first result from current line, instead of starting from the top of the log

diff --git a/Engine/Source/Programs/Horde/HordeDashboard/src/components/LogView.tsx b/Engine/Source/Programs/Horde/HordeDashboard/src/components/LogView.tsx
--- a/Engine/Source/Programs/Horde/HordeDashboard/src/components/LogView.tsx
+++ b/Engine/Source/Programs/Horde/HordeDashboard/src/components/LogView.tsx
@@ -609,7 +609,7 @@ const LogLineIndicator: React.FC<{ lineIndex: number }> = ({ lineIndex }) => {
 
    useQuery();
 
-   if (!globalHandler || !globalSearchState?.results?.length) {
+   if (!globalHandler) {
       return null;
    }
 
@@ -1012,15 +1012,15 @@ export const LogList: React.FC<{ logId: string }> = observer(({ logId }) => {
 
    const doQuery = (newValue: string) => {
 
-      handler.currentLine = undefined;
-
-      // always re-render log when querying with new value
-      if (searchState?.search !== newValue) {
-         logListKey++;
+      let currentLine = 0
+      if (handler.currentLine !== undefined) {
+         currentLine = handler.currentLine
       }
 
       if (!newValue) {
 
+         logListKey++;
+
          setSearchState({
                search: '',
                curRequest: undefined,
@@ -1028,6 +1028,21 @@ export const LogList: React.FC<{ logId: string }> = observer(({ logId }) => {
                text: ''
             });
 
+            // Reset current line and scroll to value before search was cleared
+            handler.stopTrailing();
+            let lineIdx = currentLine - 10;
+            if (lineIdx < 0) {
+               lineIdx = 0;
+            }
+
+            updateLineQuery(currentLine, navigate);
+
+            inTimeout = true;
+               setTimeout(() => {
+                  inTimeout = false;
+                  listRef?.scrollToIndex(lineIdx, () => { return handler.lineHeight; }, ScrollToMode.top);
+               }, 10)
+
       } else if (newValue !== searchState.search) {
 
          const request = backend.searchLog(logId, newValue, 0, 65535);
@@ -1041,12 +1056,23 @@ export const LogList: React.FC<{ logId: string }> = observer(({ logId }) => {
 
             if (newState.curRequest === request) {
 
-               curSearchIdx = 0;
-
                const lines = logResponse.lines;
 
                if (lines.length) {
 
+                  logListKey++;
+
+                  // Get next occurrence after current line
+                  curSearchIdx = 0;
+                  for (const [index, value] of lines.entries()) {
+                     if (value + 1 >= currentLine) {
+                        curSearchIdx = index
+                        break
+                     }
+                  }
+
+                  handler.currentLine = lines[curSearchIdx] + 1
+
                   setSearchState({
                      search: newValue,
                      curRequest: undefined,
@@ -1530,8 +1556,9 @@ export const LogList: React.FC<{ logId: string }> = observer(({ logId }) => {
                                     }}
 
                                     onGetErrorMessage={(newValue) => {
-
-                                       doQuery(newValue);
+                                       if (searchState.text !== searchState.search) {
+                                          doQuery(newValue);
+                                       }
                                        return undefined;
 
 
