Horde: When user checks in a devcie, mark it for cleaning

diff --git a/Engine/Source/Programs/Horde/Plugins/Build/HordeServer.Build/Devices/DeviceCollection.cs b/Engine/Source/Programs/Horde/Plugins/Build/HordeServer.Build/Devices/DeviceCollection.cs
--- a/Engine/Source/Programs/Horde/Plugins/Build/HordeServer.Build/Devices/DeviceCollection.cs
+++ b/Engine/Source/Programs/Horde/Plugins/Build/HordeServer.Build/Devices/DeviceCollection.cs
@@ -542,6 +542,8 @@ namespace HordeServer.Devices
 			else
 			{
 				updates.Add(updateBuilder.Set(x => x.CheckoutExpiringNotificationSent, true));
+				// Mark device as needing cleaning
+				updates.Add(updateBuilder.Set(x => x.CleanRequired, true));
 			}
 
 			await _devices.FindOneAndUpdateAsync<DeviceDocument>(x => x.Id == deviceId, updateBuilder.Combine(updates));
