Horde: Adding job/step parent/child tracking for spawned jobs

diff --git a/Engine/Source/Programs/Horde/Plugins/Build/HordeServer.Build.Tests/Jobs/JobCollectionTests.cs b/Engine/Source/Programs/Horde/Plugins/Build/HordeServer.Build.Tests/Jobs/JobCollectionTests.cs
--- a/Engine/Source/Programs/Horde/Plugins/Build/HordeServer.Build.Tests/Jobs/JobCollectionTests.cs
+++ b/Engine/Source/Programs/Horde/Plugins/Build/HordeServer.Build.Tests/Jobs/JobCollectionTests.cs
@@ -19,6 +19,7 @@ using HordeServer.Agents.Sessions;
 using HordeServer.Jobs;
 using HordeServer.Logs;
 using HordeServer.Projects;
+using HordeServer.Server;
 using HordeServer.Streams;
 using HordeServer.Users;
 using HordeServer.Utilities;
@@ -938,6 +939,60 @@ namespace HordeServer.Tests.Jobs
 			}
 		}
 
+		[TestMethod]
+		public async Task TestParentChildStepsAsync()
+		{
+			ProjectId projectId = new ProjectId("ue5");
+			StreamId streamId = new StreamId("ue5-main");
+			TemplateId templateRefId1 = new TemplateId("template1");
+			TemplateId templateRefId2 = new TemplateId("template2");
+
+			StreamConfig streamConfig = new StreamConfig { Id = streamId };
+			streamConfig.Templates.Add(new TemplateRefConfig { Id = templateRefId1, Name = "Test Template 1" });
+			streamConfig.Templates.Add(new TemplateRefConfig { Id = templateRefId2, Name = "Test Template 2" });
+			streamConfig.Tabs.Add(new TabConfig { Title = "foo", Templates = new List<TemplateId> { templateRefId1, templateRefId2 } });
+
+			ProjectConfig projectConfig = new ProjectConfig { Id = projectId };
+			projectConfig.Streams.Add(streamConfig);
+
+			BuildConfig buildConfig = new BuildConfig();
+			buildConfig.Projects.Add(projectConfig);
+
+			GlobalConfig globalConfig = new GlobalConfig();
+			globalConfig.Plugins.AddBuildConfig(buildConfig);
+
+			SetConfig(globalConfig);
+
+			CreateJobRequest request = new CreateJobRequest(new StreamId("ue5-main"), new TemplateId("template1"));
+			request.CommitId = CommitIdWithOrder.FromPerforceChange(123);
+
+			ActionResult<CreateJobResponse>  parentJobResult = await JobsController.CreateJobAsync(request);
+
+			Assert.IsNotNull(parentJobResult.Value);
+			JobId parentJobId = JobId.Parse(parentJobResult.Value!.Id);
+
+			IJob parentJob = Deref(await JobCollection.GetAsync(parentJobId));
+			
+			request.ParentJobId = parentJob.Id.ToString();
+			request.ParentJobStepId = parentJob.Batches[0].Steps[0].Id.ToString();
+
+			ActionResult<CreateJobResponse> childJobResult = await JobsController.CreateJobAsync(request);
+
+			Assert.IsNotNull(childJobResult.Value);
+			JobId childJobId = JobId.Parse(childJobResult.Value!.Id);
+
+			IJob childJob = Deref(await JobCollection.GetAsync(childJobId));
+
+			Assert.AreEqual(parentJobId, childJob.ParentJobId);
+			Assert.AreEqual(parentJob.Batches[0].Steps[0].Id, childJob.ParentJobStepId);
+
+			parentJob = Deref(await JobCollection.GetAsync(parentJobId));
+
+			Assert.AreEqual(parentJob.Batches[0].Steps[0].SpawnedJobs?.Count, 1);
+			Assert.AreEqual(parentJob.Batches[0].Steps[0].SpawnedJobs![0], childJobId);
+
+		}
+
 		private async Task<IJob> AddJobAsync(IGraph graph, string commitId)
 		{
 			return await JobCollection.AddAsync(JobIdUtils.GenerateNewId(), _streamId, _templateId1, _template1.Hash, graph, "jobName", new CommitId(commitId), null, new CreateJobOptions());
diff --git a/Engine/Source/Programs/Horde/Plugins/Build/HordeServer.Build/Jobs/IJobCollection.cs b/Engine/Source/Programs/Horde/Plugins/Build/HordeServer.Build/Jobs/IJobCollection.cs
--- a/Engine/Source/Programs/Horde/Plugins/Build/HordeServer.Build/Jobs/IJobCollection.cs
+++ b/Engine/Source/Programs/Horde/Plugins/Build/HordeServer.Build/Jobs/IJobCollection.cs
@@ -82,6 +82,12 @@ namespace HordeServer.Jobs
 		/// <inheritdoc cref="IJob.Environment"/>
 		public Dictionary<string, string> Environment { get; } = new Dictionary<string, string>(StringComparer.OrdinalIgnoreCase);
 
+		/// <inheritdoc cref="IJob.ParentJobId"/>
+		public JobId? ParentJobId { get; set; }
+
+		/// <inheritdoc cref="IJob.ParentJobStepId"/>
+		public JobStepId? ParentJobStepId { get; set; }
+
 		/// <summary>
 		/// Default constructor
 		/// </summary>
diff --git a/Engine/Source/Programs/Horde/Plugins/Build/HordeServer.Build/Jobs/JobCollection.cs b/Engine/Source/Programs/Horde/Plugins/Build/HordeServer.Build/Jobs/JobCollection.cs
--- a/Engine/Source/Programs/Horde/Plugins/Build/HordeServer.Build/Jobs/JobCollection.cs
+++ b/Engine/Source/Programs/Horde/Plugins/Build/HordeServer.Build/Jobs/JobCollection.cs
@@ -95,6 +95,10 @@ namespace HordeServer.Jobs
 			IReadOnlyDictionary<int, NotificationTriggerId> IJob.LabelIdxToTriggerId => _document._labelNotifications.ToDictionary(x => x._labelIdx, x => x._triggerId);
 			IReadOnlyList<IJobReport>? IJob.Reports => _document.Reports;
 			IReadOnlyList<IChainedJob> IJob.ChainedJobs => _document.ChainedJobs;
+
+			JobId? IJob.ParentJobId => _document.ParentJobId;
+			JobStepId? IJob.ParentJobStepId => _document.ParentJobStepId;
+
 			DateTime IJob.UpdateTimeUtc => _document.UpdateTimeUtc ?? _document.UpdateTime?.UtcDateTime ?? DateTime.UnixEpoch;
 			int IJob.UpdateIndex => _document.UpdateIndex;
 
@@ -229,9 +233,9 @@ namespace HordeServer.Jobs
 				return _collection.CreateJobObject(newDocument, Graph);
 			}
 
-			public async Task<IJob?> TryUpdateStepAsync(JobStepBatchId batchId, JobStepId stepId, JobStepState newState = JobStepState.Unspecified, JobStepOutcome newOutcome = JobStepOutcome.Unspecified, JobStepError? newError = null, bool? newAbortRequested = null, UserId? newAbortByUserId = null, LogId? newLogId = null, NotificationTriggerId? newNotificationTriggerId = null, UserId? newRetryByUserId = null, Priority? newPriority = null, List<JobReport>? newReports = null, Dictionary<string, string?>? newProperties = null, string? newCancellationReason = null, CancellationToken cancellationToken = default)
+			public async Task<IJob?> TryUpdateStepAsync(JobStepBatchId batchId, JobStepId stepId, JobStepState newState = JobStepState.Unspecified, JobStepOutcome newOutcome = JobStepOutcome.Unspecified, JobStepError? newError = null, bool? newAbortRequested = null, UserId? newAbortByUserId = null, LogId? newLogId = null, NotificationTriggerId? newNotificationTriggerId = null, UserId? newRetryByUserId = null, Priority? newPriority = null, List<JobReport>? newReports = null, Dictionary<string, string?>? newProperties = null, string? newCancellationReason = null, JobId? newSpawnedJob = null, CancellationToken cancellationToken = default)
 			{
-				JobDocument? newDocument = await _collection.TryUpdateStepAsync(_document, _graph, batchId, stepId, newState, newOutcome, newError, newAbortRequested, newAbortByUserId, newLogId, newNotificationTriggerId, newRetryByUserId, newPriority, newReports, newProperties, newCancellationReason, cancellationToken);
+				JobDocument? newDocument = await _collection.TryUpdateStepAsync(_document, _graph, batchId, stepId, newState, newOutcome, newError, newAbortRequested, newAbortByUserId, newLogId, newNotificationTriggerId, newRetryByUserId, newPriority, newReports, newProperties, newCancellationReason, newSpawnedJob, cancellationToken);
 				return _collection.CreateJobObject(newDocument, Graph);
 			}
 		}
@@ -325,6 +329,8 @@ namespace HordeServer.Jobs
 			IReadOnlyList<IJobReport>? IJobStep.Reports => Document.Reports;
 			IReadOnlyDictionary<string, string>? IJobStep.Properties => Document.Properties;
 
+			IReadOnlyList<JobId>? IJobStep.SpawnedJobs => Document.SpawnedJobs;
+
 			public JobStep(JobStepBatch batch, JobStepDocument document, INode node)
 			{
 				Batch = batch;
@@ -511,6 +517,11 @@ namespace HordeServer.Jobs
 			public List<LabelNotificationDocument> _labelNotifications = new List<LabelNotificationDocument>();
 			public List<ChainedJobDocument> ChainedJobs { get; set; } = new List<ChainedJobDocument>();
 
+			[BsonIgnoreIfNull]
+			public JobId? ParentJobId { get; set; }
+			[BsonIgnoreIfNull]
+			public JobStepId? ParentJobStepId { get; set; }
+
 			[BsonIgnoreIfNull]
 			public List<NodeRef>? RetriedNodes { get; set; }
 			public SubResourceId NextSubResourceId { get; set; }
@@ -567,6 +578,9 @@ namespace HordeServer.Jobs
 					Environment[pair.Key] = pair.Value;
 				}
 
+				ParentJobId = options.ParentJobId;
+				ParentJobStepId = options.ParentJobStepId;
+
 				NextSubResourceId = SubResourceId.GenerateNewId();
 				UpdateTimeUtc = createTimeUtc;
 			}
@@ -694,6 +708,9 @@ namespace HordeServer.Jobs
 			[BsonIgnoreIfNull]
 			public List<JobReport>? Reports { get; set; }
 
+			[BsonIgnoreIfNull]
+			public List<JobId>? SpawnedJobs { get; set; }
+
 			[BsonIgnoreIfNull]
 			public Dictionary<string, string>? Properties { get; set; }
 
@@ -1500,7 +1517,7 @@ namespace HordeServer.Jobs
 		}
 
 		/// <inheritdoc/>
-		async Task<JobDocument?> TryUpdateStepAsync(JobDocument jobDocument, IGraph graph, JobStepBatchId batchId, JobStepId stepId, JobStepState newState, JobStepOutcome newOutcome, JobStepError? newError, bool? newAbortRequested, UserId? newAbortByUserId, LogId? newLogId, NotificationTriggerId? newNotificationTriggerId, UserId? newRetryByUserId, Priority? newPriority, List<JobReport>? newReports, Dictionary<string, string?>? newProperties, string? newCancellationReason, CancellationToken cancellationToken)
+		async Task<JobDocument?> TryUpdateStepAsync(JobDocument jobDocument, IGraph graph, JobStepBatchId batchId, JobStepId stepId, JobStepState newState, JobStepOutcome newOutcome, JobStepError? newError, bool? newAbortRequested, UserId? newAbortByUserId, LogId? newLogId, NotificationTriggerId? newNotificationTriggerId, UserId? newRetryByUserId, Priority? newPriority, List<JobReport>? newReports, Dictionary<string, string?>? newProperties, string? newCancellationReason, JobId? newSpawnedJob, CancellationToken cancellationToken)
 		{
 			jobDocument = Clone(jobDocument);
 
@@ -1640,6 +1657,13 @@ namespace HordeServer.Jobs
 								updates.Add(updateBuilder.Set(x => x.Batches[batchIdx].Steps[stepIdx].Reports, step.Reports));
 							}
 
+							if (newSpawnedJob != null)
+							{
+								step.SpawnedJobs ??= new List<JobId>();
+								step.SpawnedJobs.Add(newSpawnedJob.Value);
+								updates.Add(updateBuilder.Set(x => x.Batches[batchIdx].Steps[stepIdx].SpawnedJobs, step.SpawnedJobs));
+							}
+
 							// Apply any property updates
 							if (newProperties != null)
 							{
diff --git a/Engine/Source/Programs/Horde/Plugins/Build/HordeServer.Build/Jobs/JobService.cs b/Engine/Source/Programs/Horde/Plugins/Build/HordeServer.Build/Jobs/JobService.cs
--- a/Engine/Source/Programs/Horde/Plugins/Build/HordeServer.Build/Jobs/JobService.cs
+++ b/Engine/Source/Programs/Horde/Plugins/Build/HordeServer.Build/Jobs/JobService.cs
@@ -882,7 +882,7 @@ namespace HordeServer.Jobs
 		/// <param name="newCancellationReason">The reason the step was canceled</param>
 		/// <param name="cancellationToken">Cancellation token for the operation</param>
 		/// <returns>True if the job was updated, false if it was deleted in the meantime</returns>
-		public async Task<IJob?> UpdateStepAsync(IJob job, JobStepBatchId batchId, JobStepId stepId, StreamConfig streamConfig, JobStepState newState = JobStepState.Unspecified, JobStepOutcome newOutcome = JobStepOutcome.Unspecified, JobStepError? newError = null, bool? newAbortRequested = null, UserId? newAbortByUserId = null, LogId? newLogId = null, NotificationTriggerId? newNotificationTriggerId = null, UserId? newRetryByUserId = null, Priority? newPriority = null, List<JobReport>? newReports = null, Dictionary<string, string?>? newProperties = null, string? newCancellationReason = null, CancellationToken cancellationToken = default)
+		public async Task<IJob?> UpdateStepAsync(IJob job, JobStepBatchId batchId, JobStepId stepId, StreamConfig streamConfig, JobStepState newState = JobStepState.Unspecified, JobStepOutcome newOutcome = JobStepOutcome.Unspecified, JobStepError? newError = null, bool? newAbortRequested = null, UserId? newAbortByUserId = null, LogId? newLogId = null, NotificationTriggerId? newNotificationTriggerId = null, UserId? newRetryByUserId = null, Priority? newPriority = null, List<JobReport>? newReports = null, Dictionary<string, string?>? newProperties = null, string? newCancellationReason = null, JobId? newSpawnedJob = null, CancellationToken cancellationToken = default)
 		{
 			using TelemetrySpan span = _tracer.StartActiveSpan($"{nameof(JobService)}.{nameof(UpdateStepAsync)}");
 			span.SetAttribute("Job", job.Id.ToString());
@@ -892,7 +892,7 @@ namespace HordeServer.Jobs
 			using IDisposable? scope = _logger.BeginScope("UpdateStepAsync({JobId}:{BatchId}:{StepId})", job.Id, batchId, stepId);
 			for (; ; )
 			{
-				IJob? newJob = await TryUpdateStepAsync(job, batchId, stepId, streamConfig, newState, newOutcome, newError, newAbortRequested, newAbortByUserId, newLogId, newNotificationTriggerId, newRetryByUserId, newPriority, newReports, newProperties, newCancellationReason, cancellationToken);
+				IJob? newJob = await TryUpdateStepAsync(job, batchId, stepId, streamConfig, newState, newOutcome, newError, newAbortRequested, newAbortByUserId, newLogId, newNotificationTriggerId, newRetryByUserId, newPriority, newReports, newProperties, newCancellationReason, newSpawnedJob, cancellationToken);
 				if (newJob != null)
 				{
 					return newJob;
@@ -929,7 +929,7 @@ namespace HordeServer.Jobs
 		/// <param name="newCancellationReason">Property changes. Any properties with a null value will be removed.</param>
 		/// <param name="cancellationToken">Cancellation token for the operation</param>
 		/// <returns>True if the job was updated, false if it was deleted in the meantime</returns>
-		public async Task<IJob?> TryUpdateStepAsync(IJob job, JobStepBatchId batchId, JobStepId stepId, StreamConfig streamConfig, JobStepState newState = JobStepState.Unspecified, JobStepOutcome newOutcome = JobStepOutcome.Unspecified, JobStepError? newError = null, bool? newAbortRequested = null, UserId? newAbortByUserId = null, LogId? newLogId = null, NotificationTriggerId? newTriggerId = null, UserId? newRetryByUserId = null, Priority? newPriority = null, List<JobReport>? newReports = null, Dictionary<string, string?>? newProperties = null, string? newCancellationReason = null, CancellationToken cancellationToken = default)
+		public async Task<IJob?> TryUpdateStepAsync(IJob job, JobStepBatchId batchId, JobStepId stepId, StreamConfig streamConfig, JobStepState newState = JobStepState.Unspecified, JobStepOutcome newOutcome = JobStepOutcome.Unspecified, JobStepError? newError = null, bool? newAbortRequested = null, UserId? newAbortByUserId = null, LogId? newLogId = null, NotificationTriggerId? newTriggerId = null, UserId? newRetryByUserId = null, Priority? newPriority = null, List<JobReport>? newReports = null, Dictionary<string, string?>? newProperties = null, string? newCancellationReason = null, JobId? newSpawnedJob = null, CancellationToken cancellationToken = default)
 		{
 			using TelemetrySpan span = _tracer.StartActiveSpan($"{nameof(JobService)}.{nameof(TryUpdateStepAsync)}");
 			span.SetAttribute("Job", job.Id.ToString());
@@ -961,7 +961,7 @@ namespace HordeServer.Jobs
 			}
 
 			// Update the step
-			IJob? newJob = await job.TryUpdateStepAsync(batchId, stepId, newState, newOutcome, newError, newAbortRequested, newAbortByUserId, newLogId, newTriggerId, newRetryByUserId, newPriority, newReports, newProperties, newCancellationReason, cancellationToken);
+			IJob? newJob = await job.TryUpdateStepAsync(batchId, stepId, newState, newOutcome, newError, newAbortRequested, newAbortByUserId, newLogId, newTriggerId, newRetryByUserId, newPriority, newReports, newProperties, newCancellationReason, newSpawnedJob, cancellationToken);
 			if (newJob != null)
 			{
 				job = newJob;
diff --git a/Engine/Source/Programs/Horde/Plugins/Build/HordeServer.Build/Jobs/JobsController.cs b/Engine/Source/Programs/Horde/Plugins/Build/HordeServer.Build/Jobs/JobsController.cs
--- a/Engine/Source/Programs/Horde/Plugins/Build/HordeServer.Build/Jobs/JobsController.cs
+++ b/Engine/Source/Programs/Horde/Plugins/Build/HordeServer.Build/Jobs/JobsController.cs
@@ -209,6 +209,8 @@ namespace HordeServer.Jobs
 			options.Priority = priority;
 			options.AutoSubmit = create.AutoSubmit;
 			options.UpdateIssues = updateIssues;
+			options.ParentJobId = create.ParentJobId == null ? null : JobId.Parse(create.ParentJobId);
+			options.ParentJobStepId = create.ParentJobStepId == null ? null : JobStepId.Parse(create.ParentJobStepId);
 			options.Claims.AddRange(User.Claims.Select(x => new AclClaimConfig(x)));
 			options.JobOptions ??= create.JobOptions;
 
@@ -266,6 +268,20 @@ namespace HordeServer.Jobs
 			// Create the job
 			IJob job = await _jobService.CreateJobAsync(null, streamConfig, templateRefId, template.Hash, graph, name, commitId, codeCommitId, options, cancellationToken);
 			await UpdateNotificationsAsync(job.Id, new UpdateNotificationsRequest { Slack = true }, cancellationToken);
+
+			if (options.ParentJobId != null && options.ParentJobStepId != null)
+			{
+				IJob? parentJob = await _jobService.GetJobAsync(options.ParentJobId.Value, cancellationToken);
+				if (parentJob != null)
+				{
+					IJobStep? parentStep;
+					if (parentJob.TryGetStep(options.ParentJobStepId.Value, out parentStep))
+					{
+						await _jobService.UpdateStepAsync(parentJob, parentStep.Batch.Id, parentStep.Id, streamConfig, newSpawnedJob: job.Id);
+					}					
+				}				
+			}
+
 			return new CreateJobResponse(job.Id.ToString());
 		}
 
@@ -653,6 +669,8 @@ namespace HordeServer.Jobs
 			response.UseArtifactsV2 = true;
 			response.UpdateIssues = job.UpdateIssues;
 			response.CanUpdate = canUpdate;
+			response.ParentJobId = job.ParentJobId.ToString();
+			response.ParentJobStepId = job.ParentJobStepId.ToString();
 			return response;
 		}
 
diff --git a/Engine/Source/Programs/Shared/EpicGames.Horde/Jobs/IJob.cs b/Engine/Source/Programs/Shared/EpicGames.Horde/Jobs/IJob.cs
--- a/Engine/Source/Programs/Shared/EpicGames.Horde/Jobs/IJob.cs
+++ b/Engine/Source/Programs/Shared/EpicGames.Horde/Jobs/IJob.cs
@@ -239,6 +239,16 @@ namespace EpicGames.Horde.Jobs
 		/// </summary>
 		public IReadOnlyList<IChainedJob> ChainedJobs { get; }
 
+		/// <summary>
+		/// The job which spawned this job
+		/// </summary>
+		public JobId? ParentJobId { get; }
+
+		/// <summary>
+		/// The job step which spawned this job
+		/// </summary>
+		public JobStepId? ParentJobStepId { get; }
+
 		/// <summary>
 		/// The last update time
 		/// </summary>
@@ -331,9 +341,10 @@ namespace EpicGames.Horde.Jobs
 		/// <param name="newReports">New report documents</param>
 		/// <param name="newProperties">Property changes. Any properties with a null value will be removed.</param>
 		/// <param name="newCancellationReason">The reason the job step was canceled</param>
+		/// <param name="newSpawnedJob">JobId of any job this step spawned</param>
 		/// <param name="cancellationToken">Cancellation token for the operation</param>
 		/// <returns>True if the job was updated, false if it was deleted in the meantime</returns>
-		Task<IJob?> TryUpdateStepAsync(JobStepBatchId batchId, JobStepId stepId, JobStepState newState = default, JobStepOutcome newOutcome = default, JobStepError? newError = null, bool? newAbortRequested = null, UserId? newAbortByUserId = null, LogId? newLogId = null, NotificationTriggerId? newNotificationTriggerId = null, UserId? newRetryByUserId = null, Priority? newPriority = null, List<JobReport>? newReports = null, Dictionary<string, string?>? newProperties = null, string? newCancellationReason = null, CancellationToken cancellationToken = default);
+		Task<IJob?> TryUpdateStepAsync(JobStepBatchId batchId, JobStepId stepId, JobStepState newState = default, JobStepOutcome newOutcome = default, JobStepError? newError = null, bool? newAbortRequested = null, UserId? newAbortByUserId = null, LogId? newLogId = null, NotificationTriggerId? newNotificationTriggerId = null, UserId? newRetryByUserId = null, Priority? newPriority = null, List<JobReport>? newReports = null, Dictionary<string, string?>? newProperties = null, string? newCancellationReason = null, JobId? newSpawnedJob = null, CancellationToken cancellationToken = default);
 
 		/// <summary>
 		/// Attempts to update the node groups to be executed for a job. Fails if another write happens in the meantime.
@@ -1551,6 +1562,11 @@ namespace EpicGames.Horde.Jobs
 		/// </summary>
 		public IReadOnlyList<IJobReport>? Reports { get; }
 
+		/// <summary>
+		/// List of jobs this step has spawned
+		/// </summary>
+		public IReadOnlyList<JobId>? SpawnedJobs { get; }
+
 		/// <summary>
 		/// Reports for this jobstep.
 		/// </summary>
diff --git a/Engine/Source/Programs/Shared/EpicGames.Horde/Jobs/JobMessages.cs b/Engine/Source/Programs/Shared/EpicGames.Horde/Jobs/JobMessages.cs
--- a/Engine/Source/Programs/Shared/EpicGames.Horde/Jobs/JobMessages.cs
+++ b/Engine/Source/Programs/Shared/EpicGames.Horde/Jobs/JobMessages.cs
@@ -349,6 +349,16 @@ namespace EpicGames.Horde.Jobs
 		/// </summary>
 		public List<string>? Targets { get; set; }
 
+		/// <summary>
+		/// The parent job id if any
+		/// </summary>
+		public string? ParentJobId { get; set; }
+
+		/// <summary>
+		/// The parent step id if any
+		/// </summary>
+		public string? ParentJobStepId { get; set; }
+
 		/// <summary>
 		/// Private constructor for serialization
 		/// </summary>
@@ -699,6 +709,16 @@ namespace EpicGames.Horde.Jobs
 		/// </summary>
 		public bool CanUpdate { get; set; } = true;
 
+		/// <summary>
+		/// The parent job id, if any
+		/// </summary>
+		public string? ParentJobId { get; set; }
+
+		/// <summary>
+		/// The parent job step id, if any
+		/// </summary>
+		public string? ParentJobStepId { get; set; }
+
 		/// <summary>
 		/// Constructor
 		/// </summary>
