Changing logging from warning to being informational

#rnx

diff --git a/Engine/Source/Programs/Horde/Plugins/Experimental/HordeServer.Experimental/Notifications/Sinks/ExperimentalSlackNotificationSink.cs b/Engine/Source/Programs/Horde/Plugins/Experimental/HordeServer.Experimental/Notifications/Sinks/ExperimentalSlackNotificationSink.cs
--- a/Engine/Source/Programs/Horde/Plugins/Experimental/HordeServer.Experimental/Notifications/Sinks/ExperimentalSlackNotificationSink.cs
+++ b/Engine/Source/Programs/Horde/Plugins/Experimental/HordeServer.Experimental/Notifications/Sinks/ExperimentalSlackNotificationSink.cs
@@ -992,7 +992,7 @@ namespace HordeServer.Notifications.Sinks
 			}
 			catch (Exception ex)
 			{
-				_logger.LogWarning(ex, "Failed to send Slack message for job step '{JobStepId}' with Job Id '{JobId}' for template '{Template}'", jobStepNotificationState.JobStepId, jobStepNotificationState.JobId, jobStepNotificationState.TemplateId);
+				_logger.LogInformation(ex, "Failed to send Slack message for job step '{JobStepId}' with Job Id '{JobId}' for template '{Template}' ({Error})", jobStepNotificationState.JobStepId, jobStepNotificationState.JobId, jobStepNotificationState.TemplateId, ex.Message);
 			}
 
 			return null;
