Horde: Adding spawned job ids to step responses

diff --git a/Engine/Source/Programs/Horde/Plugins/Build/HordeServer.Build/Jobs/JobsController.cs b/Engine/Source/Programs/Horde/Plugins/Build/HordeServer.Build/Jobs/JobsController.cs
--- a/Engine/Source/Programs/Horde/Plugins/Build/HordeServer.Build/Jobs/JobsController.cs
+++ b/Engine/Source/Programs/Horde/Plugins/Build/HordeServer.Build/Jobs/JobsController.cs
@@ -808,6 +808,7 @@ namespace HordeServer.Jobs
 			response.StartTime = step.StartTimeUtc;
 			response.FinishTime = step.FinishTimeUtc;
 			response.Reports = step.Reports?.ConvertAll(x => CreateGetReportResponse(x));
+			response.SpawnedJobs = step.SpawnedJobs?.ConvertAll(x => x.ToString());
 
 			if (step.Properties != null && step.Properties.Count > 0)
 			{
diff --git a/Engine/Source/Programs/Shared/EpicGames.Horde/Jobs/JobMessages.cs b/Engine/Source/Programs/Shared/EpicGames.Horde/Jobs/JobMessages.cs
--- a/Engine/Source/Programs/Shared/EpicGames.Horde/Jobs/JobMessages.cs
+++ b/Engine/Source/Programs/Shared/EpicGames.Horde/Jobs/JobMessages.cs
@@ -1019,6 +1019,11 @@ namespace EpicGames.Horde.Jobs
 		/// </summary>
 		public List<GetJobReportResponse>? Reports { get; set; }
 
+		/// <summary>
+		/// List of spawned job ids
+		/// </summary>
+		public List<string>? SpawnedJobs { get; set; }
+
 		/// <summary>
 		/// User-defined properties for this jobstep.
 		/// </summary>
