Horde: Fix warnings in job service/conntroller

diff --git a/Engine/Source/Programs/Horde/Plugins/Build/HordeServer.Build/Jobs/JobService.cs b/Engine/Source/Programs/Horde/Plugins/Build/HordeServer.Build/Jobs/JobService.cs
--- a/Engine/Source/Programs/Horde/Plugins/Build/HordeServer.Build/Jobs/JobService.cs
+++ b/Engine/Source/Programs/Horde/Plugins/Build/HordeServer.Build/Jobs/JobService.cs
@@ -880,6 +880,7 @@ namespace HordeServer.Jobs
 		/// <param name="newReports">New list of reports</param>
 		/// <param name="newProperties">Property changes. Any properties with a null value will be removed.</param>
 		/// <param name="newCancellationReason">The reason the step was canceled</param>
+		/// <param name="newSpawnedJob">A new job this step spawned</param>
 		/// <param name="cancellationToken">Cancellation token for the operation</param>
 		/// <returns>True if the job was updated, false if it was deleted in the meantime</returns>
 		public async Task<IJob?> UpdateStepAsync(IJob job, JobStepBatchId batchId, JobStepId stepId, StreamConfig streamConfig, JobStepState newState = JobStepState.Unspecified, JobStepOutcome newOutcome = JobStepOutcome.Unspecified, JobStepError? newError = null, bool? newAbortRequested = null, UserId? newAbortByUserId = null, LogId? newLogId = null, NotificationTriggerId? newNotificationTriggerId = null, UserId? newRetryByUserId = null, Priority? newPriority = null, List<JobReport>? newReports = null, Dictionary<string, string?>? newProperties = null, string? newCancellationReason = null, JobId? newSpawnedJob = null, CancellationToken cancellationToken = default)
@@ -927,6 +928,7 @@ namespace HordeServer.Jobs
 		/// <param name="newReports">New reports</param>
 		/// <param name="newProperties">Property changes. Any properties with a null value will be removed.</param>
 		/// <param name="newCancellationReason">Property changes. Any properties with a null value will be removed.</param>
+		/// <param name="newSpawnedJob">A new spawned job this step added</param>
 		/// <param name="cancellationToken">Cancellation token for the operation</param>
 		/// <returns>True if the job was updated, false if it was deleted in the meantime</returns>
 		public async Task<IJob?> TryUpdateStepAsync(IJob job, JobStepBatchId batchId, JobStepId stepId, StreamConfig streamConfig, JobStepState newState = JobStepState.Unspecified, JobStepOutcome newOutcome = JobStepOutcome.Unspecified, JobStepError? newError = null, bool? newAbortRequested = null, UserId? newAbortByUserId = null, LogId? newLogId = null, NotificationTriggerId? newTriggerId = null, UserId? newRetryByUserId = null, Priority? newPriority = null, List<JobReport>? newReports = null, Dictionary<string, string?>? newProperties = null, string? newCancellationReason = null, JobId? newSpawnedJob = null, CancellationToken cancellationToken = default)
diff --git a/Engine/Source/Programs/Horde/Plugins/Build/HordeServer.Build/Jobs/JobsController.cs b/Engine/Source/Programs/Horde/Plugins/Build/HordeServer.Build/Jobs/JobsController.cs
--- a/Engine/Source/Programs/Horde/Plugins/Build/HordeServer.Build/Jobs/JobsController.cs
+++ b/Engine/Source/Programs/Horde/Plugins/Build/HordeServer.Build/Jobs/JobsController.cs
@@ -277,7 +277,7 @@ namespace HordeServer.Jobs
 					IJobStep? parentStep;
 					if (parentJob.TryGetStep(options.ParentJobStepId.Value, out parentStep))
 					{
-						await _jobService.UpdateStepAsync(parentJob, parentStep.Batch.Id, parentStep.Id, streamConfig, newSpawnedJob: job.Id);
+						await _jobService.UpdateStepAsync(parentJob, parentStep.Batch.Id, parentStep.Id, streamConfig, newSpawnedJob: job.Id, cancellationToken: cancellationToken);
 					}					
 				}				
 			}
