Added error handling for cases when Slack has outages are fails to post messages

#rnx

diff --git a/Engine/Source/Programs/Horde/Plugins/Experimental/HordeServer.Experimental/Notifications/Sinks/ExperimentalSlackNotificationSink.cs b/Engine/Source/Programs/Horde/Plugins/Experimental/HordeServer.Experimental/Notifications/Sinks/ExperimentalSlackNotificationSink.cs
--- a/Engine/Source/Programs/Horde/Plugins/Experimental/HordeServer.Experimental/Notifications/Sinks/ExperimentalSlackNotificationSink.cs
+++ b/Engine/Source/Programs/Horde/Plugins/Experimental/HordeServer.Experimental/Notifications/Sinks/ExperimentalSlackNotificationSink.cs
@@ -425,7 +425,7 @@ namespace HordeServer.Notifications.Sinks
 			// Send the updated summary after the steps has been updated
 			if (jobNotificationStates is null)
 			{
-				_logger.LogWarning("No job notifications found for job id '{JobId}' even though job step '{StepId}' was updated.", job.Id, step.Id);
+				_logger.LogInformation("No job notifications found for job id '{JobId}' even though job step '{StepId}' was updated.", job.Id, step.Id);
 				return;
 			}
 
@@ -747,7 +747,7 @@ namespace HordeServer.Notifications.Sinks
 			if (states.Count != 1)
 			{
 				// Log error that multiple were found and could cause issues
-				_logger.LogWarning("Found {Count} records when fetching for {JobId}:{TemplateId}:{Recipient}. Using first record which could cause issues.", states.Count, jobId, templateId, recipient);
+				_logger.LogInformation("Found {Count} records when fetching for {JobId}:{TemplateId}:{Recipient}. Using first record which could cause issues.", states.Count, jobId, templateId, recipient);
 			}
 
 			return states[0];
@@ -767,6 +767,9 @@ namespace HordeServer.Notifications.Sinks
 																		Builders<JobNotificationStateDocument>.Filter.Eq(x => x.TemplateId, templateId) &
 																		Builders<JobNotificationStateDocument>.Filter.Eq(x => x.Recipient, recipient);
 			DeleteResult result = await _jobNotificationStates.DeleteManyAsync(filter, cancellationToken);
+
+			_logger.LogInformation("Removed {MessageCount} JobNotificationState(s) for jobId: {JobId}, templateId: {TemplateId}, and recipient: {Recipient}", result.DeletedCount, jobId, templateId, recipient);
+
 			return result.DeletedCount > 0;
 		}
 
@@ -844,6 +847,9 @@ namespace HordeServer.Notifications.Sinks
 																		Builders<JobStepNotificationStateDocument>.Filter.Eq(x => x.TemplateId, templateId) &
 																		Builders<JobStepNotificationStateDocument>.Filter.Eq(x => x.Recipient, recipient);
 			DeleteResult result = await _jobStepNotificationStates.DeleteManyAsync(filter, cancellationToken);
+
+			_logger.LogInformation("Removed {MessageCount} JobStepNotificationState(s) for jobId: {JobId}, templateId: {TemplateId}, and recipient: {Recipient}", result.DeletedCount, jobId, templateId, recipient);
+
 			return result.DeletedCount > 0;
 		}
 
@@ -921,7 +927,14 @@ namespace HordeServer.Notifications.Sinks
 						channel = userInfo.Id;
 					}
 
-					_ = await _slackClient.PostMessageAsync(channel, message, cancellationToken);
+					try
+					{
+						_ = await _slackClient.PostMessageAsync(channel, message, cancellationToken);
+					}
+					catch (Exception ex)
+					{
+						_logger.LogInformation(ex, "Failed to send Slack setup failure message for job '{JobId}' and the template '{Template}' ({Error})", job.Id, job.TemplateId, ex.Message);
+					}
 				}
 			}
 
@@ -943,14 +956,21 @@ namespace HordeServer.Notifications.Sinks
 			SlackMessage? message = await GenerateSlackSummaryMessageAsync(job, branch, hordeJobURL, recipient, abortedByUserName, cancellationToken);
 			if (message is null)
 			{
-				_logger.LogWarning("Unable to generate Job summary notification for job '{JobId}' for the template '{TemplateId}'", job.Id, job.TemplateId);
+				_logger.LogInformation("Unable to generate Job summary notification for job '{JobId}' and the template '{TemplateId}'", job.Id, job.TemplateId);
 				return;
 			}
 
 			JobNotificationStateDocument? jobNotificationState = await GetJobNotificationStateAsync(job.Id, job.TemplateId, recipient, cancellationToken);
 			if (jobNotificationState is not null)
 			{
-				_ = _slackClient.UpdateMessageAsync(jobNotificationState.MessageId, message, cancellationToken);
+				try
+				{
+					_ = _slackClient.UpdateMessageAsync(jobNotificationState.MessageId, message, cancellationToken);
+				}
+				catch (Exception ex)
+				{
+					_logger.LogInformation(ex, "Failed to update Slack summary message for job '{JobId}' and the template '{Template}' ({Error})", jobNotificationState.JobId, jobNotificationState.TemplateId, ex.Message);
+				}
 			}
 			else
 			{
@@ -962,8 +982,15 @@ namespace HordeServer.Notifications.Sinks
 					channel = userInfo.Id;
 				}
 
-				SlackMessageId messageId = await _slackClient.PostMessageAsync(channel, message, cancellationToken);
-				_ = await AddOrUpdateJobNotificationStateAsync(job.Id, job.TemplateId, recipient, messageId, cancellationToken);
+				try
+				{
+					SlackMessageId messageId = await _slackClient.PostMessageAsync(channel, message, cancellationToken);
+					_ = await AddOrUpdateJobNotificationStateAsync(job.Id, job.TemplateId, recipient, messageId, cancellationToken);
+				}
+				catch (Exception ex)
+				{
+					_logger.LogInformation(ex, "Failed to send Slack summary message for job '{JobId}' and the template '{Template}' ({Error})", job.Id, job.TemplateId, ex.Message);
+				}
 			}
 		}
 
