Horde: Prefer AuthorId instead of OwnerId when looking for suspects from issue commits

diff --git a/Engine/Source/Programs/Horde/Plugins/Build/HordeServer.Build.Tests/Issues/IssueServiceTests.cs b/Engine/Source/Programs/Horde/Plugins/Build/HordeServer.Build.Tests/Issues/IssueServiceTests.cs
--- a/Engine/Source/Programs/Horde/Plugins/Build/HordeServer.Build.Tests/Issues/IssueServiceTests.cs
+++ b/Engine/Source/Programs/Horde/Plugins/Build/HordeServer.Build.Tests/Issues/IssueServiceTests.cs
@@ -117,10 +117,10 @@ namespace HordeServer.Tests.Issues
 		const string DevStreamName = "//UE4/Dev";
 		readonly StreamId _devStreamId = new StreamId("ue4-dev");
 		readonly IGraph _graph;
-		readonly PerforceServiceStub _perforce;
-		readonly UserId _timId;
+		readonly PerforceServiceStub _perforce;		
 		readonly UserId _jerryId;
 		readonly UserId _bobId;
+		readonly UserId _chrisId;
 		readonly DirectoryReference _autoSdkDir;
 		readonly DirectoryReference _workspaceDir;
 
@@ -179,8 +179,8 @@ namespace HordeServer.Tests.Issues
 			IUser jerry = UserCollection.FindOrAddUserByLoginAsync("Jerry").Result;
 			IUser chris = UserCollection.FindOrAddUserByLoginAsync("Chris").Result;
 			IUser tim = UserCollection.FindOrAddUserByLoginAsync("Tim").Result;
-
-			_timId = tim.Id;
+			
+			_chrisId = chris.Id;
 			_jerryId = UserCollection.FindOrAddUserByLoginAsync("Jerry").Result.Id;
 			_bobId = UserCollection.FindOrAddUserByLoginAsync("Bob").Result.Id;
 
@@ -1019,7 +1019,7 @@ namespace HordeServer.Tests.Issues
 				//			Assert.AreEqual(null, Suspects[2].OriginatingChange);
 
 				Assert.AreEqual(120, suspects[0].CommitId.GetPerforceChange());
-				Assert.AreEqual(_timId, suspects[0].AuthorId);
+				Assert.AreEqual(_chrisId, suspects[0].AuthorId);
 				//				Assert.AreEqual(null, Suspects[0].OriginatingChange);
 
 				IReadOnlyList<IIssue> openIssues = await IssueCollection.FindIssuesAsync(resolved: false);
@@ -1048,7 +1048,7 @@ namespace HordeServer.Tests.Issues
 				Assert.AreEqual(2, suspects.Count);
 
 				Assert.AreEqual(120, suspects[0].CommitId.GetPerforceChange());
-				Assert.AreEqual(_timId, suspects[0].AuthorId);
+				Assert.AreEqual(_chrisId, suspects[0].AuthorId);
 
 				Assert.AreEqual(75, suspects[1].CommitId.GetPerforceChange());
 				Assert.AreEqual(_jerryId, suspects[1].AuthorId);
@@ -1163,7 +1163,7 @@ namespace HordeServer.Tests.Issues
 				List<UserId> primarySuspects = suspects.Select(x => x.AuthorId).ToList();
 				Assert.AreEqual(2, primarySuspects.Count);
 				Assert.IsTrue(primarySuspects.Contains(_jerryId)); // 115
-				Assert.IsTrue(primarySuspects.Contains(_timId)); // 120
+				Assert.IsTrue(primarySuspects.Contains(_chrisId)); // 120
 			}
 
 			// #3
@@ -1177,11 +1177,11 @@ namespace HordeServer.Tests.Issues
 				Assert.AreEqual(1, issues.Count);
 
 				IIssue issue = issues[0];
-				Assert.AreEqual(_timId, issue.OwnerId);
+				Assert.AreEqual(_chrisId, issue.OwnerId);
 
 				// Also check updating an issue doesn't clear the owner
 				Assert.IsTrue(await IssueService.UpdateIssueAsync(issue.Id));
-				Assert.AreEqual(_timId, issue!.OwnerId);
+				Assert.AreEqual(_chrisId, issue!.OwnerId);
 			}
 		}
 
@@ -1504,7 +1504,7 @@ namespace HordeServer.Tests.Issues
 				List<UserId> primarySuspects = suspects.Select(x => x.AuthorId).ToList();
 				Assert.AreEqual(2, primarySuspects.Count);
 				Assert.IsTrue(primarySuspects.Contains(_jerryId)); // 115
-				Assert.IsTrue(primarySuspects.Contains(_timId)); // 120
+				Assert.IsTrue(primarySuspects.Contains(_chrisId)); // 120
 			}
 
 			// #3
@@ -1513,7 +1513,7 @@ namespace HordeServer.Tests.Issues
 			{
 				IReadOnlyList<IIssue> issues = await IssueCollection.FindIssuesAsync();
 				Assert.AreEqual(1, issues.Count);
-				await IssueService.UpdateIssueAsync(issues[0].Id, declinedById: _timId);
+				await IssueService.UpdateIssueAsync(issues[0].Id, declinedById: _chrisId);
 
 				issues = await IssueCollection.FindIssuesAsync();
 				Assert.AreEqual(1, issues.Count);
diff --git a/Engine/Source/Programs/Horde/Plugins/Build/HordeServer.Build/Issues/IssueService.cs b/Engine/Source/Programs/Horde/Plugins/Build/HordeServer.Build/Issues/IssueService.cs
--- a/Engine/Source/Programs/Horde/Plugins/Build/HordeServer.Build/Issues/IssueService.cs
+++ b/Engine/Source/Programs/Horde/Plugins/Build/HordeServer.Build/Issues/IssueService.cs
@@ -1391,7 +1391,7 @@ namespace HordeServer.Issues
 
 					if (rank == maxRank)
 					{
-						NewIssueSpanSuspectData suspect = new NewIssueSpanSuspectData(commit.Id, commit.OwnerId);
+						NewIssueSpanSuspectData suspect = new NewIssueSpanSuspectData(commit.Id, commit.AuthorId == UserId.Empty || commit.AuthorId == UserId.Anonymous ? commit.OwnerId : commit.AuthorId);
 						suspect.SourceCommitId = commit.OriginalCommitId;
 						suspects.Add(suspect);
 					}
