Horde: Update notes on Epic's deployment

#skipci

diff --git a/Engine/Source/Programs/Horde/Docs/Deployment.md b/Engine/Source/Programs/Horde/Docs/Deployment.md
--- a/Engine/Source/Programs/Horde/Docs/Deployment.md
+++ b/Engine/Source/Programs/Horde/Docs/Deployment.md
@@ -59,19 +59,29 @@ the future and should not be taken as an endorsement or non-endorsement of parti
 
 ## Epic's Horde Deployment
 
-Epic's internal deployment of Horde is hosted on AWS and consists of the following:
+Epic has one large deployment of Horde hosted on AWS and is described at a high level below.
+Whether this setup should be mimicked depends on the size of your deployment.
+We recommend starting with small instance types and adjust as needed.
+Separation of servers into runmodes is also an advanced use-case and is likely not needed from the start.
 
-* Amazon Load Balancer for SSL termination and routing
-* 2 Linux containers serving the dashboard managed by Amazon ECS (0.25 vCPU, 0.5GB RAM each)
-* 10 Horde Server instances configured with the 'Server' [runmode](Deployment/ServerSettings.md#runmode-enum), running
-  in Linux containers managed by Amazon ECS (4 vCPU, 16GB RAM each)
-  * These instances handle lightweight user-facing requests.
-* 2 Horde Server instances configured with the 'Worker' [runmode](Deployment/ServerSettings.md#runmode-enum), running
-  in Linux containers managed by Amazon ECS (4 vCPU, 16GB RAM each)
-  * These instances handle heavier scheduled tasks, such as replicating Perforce metadata, reading and updating
-    config state, and starting scheduled jobs.
+### Server
+* Application Load Balancer for TLS termination, certificate handling and routing
+  * gRPC using `Http2Port` for unencrypted internal communication between ALB and ECS
+* Amazon ECS with Fargate for hosting the C# processes (called _tasks_ in ECS)
+  * 6 tasks for Horde Server configured with the `Server` [runmode](Deployment/ServerSettings.md#runmode-enum)
+    * 8 vCPU, 16 GB RAM each
+    * These handle short-lived, front-facing requests from browsers, agents and REST API
+  * 2 tasks for Horde Server configured with the `Worker` [runmode](Deployment/ServerSettings.md#runmode-enum)
+    * 2 vCPU, 8 GB RAM each
+    * These handle heavier scheduled tasks, such as replicating Perforce metadata, reading and updating
+      config state, and starting scheduled jobs.
+    * Separating these has the benefit of background tasks never affecting request serving
 * Amazon DocumentDB (compatible with MongoDB)
 * Amazon ElastiCache (compatible with Redis)
-* Several hundred EC2 instances running AMIs containing the Horde agent.
+
+### Agents
+
+* Several hundred EC2 instances running AMIs containing the Horde agent and supporting toolchains
+* Agents for remote execution (Unreal Build Accelerator) run exclusively as EC2 spot instances using Linux and Wine
 * Around 100 machines hosted on-premise to provide access to mobile devices, consoles, and unhosted platforms
   (eg. Mac).
