Added exception handling for when a configuration provides invalid regex.

#rnx

diff --git a/Engine/Source/Programs/Horde/Plugins/Experimental/HordeServer.Experimental/Notifications/Sinks/ExperimentalSlackNotificationSink.cs b/Engine/Source/Programs/Horde/Plugins/Experimental/HordeServer.Experimental/Notifications/Sinks/ExperimentalSlackNotificationSink.cs
--- a/Engine/Source/Programs/Horde/Plugins/Experimental/HordeServer.Experimental/Notifications/Sinks/ExperimentalSlackNotificationSink.cs
+++ b/Engine/Source/Programs/Horde/Plugins/Experimental/HordeServer.Experimental/Notifications/Sinks/ExperimentalSlackNotificationSink.cs
@@ -322,20 +322,27 @@ namespace HordeServer.Notifications.Sinks
 		async Task ProcessJobAsync(IJob job, JobNotificationConfig config, CancellationToken cancellationToken = default)
 		{
 			// Populate our job information needed for the notification
-			foreach (IJobStepBatch batch in job.Batches)
+			try
 			{
-				foreach (IJobStep step in batch.Steps)
+				foreach (IJobStepBatch batch in job.Batches)
 				{
-					if (config.GetStepDetails(step.Name, out string? group, out string? platform) && !String.IsNullOrEmpty(group) &&!String.IsNullOrEmpty(platform))
+					foreach (IJobStep step in batch.Steps)
 					{
-						string badge = GetBadgeForStep(step);
-						foreach (string recipient in config.Channels)
+						if (config.GetStepDetails(step.Name, out string? group, out string? platform) && !String.IsNullOrEmpty(group) && !String.IsNullOrEmpty(platform))
 						{
-							await AddOrUpdateJobStepNotificationStateAsync(job.Id, job.TemplateId, recipient, step.Id, group, platform, badge, null, cancellationToken);
+							string badge = GetBadgeForStep(step);
+							foreach (string recipient in config.Channels)
+							{
+								await AddOrUpdateJobStepNotificationStateAsync(job.Id, job.TemplateId, recipient, step.Id, group, platform, badge, null, cancellationToken);
+							}
 						}
 					}
 				}
 			}
+			catch (Exception ex)
+			{
+				_logger.LogInformation(ex, "Failed to extract step details for '{Job}' in stream '{StreamId}' and for template '{TemplateId}' ({Error})", job.Id, job.StreamId, job.TemplateId, ex.Message);
+			}
 		}
 
 		/// <summary>
