Fixing exception thrown by Slack caused by the number of blocks exceeding the limit

#rnx

diff --git a/Engine/Source/Programs/Horde/Plugins/Experimental/HordeServer.Experimental/Notifications/Sinks/ExperimentalSlackNotificationSink.cs b/Engine/Source/Programs/Horde/Plugins/Experimental/HordeServer.Experimental/Notifications/Sinks/ExperimentalSlackNotificationSink.cs
--- a/Engine/Source/Programs/Horde/Plugins/Experimental/HordeServer.Experimental/Notifications/Sinks/ExperimentalSlackNotificationSink.cs
+++ b/Engine/Source/Programs/Horde/Plugins/Experimental/HordeServer.Experimental/Notifications/Sinks/ExperimentalSlackNotificationSink.cs
@@ -150,6 +150,12 @@ namespace HordeServer.Notifications.Sinks
 		/// </summary>
 		const int MaxLogEventLength = 2000;
 
+		/// <summary>
+		/// Block limit per Slack Context before an exception is thrown
+		/// </summary>
+		/// <remarks>See https://api.slack.com/reference/block-kit/blocks#context for detailed information</remarks>
+		const int MaxSlackBlocks = 10;
+
 		/// <summary>
 		/// Constructor
 		/// </summary>
@@ -1288,7 +1294,6 @@ namespace HordeServer.Notifications.Sinks
 						stepContext = null;
 					}
 
-					itemsInGroup = 1;
 					currentGroup = state.Group;
 					header = new ContextBlock();
 					header.Elements.Add(new TextObject($"*{currentGroup}*"));
@@ -1298,7 +1303,7 @@ namespace HordeServer.Notifications.Sinks
 				string jobStepLink = $"{hordeJobURL}?step={state.JobStepId}";
 
 				// Slack can only allow for 10 elements per block (https://api.slack.com/reference/block-kit/blocks#context) and we want to track when to create new ContextBlocks
-				bool bNeedsNewContext = (itemsInGroup % 10 == 0 || stepContext is null);
+				bool bNeedsNewContext = (itemsInGroup % MaxSlackBlocks == 0 || stepContext is null);
 				if (bNeedsNewContext)
 				{
 					if (stepContext is not null)
@@ -1308,6 +1313,9 @@ namespace HordeServer.Notifications.Sinks
 					stepContext = new ContextBlock();
 				}
 				stepContext?.Elements.Add(new TextObject($"<{jobStepLink}|{state.Badge} {state.TargetPlatform}>"));
+
+				// The StepContext should always be valid here, but to disable the warning, we provide a default so that we only have 1 block per context.
+				itemsInGroup = stepContext?.Elements.Count ?? MaxSlackBlocks;
 			}
 
 			// Check to see if we have an existing step context that needs to be added
