Horde: Remove incorrect return type metadata for PUT /api/v2/devices/{deviceId}

diff --git a/Engine/Source/Programs/Horde/Plugins/Build/HordeServer.Build/Devices/DevicesController.cs b/Engine/Source/Programs/Horde/Plugins/Build/HordeServer.Build/Devices/DevicesController.cs
--- a/Engine/Source/Programs/Horde/Plugins/Build/HordeServer.Build/Devices/DevicesController.cs
+++ b/Engine/Source/Programs/Horde/Plugins/Build/HordeServer.Build/Devices/DevicesController.cs
@@ -195,7 +195,6 @@ namespace HordeServer.Devices
 		[HttpPut]
 		[Authorize]
 		[Route("/api/v2/devices/{deviceId}")]
-		[ProducesResponseType(typeof(List<GetDeviceResponse>), 200)]
 		public async Task<ActionResult> UpdateDeviceAsync(string deviceId, [FromBody] UpdateDeviceRequest update, CancellationToken cancellationToken = default)
 		{
 			IUser? internalUser = await _userCollection.GetUserAsync(User, cancellationToken);
