Horde: Add device cleaning information to devices endpoint

diff --git a/Engine/Source/Programs/Horde/Plugins/Build/HordeServer.Build/Devices/DevicesController.cs b/Engine/Source/Programs/Horde/Plugins/Build/HordeServer.Build/Devices/DevicesController.cs
--- a/Engine/Source/Programs/Horde/Plugins/Build/HordeServer.Build/Devices/DevicesController.cs
+++ b/Engine/Source/Programs/Horde/Plugins/Build/HordeServer.Build/Devices/DevicesController.cs
@@ -140,7 +140,7 @@ namespace HordeServer.Devices
 					ReservationFinishUtc = u.ReservationFinishUtc
 				}).ToList();
 
-				responses.Add(new GetDeviceResponse(device.Id.ToString(), device.PlatformId.ToString(), device.PoolId.ToString(), device.Name, device.Enabled, device.Address, device.ModelId?.ToString(), device.ModifiedByUser, device.Notes, device.ProblemTimeUtc, device.MaintenanceTimeUtc, utilization, device.CheckedOutByUser, device.CheckOutTime, checkoutExpiration));
+				responses.Add(new GetDeviceResponse(device.Id.ToString(), device.PlatformId.ToString(), device.PoolId.ToString(), device.Name, device.Enabled, device.Address, device.ModelId?.ToString(), device.ModifiedByUser, device.Notes, device.ProblemTimeUtc, device.MaintenanceTimeUtc, utilization, device.CheckedOutByUser, device.CheckOutTime, checkoutExpiration, device.CleanRequired, device.CleanStartTimeUtc, device.CleanFinishTimeUtc));
 			}
 
 			return responses;
