Horde: Adding zen backend to artifacts controller uartifact generation via metadata

diff --git a/Engine/Source/Programs/Horde/Plugins/Build/HordeServer.Build/Artifacts/ArtifactsController.cs b/Engine/Source/Programs/Horde/Plugins/Build/HordeServer.Build/Artifacts/ArtifactsController.cs
--- a/Engine/Source/Programs/Horde/Plugins/Build/HordeServer.Build/Artifacts/ArtifactsController.cs
+++ b/Engine/Source/Programs/Horde/Plugins/Build/HordeServer.Build/Artifacts/ArtifactsController.cs
@@ -605,7 +605,34 @@ namespace HordeServer.Artifacts
 
 		ActionResult GetDescriptorInternal(IArtifact artifact, IReadOnlyList<string>? fileFilter)
 		{
-			Uri baseUri = new Uri(_serverInfo.ServerUrl, $"api/v2/artifacts/{artifact.Id}");
+			Uri? baseUri = null;
+
+			// For Zen backend need to pull the build uri from attached metadata
+			bool zen = false;
+			string? backend = artifact.Metadata.FirstOrDefault(x => x.StartsWith("Backend=", StringComparison.OrdinalIgnoreCase));
+			if (!String.IsNullOrEmpty(backend) && backend.EndsWith("=Zen", StringComparison.OrdinalIgnoreCase))
+			{
+				zen = true;
+				string? zenBuildUri = artifact.Metadata.FirstOrDefault(x => x.StartsWith("ZenBuildUri=", StringComparison.OrdinalIgnoreCase));
+				if (zenBuildUri != null)
+				{
+					string[] parts = zenBuildUri.Split('=');
+					if (parts.Length == 2 && parts[1].Length > 0)
+					{
+						baseUri = new Uri(parts[1]);
+					}
+				}
+			}			
+			
+			if (baseUri == null)
+			{ 
+				if (zen)
+				{
+					return NotFound($"Unable to resolve zen uri for artifact {artifact.Id}");
+				}
+
+				baseUri = new Uri(_serverInfo.ServerUrl, $"api/v2/artifacts/{artifact.Id}");
+			}
 
 			ArtifactDescriptor descriptor = new ArtifactDescriptor(baseUri, new RefName("default"), fileFilter);
 			descriptor.Name = artifact.Name.ToString();
