Horde: Add parent/child job relationships, dashboard side

diff --git a/Engine/Source/Programs/Horde/HordeDashboard/src/backend/Api.ts b/Engine/Source/Programs/Horde/HordeDashboard/src/backend/Api.ts
--- a/Engine/Source/Programs/Horde/HordeDashboard/src/backend/Api.ts
+++ b/Engine/Source/Programs/Horde/HordeDashboard/src/backend/Api.ts
@@ -1333,6 +1333,13 @@ export type GetJobResponse = {
 
 	/**  Whether the current user is allowed to update this job */
 	canUpdate?: boolean;
+
+	/// The parent job id, if any
+	parentJobId?: string;
+
+	/// The parent job step id, if any
+	parentJobStepId?: string;
+
 }
 
 /**Request used to update a jobstep */
@@ -1451,6 +1458,9 @@ export type GetStepResponse = {
 
 	/// Annotations for this node
 	annotations?: Record<string, string>;
+
+	/// Spawned jobs
+	spawnedJobs?: string[];
 }
 
 //**Returns information about test data */
diff --git a/Engine/Source/Programs/Horde/HordeDashboard/src/components/jobDetailsV2/JobDetailViewSteps.tsx b/Engine/Source/Programs/Horde/HordeDashboard/src/components/jobDetailsV2/JobDetailViewSteps.tsx
--- a/Engine/Source/Programs/Horde/HordeDashboard/src/components/jobDetailsV2/JobDetailViewSteps.tsx
+++ b/Engine/Source/Programs/Horde/HordeDashboard/src/components/jobDetailsV2/JobDetailViewSteps.tsx
@@ -112,7 +112,7 @@ const RenderDynamic: React.FC<{ jobDetails: JobDetailsV2, dataView: StepsDataVie
 
    if (step && column === "ETA") {
 
-      let eta:HordeTime = {
+      let eta: HordeTime = {
          display: "",
          server: ""
       };
@@ -121,7 +121,7 @@ const RenderDynamic: React.FC<{ jobDetails: JobDetailsV2, dataView: StepsDataVie
          return null;
       }
 
-      let finished:HordeTime = { display: "", server: "" };
+      let finished: HordeTime = { display: "", server: "" };
 
       eta = getStepETA(step, jobDetails.jobData!);
 
@@ -303,7 +303,8 @@ export const StepsPanelInner: React.FC<{ jobDetails: JobDetailsV2, depStepId?: s
 
       const step = item.step!;
       const disabled = !item.step?.logId;
-
+      const spawnedJobs = step.spawnedJobs ?? [];
+      const spawnedJobUrl = spawnedJobs.length > 0 ? `/job/${spawnedJobs[0]}` : undefined;
 
       const renderViewLogColumn = () => {
 
@@ -334,9 +335,14 @@ export const StepsPanelInner: React.FC<{ jobDetails: JobDetailsV2, depStepId?: s
 
          case 'Name':
             // see data-automation-key in onRenderRow selectors
-            return <Link to={stepUrl}><Stack horizontal horizontalAlign="start" verticalAlign="center" style={{ width: "100%", height: "100%", paddingLeft: 8 }}><StepStatusIcon step={step} />
-               <Text style={{ paddingTop: 1 }}>{stepName}</Text>
-            </Stack></Link>;
+            return <Stack horizontal horizontalAlign="start" verticalAlign="center" verticalFill tokens={{ childrenGap: 24 }}>
+               <Stack><Link to={stepUrl}><Stack horizontal horizontalAlign="start" verticalAlign="center" style={{ width: "100%", height: "100%", paddingLeft: 8 }}><StepStatusIcon step={step} />
+                  <Text style={{ paddingTop: 1 }}>{stepName}</Text>
+               </Stack>
+               </Link>
+               </Stack>
+               {!!spawnedJobUrl && <Stack><Text>(<a href={spawnedJobUrl}> Spawned Job </a>)</Text></Stack>}
+            </Stack>;
          case 'Progress':
             return <RenderDynamic jobDetails={jobDetails} dataView={dataView} step={step} column="Progress" />
          case 'Duration':
diff --git a/Engine/Source/Programs/Horde/HordeDashboard/src/components/jobDetailsV2/JobDetailViewSummary.tsx b/Engine/Source/Programs/Horde/HordeDashboard/src/components/jobDetailsV2/JobDetailViewSummary.tsx
--- a/Engine/Source/Programs/Horde/HordeDashboard/src/components/jobDetailsV2/JobDetailViewSummary.tsx
+++ b/Engine/Source/Programs/Horde/HordeDashboard/src/components/jobDetailsV2/JobDetailViewSummary.tsx
@@ -107,6 +107,8 @@ export const SummaryPanel: React.FC<{ jobDetails: JobDetailsV2 }> = observer(({
    const summary = text.join(".  ");
 
    const reportData = jobDetails.getReportData(ReportPlacement.Summary);
+   
+   const parentJobUrl = jobData.parentJobId && jobData.parentJobId ? `/job/${jobData.parentJobId}?step=${jobData.parentJobStepId}` : undefined;
 
    return (<Stack id={sideRail.url} styles={{ root: { paddingTop: 0, paddingRight: 12 } }}>
       <Stack className={hordeClasses.raised} >
@@ -124,7 +126,7 @@ export const SummaryPanel: React.FC<{ jobDetails: JobDetailsV2 }> = observer(({
                      <Markdown>{jobData.cancellationReason}</Markdown>
                   </Stack>
                   }
-
+                  {!!parentJobUrl && <Stack><Markdown>{`This job was spawned by a [parent job](${parentJobUrl}).`}</Markdown></Stack>}
                   {!!reportData && <Stack> <Markdown>{reportData}</Markdown></Stack>}
                   {!!price && <Stack>
                      <Text>{`Estimated cost: $${price.toFixed(2)}`}</Text>
diff --git a/Engine/Source/Programs/Horde/HordeDashboard/src/components/jobDetailsV2/StepDetailSummary.tsx b/Engine/Source/Programs/Horde/HordeDashboard/src/components/jobDetailsV2/StepDetailSummary.tsx
--- a/Engine/Source/Programs/Horde/HordeDashboard/src/components/jobDetailsV2/StepDetailSummary.tsx
+++ b/Engine/Source/Programs/Horde/HordeDashboard/src/components/jobDetailsV2/StepDetailSummary.tsx
@@ -235,6 +235,10 @@ export const StepSummaryPanel: React.FC<{ jobDetails: JobDetailsV2; stepId: stri
 
    const reportData = jobDetails.getReportData(ReportPlacement.Summary, stepId);
 
+   const spawnedJobs = step.spawnedJobs ?? [];
+   const spawnedJobUrl = spawnedJobs.length > 0 ? `/job/${spawnedJobs[0]}` : undefined;
+   const parentJobUrl = jobData.parentJobId && jobData.parentJobId ? `/job/${jobData.parentJobId}?step=${jobData.parentJobStepId}` : undefined;
+
    return (<Stack id={sideRail.url} styles={{ root: { paddingTop: 0, paddingRight: 12 } }}>
       <Stack className={hordeClasses.raised}>
          <Stack tokens={{ childrenGap: 18 }}>
@@ -252,6 +256,8 @@ export const StepSummaryPanel: React.FC<{ jobDetails: JobDetailsV2; stepId: stri
                   <Markdown>{step.cancellationReason}</Markdown>
                </Stack>
                }
+               {!!parentJobUrl && <Stack style={{ paddingTop: 8}}><Markdown>{`This step was spawned by a [parent job](${parentJobUrl}).`}</Markdown></Stack>}
+               {!!spawnedJobUrl && <Stack style={{ paddingTop: 8 }}><Markdown>{`This step spawned a [new job](${spawnedJobUrl}).`}</Markdown></Stack>}
                {!!reportData && <Stack style={{ paddingTop: 8 }}> <Markdown>{reportData}</Markdown> </Stack>}
                {!!priceText && <Stack style={{ paddingTop: 8 }}>
                   <Text>{priceText}</Text>
