Horde: Properly encode zen backend into uartifact files

diff --git a/Engine/Source/Programs/Horde/Plugins/Build/HordeServer.Build/Artifacts/ArtifactsController.cs b/Engine/Source/Programs/Horde/Plugins/Build/HordeServer.Build/Artifacts/ArtifactsController.cs
--- a/Engine/Source/Programs/Horde/Plugins/Build/HordeServer.Build/Artifacts/ArtifactsController.cs
+++ b/Engine/Source/Programs/Horde/Plugins/Build/HordeServer.Build/Artifacts/ArtifactsController.cs
@@ -640,6 +640,7 @@ namespace HordeServer.Artifacts
 			descriptor.Description = artifact.Description;
 			descriptor.Keys = new List<string>(artifact.Keys);
 			descriptor.Metadata = new List<string>(artifact.Metadata);
+			descriptor.Backend = new ArtifactBackend { Type = zen ? ArtifactBackend.BackendTypeEnum.Zen : ArtifactBackend.BackendTypeEnum.Horde };
 
 			foreach (string key in artifact.Keys)
 			{
