Horde: Make parameter boolean checks case insensitive on new build dialog

diff --git a/Engine/Source/Programs/Horde/HordeDashboard/src/components/build/NewBuildV2.tsx b/Engine/Source/Programs/Horde/HordeDashboard/src/components/build/NewBuildV2.tsx
--- a/Engine/Source/Programs/Horde/HordeDashboard/src/components/build/NewBuildV2.tsx
+++ b/Engine/Source/Programs/Horde/HordeDashboard/src/components/build/NewBuildV2.tsx
@@ -640,7 +640,7 @@ class BuildOptions {
          switch (p.type) {
             case ParameterType.Bool:
                const b = p as BoolParameterData;
-               if (this.parameters[b.id] === "true") {
+               if (this.parameters[b.id]?.toLowerCase() === "true") {
                   if (b.argumentIfEnabled) {
                      allArgs.push(b.argumentIfEnabled);
                   }
@@ -674,7 +674,7 @@ class BuildOptions {
                list.items.forEach(i => {
                   allArgs = [];
 
-                  if (this.parameters[i.id] === "true") {
+                  if (this.parameters[i.id]?.toLowerCase() === "true") {
                      if (i.argumentIfEnabled) {
                         allArgs.push(i.argumentIfEnabled);
                      }
@@ -1033,7 +1033,7 @@ const BoolParameter: React.FC<{ param: BoolParameterData, disabled?: boolean }>
       <Checkbox key={key}
          label={param.label}
          disabled={options.readOnly || disabled}
-         checked={options.parameters[param.id] == "true"}
+         checked={options.parameters[param.id]?.toLowerCase() === "true"}
          onChange={(ev, value) => {
             ev?.preventDefault();
             options.onBooleanChanged(param.id, value ?? false);
@@ -1095,7 +1095,7 @@ const TagPickerParameter: React.FC<{ param: ListParameterData, disabled?: boolea
    });
 
    const selectedItems = allItems.filter(item => {
-      return options.parameters[item.itemData.id] === "true";
+      return options.parameters[item.itemData.id]?.toLowerCase() === "true";
    });
 
    // tag picker
@@ -1148,7 +1148,7 @@ const BasicListParameter: React.FC<{ param: ListParameterData, disabled?: boolea
       doptions.push({
          key: item.id,
          text: item.text,
-         selected: options.parameters[item.id] === "true"
+         selected: options.parameters[item.id]?.toLowerCase() === "true"
       });
    });
 
@@ -1217,7 +1217,7 @@ const MultiListParameter: React.FC<{ param: ListParameterData, disabled?: boolea
       param.items.forEach(item => {
          if (item.group === group) {
             const key = item.id;
-            const selected = options.parameters[item.id] === "true"
+            const selected = options.parameters[item.id]?.toLowerCase() === "true"
             if (selected) {
                selectedKeys.push(key);
             }
@@ -1310,7 +1310,7 @@ const LongListParameter: React.FC<{ param: ListParameterData, disabled?: boolean
       return !(item.name.startsWith("**") && item.name.endsWith( "**" ));
    });
    const selectedItems = allItems.values().filter(item => {
-         return options.parameters[item.itemData.id] === "true" && selectFilter( item );
+         return options.parameters[item.itemData.id]?.toLowerCase() === "true" && selectFilter( item );
       })
       .toArray();
 
