Horde: Update release notes
#rnx

diff --git a/Engine/Source/Programs/Horde/Docs/ReleaseNotes.md b/Engine/Source/Programs/Horde/Docs/ReleaseNotes.md
--- a/Engine/Source/Programs/Horde/Docs/ReleaseNotes.md
+++ b/Engine/Source/Programs/Horde/Docs/ReleaseNotes.md
@@ -1,5 +1,16 @@
 # Release Notes
 
+## 2025-03-06
+
+* Update approved licenses (HTML changes) (40539666)
+* Properly encode zen backend into uartifact files (40510907)
+* Add parent/child job relationships, dashboard side (40508326)
+* Adding zen backend to artifacts controller uartifact generation via metadata (40502152)
+* Add device cleaning information to devices endpoint (40463300)
+* Remove incorrect return type metadata for PUT /api/v2/devices/{deviceId} (40461854)
+* Added ability to specify backends in uartifact files and added a prototype of how it could be used in the create cloud artifact task to generate a uartifact file that UGS could sync. (40451830)
+* Update notes on Epic's deployment (40349855)
+
 ## 2025-02-20
 
 * Explain GitHub image registry auth better in docs (40164690)
