Horde: Add tooltips to issue navigation buttons

diff --git a/Engine/Source/Programs/Horde/HordeDashboard/src/components/LogView.tsx b/Engine/Source/Programs/Horde/HordeDashboard/src/components/LogView.tsx
--- a/Engine/Source/Programs/Horde/HordeDashboard/src/components/LogView.tsx
+++ b/Engine/Source/Programs/Horde/HordeDashboard/src/components/LogView.tsx
@@ -1416,84 +1416,97 @@ export const LogList: React.FC<{ logId: string }> = observer(({ logId }) => {
 
                         <Stack horizontal verticalAlign='center'>
                            <Stack horizontal verticalAlign='center' style={{position: 'relative'}}>
-                              <DefaultButton disabled={!nextErrorEnabled} className={nextErrorEnabled ? handler.style.errorButton : handler.style.errorButtonDisabled}
-                                 text={`${errorText} ${errors.length === 1 ? "Error" : "Errors"}`}
-                                 onClick={() => {
-                                    nextError();
-                                 }}
-                                 style={{ color: errors.length ? "#F9F9FB" : buttonNoneText, height: 30, fontSize: 19, paddingLeft: 12, borderRadius: '4px 0 0 4px' }}>
-                                 {!!errors.length && <Icon style={{ fontSize: 19, paddingLeft: 4, paddingRight: 0 }} iconName='ChevronDown' />}
-                              </DefaultButton>
+                              <TooltipHost content={nextErrorEnabled ? 'Navigate to next error': ''}>
+                                 <DefaultButton disabled={!nextErrorEnabled} className={nextErrorEnabled ? handler.style.errorButton : handler.style.errorButtonDisabled}
+                                    text={`${errorText} ${errors.length === 1 ? "Error" : "Errors"}`}
+                                    onClick={() => {
+                                       nextError();
+                                    }}
+                                    style={{ color: errors.length ? "#F9F9FB" : buttonNoneText, height: 30, fontSize: 19, paddingLeft: 12, borderRadius: '4px 0 0 4px' }}>
+                                    {!!errors.length && <Icon style={{ fontSize: 19, paddingLeft: 4, paddingRight: 0 }} iconName='ChevronDown' />}
+                                 </DefaultButton>
+                              </TooltipHost>
                               {!!errors.length && <div style={{backgroundColor: "#F9F9FB", height: 18, width: 1, position: 'absolute', right: -0.5, zIndex: 1 }} />}
                            </Stack>
-                           {!!errors.length && <IconButton disabled={!prevErrorEnabled} className={prevErrorEnabled ? handler.style.errorButton : handler.style.errorButtonDisabled} style={{ color: "#F9F9FB", height: 30, fontSize: 19, padding: 2, borderRadius: '0 4px 4px 0' }} iconProps={{ iconName: 'ChevronUp' }} onClick={(event: any) => {
-                              event?.stopPropagation();
-                              prevError();
-
-                           }} />}
+                           {!!errors.length && <TooltipHost content={'Navigate to previous error'}>
+                              <IconButton disabled={!prevErrorEnabled} className={prevErrorEnabled ? handler.style.errorButton : handler.style.errorButtonDisabled} style={{ color: "#F9F9FB", height: 30, fontSize: 19, padding: 2, borderRadius: '0 4px 4px 0' }} iconProps={{ iconName: 'ChevronUp' }} onClick={(event: any) => {
+                                 event?.stopPropagation();
+                                 prevError();
+                              }} />
+                           </TooltipHost>}
                         </Stack>
 
                         {!!errors.length && <Stack horizontal verticalAlign='center' style={{paddingLeft: 1}} >
                            <Stack horizontal verticalAlign='center' style={{position: 'relative'}}>
-                              <IconButton disabled={!nextErrorBlockEnabled} className={nextErrorBlockEnabled ? handler.style.errorButton : handler.style.errorButtonDisabled}
-                                 style={{ color: errors.length ? "#F9F9FB" : buttonNoneText, height: 30, fontSize: 19, padding: 9, borderRadius: '4px 0 0 4px' }}
-                                 styles={{rootPressed: {border: 'none', boxShadow: 'none'}}} 
-                                 iconProps={{ iconName: 'DoubleChevronDown' }}
-                                 onClick={() => {
-                                    nextErrorBlock();
-                                 }} >
-                                 </IconButton>
+                              <TooltipHost content={'Skip to next block of errors or issue'}>
+                                 <IconButton disabled={!nextErrorBlockEnabled} className={nextErrorBlockEnabled ? handler.style.errorButton : handler.style.errorButtonDisabled}
+                                    style={{ color: errors.length ? "#F9F9FB" : buttonNoneText, height: 30, fontSize: 19, padding: 9, borderRadius: '4px 0 0 4px' }}
+                                    styles={{rootPressed: {border: 'none', boxShadow: 'none'}}} 
+                                    iconProps={{ iconName: 'DoubleChevronDown' }}
+                                    onClick={() => {
+                                       nextErrorBlock();
+                                    }} />
+                              </TooltipHost>
                                  <div style={{backgroundColor: "#F9F9FB", height: 18, width: 1, position: 'absolute', right: -0.5, zIndex: 1 }}></div>
                            </Stack>
-                           <IconButton disabled={!prevErrorBlockEnabled} className={prevErrorBlockEnabled ? handler.style.errorButton : handler.style.errorButtonDisabled}
-                              style={{ color: errors.length ? "#F9F9FB" : buttonNoneText, height: 30, fontSize: 19, padding: 9,  borderRadius: '0 4px 4px 0'}} 
-                              iconProps={{ iconName: 'DoubleChevronUp' }}
-                              onClick={(event: any) => {
-                                 event?.stopPropagation();
-                                 prevErrorBlock();
-                              }} />
+                           <TooltipHost content={'Skip to previous block of errors or issue'}>
+                              <IconButton disabled={!prevErrorBlockEnabled} className={prevErrorBlockEnabled ? handler.style.errorButton : handler.style.errorButtonDisabled}
+                                 style={{ color: errors.length ? "#F9F9FB" : buttonNoneText, height: 30, fontSize: 19, padding: 9,  borderRadius: '0 4px 4px 0'}} 
+                                 iconProps={{ iconName: 'DoubleChevronUp' }}
+                                 onClick={(event: any) => {
+                                    event?.stopPropagation();
+                                    prevErrorBlock();
+                                 }} />
+                           </TooltipHost>
                         </Stack>}
 
 
                         <Stack horizontal verticalAlign='center' style={{ paddingLeft: 12 }} >
                            <Stack horizontal verticalAlign='center' style={{position: 'relative'}}>
-                              <DefaultButton disabled={!nextWarningEnabled} className={nextWarningEnabled ? handler.style.warningButton : handler.style.warningButtonDisabled}
-                                 text={`${warningsText} ${warnings.length === 1 ? "Warning" : "Warnings"}`}
-                                 onClick={() => {
-                                    nextWarning();
-                                 }}
-                                 style={{ color: (dashboard.darktheme && nextWarningEnabled) ? "#F9F9FB" : buttonNoneText, height: 30, fontSize: 19, padding: 12, borderRadius: '4px 0 0 4px'}} >
-                                 {!!warnings.length && <Icon style={{ fontSize: 19, paddingLeft: 4, paddingRight: 0 }} iconName='ChevronDown' />}
-                              </DefaultButton>
+                              <TooltipHost content={nextWarningEnabled ? 'Navigate to next warning': ''}>
+                                 <DefaultButton disabled={!nextWarningEnabled} className={nextWarningEnabled ? handler.style.warningButton : handler.style.warningButtonDisabled}
+                                    text={`${warningsText} ${warnings.length === 1 ? "Warning" : "Warnings"}`}
+                                    onClick={() => {
+                                       nextWarning();
+                                    }}
+                                    style={{ color: (dashboard.darktheme && nextWarningEnabled) ? "#F9F9FB" : buttonNoneText, height: 30, fontSize: 19, padding: 12, borderRadius: '4px 0 0 4px'}} >
+                                    {!!warnings.length && <Icon style={{ fontSize: 19, paddingLeft: 4, paddingRight: 0 }} iconName='ChevronDown' />}
+                                 </DefaultButton>
+                              </TooltipHost>     
                               {!!warnings.length && <div style={{backgroundColor: (dashboard.darktheme && nextWarningEnabled) ? "#F9F9FB" : buttonNoneText, height: 18, width: 1, position: 'absolute', right: -0.5, zIndex: 1 }} />}
                            </Stack>
                            {!!warnings.length && <Stack>
-                              <IconButton disabled={!prevWarningEnabled} className={prevWarningEnabled ? handler.style.warningButton : handler.style.warningButtonDisabled} style={{ height: 30, fontSize: 19, padding: 2, color: (dashboard.darktheme && nextWarningEnabled) ? "#F9F9FB" : buttonNoneText, borderRadius: '0 4px 4px 0'}} iconProps={{ iconName: 'ChevronUp' }} onClick={(event: any) => {
-                                 event?.stopPropagation();
-                                 prevWarning();
-                              }} />
+                              <TooltipHost content={'Navigate to previous warning'}>
+                                 <IconButton disabled={!prevWarningEnabled} className={prevWarningEnabled ? handler.style.warningButton : handler.style.warningButtonDisabled} style={{ height: 30, fontSize: 19, padding: 2, color: (dashboard.darktheme && nextWarningEnabled) ? "#F9F9FB" : buttonNoneText, borderRadius: '0 4px 4px 0'}} iconProps={{ iconName: 'ChevronUp' }} onClick={(event: any) => {
+                                    event?.stopPropagation();
+                                    prevWarning();
+                                 }} />
+                              </TooltipHost>
                            </Stack>}
                         </Stack>
 
                         {!!warnings.length && <Stack horizontal verticalAlign='center' style={{paddingLeft: 1}} >
                            <Stack horizontal verticalAlign='center' style={{position: 'relative'}}>
-                              <IconButton disabled={!nextWarningBlockEnabled} className={nextWarningBlockEnabled ? handler.style.warningButton : handler.style.warningButtonDisabled}
-                                 style={{ color: (dashboard.darktheme && nextWarningEnabled) ? "#F9F9FB" : buttonNoneText, height: 30, fontSize: 19, padding: 9, borderRadius: '4px 0 0 4px' }}
-                                 styles={{rootPressed: {border: 'none', boxShadow: 'none'}}} 
-                                 iconProps={{ iconName: 'DoubleChevronDown' }}
-                                 onClick={() => {
-                                    nextWarningBlock();
-                                 }} >
-                                 </IconButton>
+                              <TooltipHost content={'Skip to next block of warnings or issue'}>
+                                 <IconButton disabled={!nextWarningBlockEnabled} className={nextWarningBlockEnabled ? handler.style.warningButton : handler.style.warningButtonDisabled}
+                                    style={{ color: (dashboard.darktheme && nextWarningEnabled) ? "#F9F9FB" : buttonNoneText, height: 30, fontSize: 19, padding: 9, borderRadius: '4px 0 0 4px' }}
+                                    styles={{rootPressed: {border: 'none', boxShadow: 'none'}}} 
+                                    iconProps={{ iconName: 'DoubleChevronDown' }}
+                                    onClick={() => {
+                                       nextWarningBlock();
+                                    }} />
+                              </TooltipHost>
                                  <div style={{backgroundColor: (dashboard.darktheme && nextWarningEnabled) ? "#F9F9FB" : buttonNoneText, height: 18, width: 1, position: 'absolute', right: -0.5, zIndex: 1 }}></div>
                            </Stack>
-                           <IconButton disabled={!prevWarningBlockEnabled} className={prevWarningBlockEnabled ? handler.style.warningButton : handler.style.warningButtonDisabled}
-                              style={{ color: (dashboard.darktheme && nextWarningEnabled) ? "#F9F9FB" : buttonNoneText, height: 30, fontSize: 19, padding: 9,  borderRadius: '0 4px 4px 0'}} 
-                              iconProps={{ iconName: 'DoubleChevronUp' }}
-                              onClick={(event: any) => {
-                                 event?.stopPropagation();
-                                 prevWarningBlock();
-                              }} />
+                           <TooltipHost content={'Skip to previous block of warnings or issue'}>
+                              <IconButton disabled={!prevWarningBlockEnabled} className={prevWarningBlockEnabled ? handler.style.warningButton : handler.style.warningButtonDisabled}
+                                 style={{ color: (dashboard.darktheme && nextWarningEnabled) ? "#F9F9FB" : buttonNoneText, height: 30, fontSize: 19, padding: 9,  borderRadius: '0 4px 4px 0'}} 
+                                 iconProps={{ iconName: 'DoubleChevronUp' }}
+                                 onClick={(event: any) => {
+                                    event?.stopPropagation();
+                                    prevWarningBlock();
+                                 }} />
+                           </TooltipHost>
                         </Stack>}
 
                         {!!menuProps.items.find(i => !i.disabled) && <Stack style={{ paddingLeft: 18 }}><PrimaryButton disabled={!menuProps.items.find(i => !i.disabled)} text="View" menuProps={menuProps} style={{ fontFamily: "Horde Open Sans SemiBold", borderStyle: "hidden", padding: 15 }} /></Stack>}
