Horde: Improve description of OIDC client secret config value to clarify public vs confidential client requirements

diff --git a/Engine/Source/Programs/Horde/HordeServer/ServerSettings.cs b/Engine/Source/Programs/Horde/HordeServer/ServerSettings.cs
--- a/Engine/Source/Programs/Horde/HordeServer/ServerSettings.cs
+++ b/Engine/Source/Programs/Horde/HordeServer/ServerSettings.cs
@@ -180,7 +180,10 @@ namespace HordeServer
 		public string? OidcClientId { get; set; }
 
 		/// <summary>
-		/// Client secret for the OIDC authority
+		/// Client secret for authenticating with the OIDC provider.
+		/// Note: If you need authentication support in Unreal Build Tool or Unreal Game Sync,
+		/// configure your OIDC client as a public client (using PKCE flow without a client secret) instead of a confidential client.
+		/// These tools utilize the EpicGames.OIDC library which only supports public clients with authorization code flow + PKCE.
 		/// </summary>
 		public string? OidcClientSecret { get; set; }
 
