Horde: Constrain jupiterdownloads to UGS
#submitool bypass

diff --git a/Engine/Source/Programs/Horde/HordeDashboard/src/components/artifacts/ArtifactButton.tsx b/Engine/Source/Programs/Horde/HordeDashboard/src/components/artifacts/ArtifactButton.tsx
--- a/Engine/Source/Programs/Horde/HordeDashboard/src/components/artifacts/ArtifactButton.tsx
+++ b/Engine/Source/Programs/Horde/HordeDashboard/src/components/artifacts/ArtifactButton.tsx
@@ -12,10 +12,9 @@ import { useNavigate } from "react-router-dom";
 export const ArtifactButton: React.FC<{ jobId?: string, stepId?: string, artifact?: GetArtifactResponse, pathIn?: string, selection?: ArtifactSelection, openArtifactInfo?: () => void, disabled?: boolean, minWidth?: number }> = ({ jobId, stepId, artifact, pathIn, selection, openArtifactInfo, disabled, minWidth }) => {
 
    const [state, setState] = useState<{ showZipDownloadWarning?: boolean }>({});
-   
-   const navigate = useNavigate();
 
-   const format = dashboard.downloadFormat;
+   const navigate = useNavigate();
+   
    const path = pathIn ? pathIn + "/" : "";
    const items = selection?.items;
 
@@ -24,6 +23,10 @@ export const ArtifactButton: React.FC<{ jobId?: string, stepId?: string, artifac
       sizeText = formatBytes(selection.size, (selection.size < (1024 * 1024)) ? 0 : 1)
    }
 
+   const zenArtifact = !!artifact?.metadata?.find(m => m.toLowerCase() === "backend=zen");
+
+   const format = zenArtifact ? DownloadFormat.UGS : dashboard.downloadFormat;
+
    let buttonText = "Download";
 
    if (selection?.directoriesSelected || (selection?.filesSelected ?? 0) > 1) {
@@ -147,15 +150,17 @@ export const ArtifactButton: React.FC<{ jobId?: string, stepId?: string, artifac
 
    if (artifact?.id) {
 
-      downloadProps.items.push(
-         {
-            key: 'download_zip',
-            text: 'Download Zip',
-            onClick: () => {
-               showDownloadWarning() ? setState({ showZipDownloadWarning: true }) : downloadZip()
+      if (!zenArtifact) {
+         downloadProps.items.push(
+            {
+               key: 'download_zip',
+               text: 'Download Zip',
+               onClick: () => {
+                  showDownloadWarning() ? setState({ showZipDownloadWarning: true }) : downloadZip()
+               }
             }
-         }
-      )
+         )
+      }
 
       downloadProps.items.push(
          {
@@ -166,15 +171,18 @@ export const ArtifactButton: React.FC<{ jobId?: string, stepId?: string, artifac
             }
          }
       )
-      downloadProps.items.push(
-         {
-            key: 'download_toolbox',
-            text: 'Download with Toolbox',
-            onClick: () => {
-               downloadToolbox();
+
+      if (!zenArtifact) {
+         downloadProps.items.push(
+            {
+               key: 'download_toolbox',
+               text: 'Download with Toolbox',
+               onClick: () => {
+                  downloadToolbox();
+               }
             }
-         }
-      )
+         )
+      }
 
       downloadProps.items.push(
          {
@@ -199,9 +207,9 @@ export const ArtifactButton: React.FC<{ jobId?: string, stepId?: string, artifac
       }
 
       return true
-   }   
+   }
 
-   return <Stack>{!!state.showZipDownloadWarning && <ZipDownloadWarning onContinue={() => { downloadZip(); setState({ showZipDownloadWarning: false }); }} onClose={() => { setState({ showZipDownloadWarning: false }) }}/>}
+   return <Stack>{!!state.showZipDownloadWarning && <ZipDownloadWarning onContinue={() => { downloadZip(); setState({ showZipDownloadWarning: false }); }} onClose={() => { setState({ showZipDownloadWarning: false }) }} />}
       <PrimaryButton split menuProps={downloadProps} text={buttonText} styles={{ root: { minWidth: minWidth ?? 104, fontFamily: 'Horde Open Sans SemiBold !important' } }} disabled={disabled} onClick={() => {
          if (format === DownloadFormat.Zip) {
             showDownloadWarning() ? setState({ showZipDownloadWarning: true }) : downloadZip()
@@ -218,11 +226,11 @@ export const ArtifactButton: React.FC<{ jobId?: string, stepId?: string, artifac
 }
 
 const ZipDownloadWarning: React.FC<{ onContinue: () => void, onClose: () => void }> = ({ onContinue, onClose }) => {
-   
+
    const { hordeClasses } = getHordeStyling();
 
    return <Modal className={hordeClasses.modal} isOpen={true} isBlocking={true} topOffsetFixed={true} styles={{ main: { padding: 8, height: 200, width: 600, hasBeenOpened: false, top: 0, bottom: 150, left: 0, right: 0, margin: "auto", position: "absolute" } }} >
-      
+
       <Stack style={{ padding: 8 }}>
          <Stack style={{ paddingBottom: 24 }}>
             <Text variant="mediumPlus" style={{ fontFamily: "Horde Open Sans SemiBold" }}>Warning: Slow Download</Text>
@@ -231,7 +239,7 @@ const ZipDownloadWarning: React.FC<{ onContinue: () => void, onClose: () => void
             <Text variant="medium" style={{ fontFamily: "Horde Open Sans SemiBold" }}>Downloading artifacts in zip format is significantly slower than downloading with UGS.</Text>
          </Stack>
          <Stack>
-            <Checkbox label="Don't show this warning again" onChange={(ev?: any, checked?: boolean | undefined) => {localStorage.setItem("horde_hide_zip_warning", checked!.toString())}}/>
+            <Checkbox label="Don't show this warning again" onChange={(ev?: any, checked?: boolean | undefined) => { localStorage.setItem("horde_hide_zip_warning", checked!.toString()) }} />
          </Stack>
          <Stack horizontal style={{ paddingTop: 16 }}>
             <Stack grow />
diff --git a/Engine/Source/Programs/Horde/HordeDashboard/src/components/jobDetailsV2/JobDetailArtifacts.tsx b/Engine/Source/Programs/Horde/HordeDashboard/src/components/jobDetailsV2/JobDetailArtifacts.tsx
--- a/Engine/Source/Programs/Horde/HordeDashboard/src/components/jobDetailsV2/JobDetailArtifacts.tsx
+++ b/Engine/Source/Programs/Horde/HordeDashboard/src/components/jobDetailsV2/JobDetailArtifacts.tsx
@@ -116,10 +116,17 @@ export const JobArtifactsPanel: React.FC<{ jobDetails: JobDetailsV2 }> = observe
    let grouped: GetJobArtifactResponse[] = [];
    const ungrouped: GetJobArtifactResponse[] = [];
 
+   const zenArtifacts = new Set<string>();
+
    artifacts.forEach(a => {
 
       let category = "";
 
+      if(a.metadata.find(m => m.toLowerCase() === "backend=zen"))
+      {
+         zenArtifacts.add(a.id);
+      }
+
       const m = a.metadata.find(m => m.toLowerCase().startsWith("dashboard-category"));
 
       if (m) {
@@ -207,13 +214,13 @@ export const JobArtifactsPanel: React.FC<{ jobDetails: JobDetailsV2 }> = observe
       if (column.key === 'column_download') {
          return <Stack horizontalAlign="end" verticalAlign="center" verticalFill={true} style={{ paddingRight: 8 }}>
             <Stack horizontal tokens={{ childrenGap: 18 }}>
-               <PrimaryButton text="Browse"
+               {!zenArtifacts.has(item.id) && <PrimaryButton text="Browse"
                   disabled={!stepFinished || !item.id}
                   style={{ fontFamily: "Horde Open Sans SemiBold" }}
                   onClick={() => {
                      setSelected(item)
                   }}
-               />
+               />}
                <ArtifactButton   
                   disabled={!stepFinished || !item.id}
                   artifact={item}
