Horde supportsPartitionedWorkspace doc updates.

diff --git a/Engine/Source/Programs/Horde/Docs/Tutorials/BuildAutomation.md b/Engine/Source/Programs/Horde/Docs/Tutorials/BuildAutomation.md
--- a/Engine/Source/Programs/Horde/Docs/Tutorials/BuildAutomation.md
+++ b/Engine/Source/Programs/Horde/Docs/Tutorials/BuildAutomation.md
@@ -38,6 +38,8 @@ ensuring that the state of a project is continuously being monitored and that bu
    * You may wish to go to the Horde Dashboard ([http://localhost:13340](http://localhost:13340) by default),
      open the **Server** menu, and select **Status**. This page shows the status of various server subsystems,
      including the Perforce connection status - which should be confirmed as working.
+   * Consider setting `supportsPartitionedWorkspaces` to true, in order to reduce Perforce journal load 
+    ([https://portal.perforce.com/s/article/17297](Perforce documentation)) 
 
 4. Enable the example project by uncommenting the following lines in `globals.json`:
 
