Horde: Updating schema documentation

#rnx

diff --git a/Engine/Source/Programs/Horde/Docs/Config/Schema/AclActions.md b/Engine/Source/Programs/Horde/Docs/Config/Schema/AclActions.md
--- a/Engine/Source/Programs/Horde/Docs/Config/Schema/AclActions.md
+++ b/Engine/Source/Programs/Horde/Docs/Config/Schema/AclActions.md
@@ -348,6 +348,19 @@
 | `AddComputeTasks` |  |
 | `GetComputeTasks` |  |
 
+## Ddc
+
+| Name | Description |
+| ---- | ----------- |
+| `DdcReadObject` | General read access to refs / blobs and so on |
+| `DdcWriteObject` | General write access to upload refs / blobs etc |
+| `DdcDeleteObject` | Access to delete blobs / refs etc |
+| `DdcDeleteBucket` | Access to delete a particular bucket |
+| `DdcDeleteNamespace` | Access to delete a whole namespace |
+| `DdcReadTransactionLog` | Access to read the transaction log |
+| `DdcWriteTransactionLog` | Access to write the transaction log |
+| `DdcAdminAction` | Access to perform administrative task |
+
 ## Devices
 
 | Name | Description |
@@ -507,14 +520,14 @@
 
 | Name | Description |
 | ---- | ----------- |
-| `DdcReadObject` | General read access to refs / blobs and so on |
-| `DdcWriteObject` | General write access to upload refs / blobs etc |
-| `DdcDeleteObject` | Access to delete blobs / refs etc |
-| `DdcDeleteBucket` | Access to delete a particular bucket |
-| `DdcDeleteNamespace` | Access to delete a whole namespace |
-| `DdcReadTransactionLog` | Access to read the transaction log |
-| `DdcWriteTransactionLog` | Access to write the transaction log |
-| `DdcAdminAction` | Access to perform administrative task |
+| `DdcReadObject` |  |
+| `DdcWriteObject` |  |
+| `DdcDeleteObject` |  |
+| `DdcDeleteBucket` |  |
+| `DdcDeleteNamespace` |  |
+| `DdcReadTransactionLog` |  |
+| `DdcWriteTransactionLog` |  |
+| `DdcAdminAction` |  |
 
 ## Devices
 
@@ -592,6 +605,12 @@
 | `UpdateReplicator` |  |
 | `ViewReplicator` |  |
 
+## Secrets
+
+| Name | Description |
+| ---- | ----------- |
+| `ViewSecret` | View a credential |
+
 ## ServiceAccounts
 
 | Name | Description |
@@ -764,7 +783,7 @@
 
 | Name | Description |
 | ---- | ----------- |
-| `ViewSecret` | View a credential |
+| `ViewSecret` |  |
 
 ## ServiceAccounts
 
@@ -964,19 +983,6 @@
 | `DownloadSoftware` |  |
 | `DeleteSoftware` |  |
 
-## Storage
-
-| Name | Description |
-| ---- | ----------- |
-| `ReadBlobs` | Ability to read blobs from the storage service |
-| `WriteBlobs` | Ability to write blobs to the storage service |
-| `ReadAliases` | Ability to read aliases from the storage service |
-| `WriteAliases` | Ability to write aliases to the storage service |
-| `DeleteAliases` | Ability to write aliases to the storage service |
-| `ReadRefs` | Ability to read refs from the storage service |
-| `WriteRefs` | Ability to write refs to the storage service |
-| `DeleteRefs` | Ability to delete refs |
-
 ## Streams
 
 | Name | Description |
@@ -988,6 +994,12 @@
 | `ViewChanges` |  |
 | `ViewTemplate` |  |
 
+## Symbols
+
+| Name | Description |
+| ---- | ----------- |
+| `ReadSymbols` | Ability to download symbols |
+
 ## Accounts
 
 | Name | Description |
@@ -1155,14 +1167,14 @@
 
 | Name | Description |
 | ---- | ----------- |
-| `ReadBlobs` |  |
-| `WriteBlobs` |  |
-| `ReadAliases` |  |
-| `WriteAliases` |  |
-| `DeleteAliases` |  |
-| `ReadRefs` |  |
-| `WriteRefs` |  |
-| `DeleteRefs` |  |
+| `ReadBlobs` | Ability to read blobs from the storage service |
+| `WriteBlobs` | Ability to write blobs to the storage service |
+| `ReadAliases` | Ability to read aliases from the storage service |
+| `WriteAliases` | Ability to write aliases to the storage service |
+| `DeleteAliases` | Ability to write aliases to the storage service |
+| `ReadRefs` | Ability to read refs from the storage service |
+| `WriteRefs` | Ability to write refs to the storage service |
+| `DeleteRefs` | Ability to delete refs |
 
 ## Streams
 
@@ -1179,7 +1191,7 @@
 
 | Name | Description |
 | ---- | ----------- |
-| `ReadSymbols` | Ability to download symbols |
+| `ReadSymbols` |  |
 
 ## Accounts
 
diff --git a/Engine/Source/Programs/Horde/Docs/Deployment/ServerSettings.md b/Engine/Source/Programs/Horde/Docs/Deployment/ServerSettings.md
--- a/Engine/Source/Programs/Horde/Docs/Deployment/ServerSettings.md
+++ b/Engine/Source/Programs/Horde/Docs/Deployment/ServerSettings.md
@@ -29,7 +29,7 @@ Name | Description
 `oidcAuthority` | `string`<br>OpenID Connect (OIDC) authority URL (required when OIDC is enabled)
 `oidcAudience` | `string`<br>Audience for validating externally issued tokens (required when OIDC is enabled)
 `oidcClientId` | `string`<br>Client ID for the OIDC authority (required when OIDC is enabled)
-`oidcClientSecret` | `string`<br>Client secret for the OIDC authority
+`oidcClientSecret` | `string`<br>Client secret for authenticating with the OIDC provider. Note: If you need authentication support in Unreal Build Tool or Unreal Game Sync, configure your OIDC client as a public client (using PKCE flow without a client secret) instead of a confidential client. These tools utilize the EpicGames.OIDC library which only supports public clients with authorization code flow + PKCE.
 `oidcSigninRedirect` | `string`<br>Optional redirect url provided to OIDC login
 `oidcLocalRedirectUrls` | `string[]`<br>Optional redirect url provided to OIDC login for external tools (typically to a local server) Default value is the local web server started during signin by EpicGames.OIDC library
 `oidcDebugMode` | `boolean`<br>Debug mode for OIDC which logs reasons for why JWT tokens fail to authenticate Also turns off HTTPS requirement for OIDC metadata fetching. NOT FOR PRODUCTION USE!
@@ -130,11 +130,11 @@ Name | Description
 `analytics` | [AnalyticsServerConfig](#analyticsserverconfig)<br>Configuration for the analytics plugin
 `build` | [BuildServerConfig](#buildserverconfig)<br>Configuration for the build plugin
 `compute` | [ComputeServerConfig](#computeserverconfig)<br>Configuration for the compute plugin
-`experimental` | [ExperimentalServerConfig](#experimentalserverconfig)<br>Configuration for the experimental plugin
 `ddc` | [PluginServerConfig](#pluginserverconfig)<br>Configuration for the ddc plugin
 `secrets` | [PluginServerConfig](#pluginserverconfig)<br>Configuration for the secrets plugin
-`storage` | [StorageServerConfig](#storageserverconfig)<br>Configuration for the storage plugin
+`experimental` | [ExperimentalServerConfig](#experimentalserverconfig)<br>Configuration for the experimental plugin
 `symbols` | [PluginServerConfig](#pluginserverconfig)<br>Configuration for the symbols plugin
+`storage` | [StorageServerConfig](#storageserverconfig)<br>Configuration for the storage plugin
 `tools` | [ToolsServerConfig](#toolsserverconfig)<br>Configuration for the tools plugin
 
 ## AnalyticsServerConfig
@@ -360,17 +360,17 @@ Name | Description
 `ComputeQueueAwsMetric` | A no-op strategy that reports metrics to let an external AWS auto-scaling policy scale the fleet
 `LeaseUtilizationAwsMetric` | A no-op strategy that reports metrics to let an external AWS auto-scaling policy scale the fleet
 
-## ExperimentalServerConfig
+## PluginServerConfig
 
-Server configuration for the experimental plugin
+Base class for plugin server config objects
 
 Name | Description
 ---- | -----------
 `enabled` | `boolean`<br>Whether the plugin should be enabled or not
 
-## PluginServerConfig
+## ExperimentalServerConfig
 
-Base class for plugin server config objects
+Server configuration for the experimental plugin
 
 Name | Description
 ---- | -----------
