Implement Horde TestData v2 collectioms, schema digest and services

#rnx

diff --git a/Engine/Source/Programs/Horde/Plugins/Build/HordeServer.Build.Tests/BuildTestSetup.cs b/Engine/Source/Programs/Horde/Plugins/Build/HordeServer.Build.Tests/BuildTestSetup.cs
--- a/Engine/Source/Programs/Horde/Plugins/Build/HordeServer.Build.Tests/BuildTestSetup.cs
+++ b/Engine/Source/Programs/Horde/Plugins/Build/HordeServer.Build.Tests/BuildTestSetup.cs
@@ -73,6 +73,7 @@ namespace HordeServer.Tests
 		public ILeaseCollection LeaseCollection => ServiceProvider.GetRequiredService<ILeaseCollection>();
 		public ILogCollection LogCollection => ServiceProvider.GetRequiredService<ILogCollection>();
 		public ITestDataCollection TestDataCollection => ServiceProvider.GetRequiredService<ITestDataCollection>();
+		public ITestDataCollectionV2 TestDataCollectionV2 => ServiceProvider.GetRequiredService<ITestDataCollectionV2>();
 		public IUserCollection UserCollection => ServiceProvider.GetRequiredService<IUserCollection>();
 		public IDeviceCollection DeviceCollection => ServiceProvider.GetRequiredService<IDeviceCollection>();
 		public IDashboardPreviewCollection DashboardPreviewCollection => ServiceProvider.GetRequiredService<IDashboardPreviewCollection>();
@@ -144,6 +145,7 @@ namespace HordeServer.Tests
 			services.AddSingleton<IStreamCollection, StreamCollection>();
 			services.AddSingleton<ITemplateCollection, TemplateCollection>();
 			services.AddSingleton<ITestDataCollection, TestDataCollection>();
+			services.AddSingleton<ITestDataCollectionV2, TestDataCollectionV2>();
 			services.AddSingleton<IUtilizationDataCollection, UtilizationDataCollection>();
 			services.AddSingleton<ITemplateCollection, TemplateCollection>();
 			services.AddSingleton<IToolCollection, ToolCollection>();
diff --git a/Engine/Source/Programs/Horde/Plugins/Build/HordeServer.Build.Tests/Jobs/TestData/TestDataTests.cs b/Engine/Source/Programs/Horde/Plugins/Build/HordeServer.Build.Tests/Jobs/TestData/TestDataTests.cs
--- a/Engine/Source/Programs/Horde/Plugins/Build/HordeServer.Build.Tests/Jobs/TestData/TestDataTests.cs
+++ b/Engine/Source/Programs/Horde/Plugins/Build/HordeServer.Build.Tests/Jobs/TestData/TestDataTests.cs
@@ -328,6 +328,129 @@ namespace HordeServer.Tests.Jobs.TestData
 			Assert.AreEqual(1, refs[0].SuiteErrorCount);
 		}
 
+		[TestMethod]
+		public async Task SessionReportTestV2Async()
+		{
+			await TestDataService.StartAsync(CancellationToken.None);
+
+			string[] streamIds = new string[] { _mainStreamId.ToString(), _releaseStreamId.ToString() };
+
+			IJob job = CreateJob(_mainStreamId, 105, "Test Build", _graph);
+			IJobStep step = job.Batches[0].Steps[0];
+			IJob job2 = CreateJob(_releaseStreamId, 105, "Test Build", _graph);
+			IJobStep step2 = job2.Batches[0].Steps[0];
+			IJob job3 = CreateJob(_mainStreamId, 106, "Test Build", _graph);
+			IJobStep step3 = job3.Batches[0].Steps[0];
+
+			BsonDocument testData = BsonDocument.Parse(String.Join('\n', _sessionTestDataV2Lines));
+			List<(string key, BsonDocument)> data = new List<(string key, BsonDocument)>();
+			foreach (BsonValue item in testData.GetValue("items").AsBsonArray.ToList())
+			{
+				BsonDocument value = item.AsBsonDocument.GetValue("data").AsBsonDocument;
+				data.Add((item.AsBsonDocument.GetValue("key").AsString, value));
+			}
+
+			BsonDocument testData2 = BsonDocument.Parse(String.Join('\n', _sessionTestDataV2Lines2));
+			BsonDocument testData3 = BsonDocument.Parse(String.Join('\n', _sessionTestDataV2Lines3));
+			List<(string key, BsonDocument)> data2 = new List<(string key, BsonDocument)>();
+			foreach (BsonValue item in testData2.GetValue("items").AsBsonArray.ToList())
+			{
+				BsonDocument value = item.AsBsonDocument.GetValue("data").AsBsonDocument;
+				data2.Add((item.AsBsonDocument.GetValue("key").AsString, value));
+			}
+			foreach (BsonValue item in testData3.GetValue("items").AsBsonArray.ToList())
+			{
+				BsonDocument value = item.AsBsonDocument.GetValue("data").AsBsonDocument;
+				data2.Add((item.AsBsonDocument.GetValue("key").AsString, value));
+			}
+
+			await TestDataCollectionV2.AddAsync(job, step, data.ToArray());
+			await TestDataCollectionV2.AddAsync(job2, step2, data.ToArray());
+			await TestDataCollectionV2.AddAsync(job3, step3, data.ToArray());
+			await TestDataCollectionV2.AddAsync(job3, step3, data2.ToArray());
+
+			await TestDataService.TickForTestingAsync();
+
+			ActionResult<List<GetTestSessionStreamResponse>> streamResult = await TestDataController.GetTestStreamsV2Async(streamIds);
+			Assert.IsNotNull(streamResult);
+			Assert.IsNotNull(streamResult.Value);
+			List<GetTestSessionStreamResponse> streams = streamResult.Value;
+			Assert.AreEqual(2, streams.Count);
+			Assert.AreEqual(2, streams[0].Tests.Count);
+			Assert.AreEqual(3, streams[0].TestMetadata.Count);
+			Assert.AreEqual(3, streams[0].TestTags.Count);
+			Assert.AreEqual(1, streams[1].Tests.Count);
+			Assert.AreEqual(1, streams[1].TestMetadata.Count);
+			Assert.AreEqual(2, streams[1].TestTags.Count);
+
+			GetTestsRequest testRequest = new GetTestsRequest() { TestIds = streams[0].Tests.Select(t => t.Id).ToList() };
+			ActionResult<List<GetTestNameResponse>> testnameResult = await TestDataController.GetTestNameRefV2Async(testRequest);
+			Assert.IsNotNull(testnameResult);
+			Assert.IsNotNull(testnameResult.Value);
+			List<GetTestNameResponse> testnames = testnameResult.Value;
+			Assert.AreEqual(2, testnames.Count);
+			Assert.IsTrue(testnames.Any(t => t.Key == "UE.Automation(Group:Editor) EngineTest"));
+			Assert.IsTrue(testnames.Any(t => t.Key == "UE.Automation(Group:Editor) Lyra"));
+			Assert.AreEqual("Editor", testnames[0].Name);
+
+			GetMetadataRequest metaRequest = new GetMetadataRequest() { KeyValuePairs = new() { { "Platform", "Win64" } } };
+			ActionResult<List<GetTestMetadataResponse>> metaResult = await TestDataController.GetTestMetadataV2Async(metaRequest);
+			Assert.IsNotNull(metaResult);
+			Assert.IsNotNull(metaResult.Value);
+			List<GetTestMetadataResponse> meta = metaResult.Value;
+			Assert.AreEqual(2, meta.Count);
+			Assert.IsTrue(meta.All(t => t.KeyValuePairs.Any(i => i.Key == "BuildTarget" && i.Value == "Editor")));
+			Assert.IsTrue(meta.Any(t => t.KeyValuePairs.Any(i => i.Key == "Project" && i.Value == "EngineTest")));
+			Assert.IsTrue(meta.Any(t => t.KeyValuePairs.Any(i => i.Key == "Project" && i.Value == "Lyra")));
+
+			GetTestTagRequest tagRequest = new GetTestTagRequest() { TagIds = streams[0].TestTags.Select(t => t.Id).ToList() };
+			ActionResult<List<GetTestTagResponse>> tagResult = await TestDataController.GetTestTagRefV2Async(tagRequest);
+			Assert.IsNotNull(tagResult);
+			Assert.IsNotNull(tagResult.Value);
+			List<GetTestTagResponse> tags = tagResult.Value;
+			Assert.AreEqual(3, tags.Count);
+			Assert.IsTrue(tags.Any(t => t.Name == "Engine"));
+			Assert.IsTrue(tags.Any(t => t.Name == "Group"));
+			Assert.IsTrue(tags.Any(t => t.Name == "Parser"));
+
+			GetTestPhasesRequest phaseRequest = new GetTestPhasesRequest() { TestKeys = ["UE.Automation(Group:Editor) EngineTest"] };
+			ActionResult<List<GetTestPhaseResponse>> phaseResult = await TestDataController.GetTestPhaseRefFromTestV2Async(phaseRequest);
+			Assert.IsNotNull(phaseResult);
+			Assert.IsNotNull(phaseResult.Value);
+			List<GetTestPhaseResponse> tests = phaseResult.Value;
+			Assert.AreEqual(1, tests.Count);
+			Assert.AreEqual("Editor", tests[0].TestName);
+			Assert.AreEqual(3, tests[0].Phases.Count);
+			Assert.IsTrue(tests[0].Phases.Any(i => i.Name == "EditConditionParser.EvaluateBool"));
+			Assert.IsTrue(tests[0].Phases.Any(i => i.Name == "Editor.Failing.Test"));
+			string phaseId = tests[0].Phases.First(i => i.Name == "Editor.Failing.Test").Id.ToString();
+
+			ActionResult<List<GetTestSessionResponse>> testSessionResult = await TestDataController.GetTestSessionsV2Async(streamIds, testnames.Select(t => t.Id.ToString()).ToArray());
+			Assert.IsNotNull(testSessionResult);
+			Assert.IsNotNull(testSessionResult.Value);
+			List<GetTestSessionResponse> sessions = testSessionResult.Value;
+			Assert.AreEqual(5, sessions.Count);
+			Assert.AreEqual(TestOutcome.Failure, sessions[0].Outcome);
+
+			ActionResult<List<GetTestPhaseSessionResponse>> phaseSessionResult = await TestDataController.GetTestPhaseSessionsV2Async(streamIds, tests.SelectMany(t => t.Phases.Select(p => p.Id)).ToArray());
+			Assert.IsNotNull(phaseSessionResult);
+			Assert.IsNotNull(phaseSessionResult.Value);
+			List<GetTestPhaseSessionResponse> phaseSessions = phaseSessionResult.Value;
+			Assert.AreEqual(12, phaseSessions.Count);
+			Assert.AreEqual(TestPhaseOutcome.Failed, phaseSessions.First(p => p.PhaseRef == phaseId).Outcome);
+
+			// Get the sessions for one phase in one stream
+			phaseSessionResult = await TestDataController.GetTestPhaseSessionsV2Async([streamIds[0]], [tests[0].Phases[0].Id]);
+			Assert.IsNotNull(phaseSessionResult);
+			Assert.IsNotNull(phaseSessionResult.Value);
+			phaseSessions = phaseSessionResult.Value;
+			Assert.AreEqual(3, phaseSessions.Count);
+
+			ActionResult<object> testdata =  await TestDataController.GetTestDataV2Async(sessions[0].TestDataId.ToString());
+			Assert.IsNotNull(testdata);
+			Assert.IsNotNull(testdata.Value);
+		}
+
 		private readonly string[] _simpleTestDataLines =
 		{
 			@"{",
@@ -676,5 +799,247 @@ namespace HordeServer.Tests.Jobs.TestData
 			@"    ]",
 			@"}"
 		};
+
+		private readonly string[] _sessionTestDataV2Lines =
+		{
+			@"{",
+			@"  ""items"": [",
+			@"        {",
+			@"            ""key"": ""UE.Automation(Group:Editor) EngineTest"",",
+			@"            ""data"": {",
+			@"   		      ""version"" : 2,",
+			@"                ""summary"": {",
+			@"                   ""testName"": ""Editor"",",
+			@"        		     ""dateTime"": ""2025-03-03T18:14:04"",",
+			@"		             ""timeElapseSec"": 207.02324,",
+			@"        		     ""phasesTotalCount"": 1075,",
+			@"		             ""phasesSucceededCount"": 1074,",
+			@"        		     ""phasesUndefinedCount"": 0,",
+			@"		             ""phasesFailedCount"": 1",
+			@"        		  },",
+			@"                ""phases"": [",
+			@"                    {",
+			@"                        ""key"": ""874a6711"",",
+			@"                        ""name"": ""EditConditionParser.EvaluateBitFlags"",",
+			@"                        ""dateTime"": ""2025-03-03T18:07:55"",",
+			@"                        ""timeElapseSec"": 0.0879574,",
+			@"                        ""outcome"": ""Success"",",
+			@"                        ""deviceKeys"": [",
+			@"                            ""9AF6B8CC484096D351EC56ABDFBB919F""",
+			@"                        ]",
+			@"                    },",
+			@"                    {",
+			@"                        ""key"": ""03f772e3"",",
+			@"                        ""name"": ""EditConditionParser.EvaluateBool"",",
+			@"                        ""dateTime"": ""2025-03-03T18:08:08"",",
+			@"                        ""timeElapseSec"": 0.008163903,",
+			@"                        ""outcome"": ""Success"",",
+			@"                        ""deviceKeys"": [",
+			@"                            ""9AF6B8CC484096D351EC56ABDFBB919F""",
+			@"                        ]",
+			@"                   },",
+			@"                   {",
+			@"                        ""key"": ""62d99b42"",",
+			@"                        ""name"": ""Editor.Failing.Test"",",
+			@"                        ""dateTime"": ""2025-03-03T18:08:14"",",
+			@"                        ""timeElapseSec"": 0.1961886,",
+			@"                        ""outcome"": ""Failed"",",
+			@"                        ""errorFingerprint"": ""a6efc189"",",
+			@"                        ""deviceKeys"": [",
+			@"                            ""9AF6B8CC484096D351EC56ABDFBB919F""",
+			@"                        ],",
+			@"                        ""eventStreamPath"": ""Engine/Programs/AutomationTool/Saved/Logs/UE.EditorAutomation(RunTest=Group_Editor)_(Win64_Development_Editor)/EventStreams/62d99b42.json""",
+			@"                   }",
+			@"                ]",
+			@"                ""tags"": [""Group"", ""Engine""]",
+			@"                ""metadata"": {",
+			@"                    ""Platform"": ""Win64"",",
+ 			@"                    ""BuildTarget"": ""Editor"",",
+			@"                    ""Configuration"": ""Development"",",
+ 			@"                    ""Project"": ""EngineTest"",",
+			@"                    ""RHI"": ""default""",
+			@"                }",
+			@"                ""devices"": {",
+			@"                  ""9AF6B8CC484096D351EC56ABDFBB919F"": {",
+			@"                    ""name"": ""M-SIHLOYQHDOUKO"",",
+			@"                    ""appInstanceLogPath"": ""Engine/Programs/AutomationTool/Saved/Logs/UE.EditorAutomation(RunTest=Group_Editor)_(Win64_Development_Editor)/Editor/EditorOutput.log"",",
+			@"                    ""metadata"": {",
+			@"                      ""platform"": ""WindowsEditor"",",
+			@"                      ""os_version"": ""Windows 10 (22H2) [10.0.19045.5487] "",",
+			@"                      ""model"": ""Default"",",
+			@"                      ""gpu"": ""Parsec Virtual Display Adapter"",",
+			@"                      ""cpumodel"": ""AMD EPYC 7R32"",",
+			@"                      ""ram_in_gb"": ""127"",",
+			@"                      ""render_mode"": ""D3D12_SM6"",",
+			@"                      ""rhi"": ""D3D12 (SM6)""",
+			@"                    }",
+			@"                  }",
+			@"                },",
+			@"            }",
+			@"        }",
+			@"    ]",
+			@"}"
+		};
+
+		private readonly string[] _sessionTestDataV2Lines2 =
+		{
+			@"{",
+			@"  ""items"": [",
+			@"        {",
+			@"            ""key"": ""UE.Automation(Group:Editor) EngineTest"",",
+			@"            ""data"": {",
+			@"   		      ""version"" : 2,",
+			@"                ""summary"": {",
+			@"                   ""testName"": ""Editor"",",
+			@"        		     ""dateTime"": ""2025-03-03T18:14:04"",",
+			@"		             ""timeElapseSec"": 207.02324,",
+			@"        		     ""phasesTotalCount"": 1075,",
+			@"		             ""phasesSucceededCount"": 1074,",
+			@"        		     ""phasesUndefinedCount"": 0,",
+			@"		             ""phasesFailedCount"": 1",
+			@"        		  },",
+			@"                ""phases"": [",
+			@"                    {",
+			@"                        ""key"": ""874a6711"",",
+			@"                        ""name"": ""EditConditionParser.EvaluateBitFlags"",",
+			@"                        ""dateTime"": ""2025-03-03T18:07:55"",",
+			@"                        ""timeElapseSec"": 0.0879574,",
+			@"                        ""outcome"": ""Success"",",
+			@"                        ""deviceKeys"": [",
+			@"                            ""9AF6B8CC484096D351EC56ABDFBB919F""",
+			@"                        ]",
+			@"                    },",
+			@"                    {",
+			@"                        ""key"": ""03f772e3"",",
+			@"                        ""name"": ""EditConditionParser.EvaluateBool"",",
+			@"                        ""dateTime"": ""2025-03-03T18:08:08"",",
+			@"                        ""timeElapseSec"": 0.008163903,",
+			@"                        ""outcome"": ""Success"",",
+			@"                        ""deviceKeys"": [",
+			@"                            ""9AF6B8CC484096D351EC56ABDFBB919F""",
+			@"                        ]",
+			@"                   },",
+			@"                   {",
+			@"                        ""key"": ""62d99b42"",",
+			@"                        ""name"": ""Editor.Failing.Test"",",
+			@"                        ""dateTime"": ""2025-03-03T18:08:14"",",
+			@"                        ""timeElapseSec"": 0.1961886,",
+			@"                        ""outcome"": ""Failed"",",
+			@"                        ""errorFingerprint"": ""a6efc189"",",
+			@"                        ""deviceKeys"": [",
+			@"                            ""9AF6B8CC484096D351EC56ABDFBB919F""",
+			@"                        ],",
+			@"                        ""eventStreamPath"": ""Engine/Programs/AutomationTool/Saved/Logs/UE.EditorAutomation(RunTest=Group_Editor)_(Win64_Development_Editor)/EventStreams/62d99b42.json""",
+			@"                   }",
+			@"                ]",
+			@"                ""tags"": [""Parser"", ""Engine""]",
+			@"                ""metadata"": {",
+			@"                    ""Platform"": ""Mac"",",
+ 			@"                    ""BuildTarget"": ""Editor"",",
+			@"                    ""Configuration"": ""Development"",",
+ 			@"                    ""Project"": ""EngineTest"",",
+			@"                    ""RHI"": ""default""",
+			@"                }",
+			@"                ""devices"": {",
+			@"                  ""9AF6B8CC484096D351EC56ABDFBB919F"": {",
+			@"                    ""name"": ""M-SIHLOYQHDOUKO"",",
+			@"                    ""appInstanceLogPath"": ""Engine/Programs/AutomationTool/Saved/Logs/UE.EditorAutomation(RunTest=Group_Editor)_(Win64_Development_Editor)/Editor/EditorOutput.log"",",
+			@"                    ""metadata"": {",
+			@"                      ""platform"": ""WindowsEditor"",",
+			@"                      ""os_version"": ""Windows 10 (22H2) [10.0.19045.5487] "",",
+			@"                      ""model"": ""Default"",",
+			@"                      ""gpu"": ""Parsec Virtual Display Adapter"",",
+			@"                      ""cpumodel"": ""AMD EPYC 7R32"",",
+			@"                      ""ram_in_gb"": ""127"",",
+			@"                      ""render_mode"": ""D3D12_SM6"",",
+			@"                      ""rhi"": ""D3D12 (SM6)""",
+			@"                    }",
+			@"                  }",
+			@"                },",
+			@"            }",
+			@"        }",
+			@"    ]",
+			@"}"
+		};
+
+		private readonly string[] _sessionTestDataV2Lines3 =
+		{
+			@"{",
+			@"  ""items"": [",
+			@"        {",
+			@"            ""key"": ""UE.Automation(Group:Editor) Lyra"",",
+			@"            ""data"": {",
+			@"   		      ""version"" : 2,",
+			@"                ""summary"": {",
+			@"                   ""testName"": ""Editor"",",
+			@"        		     ""dateTime"": ""2025-03-03T18:14:04"",",
+			@"		             ""timeElapseSec"": 207.02324,",
+			@"        		     ""phasesTotalCount"": 1075,",
+			@"		             ""phasesSucceededCount"": 1074,",
+			@"        		     ""phasesUndefinedCount"": 0,",
+			@"		             ""phasesFailedCount"": 1",
+			@"        		  },",
+			@"                ""phases"": [",
+			@"                    {",
+			@"                        ""key"": ""874a6711"",",
+			@"                        ""name"": ""EditConditionParser.EvaluateBitFlags"",",
+			@"                        ""dateTime"": ""2025-03-03T18:07:55"",",
+			@"                        ""timeElapseSec"": 0.0879574,",
+			@"                        ""outcome"": ""Success"",",
+			@"                        ""deviceKeys"": [",
+			@"                            ""9AF6B8CC484096D351EC56ABDFBB919F""",
+			@"                        ]",
+			@"                    },",
+			@"                    {",
+			@"                        ""key"": ""03f772e3"",",
+			@"                        ""name"": ""EditConditionParser.EvaluateBool"",",
+			@"                        ""dateTime"": ""2025-03-03T18:08:08"",",
+			@"                        ""timeElapseSec"": 0.008163903,",
+			@"                        ""outcome"": ""Success"",",
+			@"                        ""deviceKeys"": [",
+			@"                            ""9AF6B8CC484096D351EC56ABDFBB919F""",
+			@"                        ]",
+			@"                   },",
+			@"                   {",
+			@"                        ""key"": ""62d99b42"",",
+			@"                        ""name"": ""Editor.Failing.Test"",",
+			@"                        ""dateTime"": ""2025-03-03T18:08:14"",",
+			@"                        ""timeElapseSec"": 0.1961886,",
+			@"                        ""outcome"": ""Failed"",",
+			@"                        ""errorFingerprint"": ""a6efc189"",",
+			@"                        ""deviceKeys"": [",
+			@"                            ""9AF6B8CC484096D351EC56ABDFBB919F""",
+			@"                        ],",
+			@"                        ""eventStreamPath"": ""Engine/Programs/AutomationTool/Saved/Logs/UE.EditorAutomation(RunTest=Group_Editor)_(Win64_Development_Editor)/EventStreams/62d99b42.json""",
+			@"                   }",
+			@"                ]",
+			@"                ""metadata"": {",
+			@"                    ""Platform"": ""Win64"",",
+ 			@"                    ""BuildTarget"": ""Editor"",",
+			@"                    ""Configuration"": ""Development"",",
+ 			@"                    ""Project"": ""Lyra"",",
+			@"                    ""RHI"": ""default""",
+			@"                }",
+			@"                ""devices"": {",
+			@"                  ""9AF6B8CC484096D351EC56ABDFBB919F"": {",
+			@"                    ""name"": ""M-SIHLOYQHDOUKO"",",
+			@"                    ""appInstanceLogPath"": ""Engine/Programs/AutomationTool/Saved/Logs/UE.EditorAutomation(RunTest=Group_Editor)_(Win64_Development_Editor)/Editor/EditorOutput.log"",",
+			@"                    ""metadata"": {",
+			@"                      ""platform"": ""WindowsEditor"",",
+			@"                      ""os_version"": ""Windows 10 (22H2) [10.0.19045.5487] "",",
+			@"                      ""model"": ""Default"",",
+			@"                      ""gpu"": ""Parsec Virtual Display Adapter"",",
+			@"                      ""cpumodel"": ""AMD EPYC 7R32"",",
+			@"                      ""ram_in_gb"": ""127"",",
+			@"                      ""render_mode"": ""D3D12_SM6"",",
+			@"                      ""rhi"": ""D3D12 (SM6)""",
+			@"                    }",
+			@"                  }",
+			@"                },",
+			@"            }",
+			@"        }",
+			@"    ]",
+			@"}"
+		};
 	}
 }
diff --git a/Engine/Source/Programs/Horde/Plugins/Build/HordeServer.Build/BuildPlugin.cs b/Engine/Source/Programs/Horde/Plugins/Build/HordeServer.Build/BuildPlugin.cs
--- a/Engine/Source/Programs/Horde/Plugins/Build/HordeServer.Build/BuildPlugin.cs
+++ b/Engine/Source/Programs/Horde/Plugins/Build/HordeServer.Build/BuildPlugin.cs
@@ -101,6 +101,7 @@ namespace HordeServer
 			services.AddSingleton<ITemplateCollection>(sp => sp.GetRequiredService<TemplateCollection>());
 			services.AddSingleton<ITemplateCollectionInternal>(sp => sp.GetRequiredService<TemplateCollection>());
 			services.AddSingleton<ITestDataCollection, TestDataCollection>();
+			services.AddSingleton<ITestDataCollectionV2, TestDataCollectionV2>();
 			services.AddSingleton<IUtilizationDataCollection, UtilizationDataCollection>();
 			services.AddSingleton<ITemplateCollection, TemplateCollection>();
 
diff --git a/Engine/Source/Programs/Horde/Plugins/Build/HordeServer.Build/Jobs/JobRpcService.cs b/Engine/Source/Programs/Horde/Plugins/Build/HordeServer.Build/Jobs/JobRpcService.cs
--- a/Engine/Source/Programs/Horde/Plugins/Build/HordeServer.Build/Jobs/JobRpcService.cs
+++ b/Engine/Source/Programs/Horde/Plugins/Build/HordeServer.Build/Jobs/JobRpcService.cs
@@ -53,6 +53,7 @@ namespace HordeServer.Jobs
 		readonly ILogCollection _logCollection;
 		readonly IGraphCollection _graphs;
 		readonly ITestDataCollection _testData;
+		readonly ITestDataCollectionV2 _testDataV2;
 		readonly IJobStepRefCollection _jobStepRefCollection;
 		readonly ITemplateCollection _templateCollection;
 		readonly HttpClient _httpClient;
@@ -63,7 +64,7 @@ namespace HordeServer.Jobs
 		/// <summary>
 		/// Constructor
 		/// </summary>
-		public JobRpcService(IAclService aclService, JobService jobService, AgentService agentService, PoolService poolService, ConformTaskSource conformTaskSource, IArtifactCollection artifactCollection, IJobCollection jobCollection, ILogCollection logCollection, IGraphCollection graphs, ITestDataCollection testData, IJobStepRefCollection jobStepRefCollection, ITemplateCollection templateCollection, HttpClient httpClient, IClock clock, IOptionsSnapshot<BuildConfig> buildConfig, ILogger<JobRpcService> logger)
+		public JobRpcService(IAclService aclService, JobService jobService, AgentService agentService, PoolService poolService, ConformTaskSource conformTaskSource, IArtifactCollection artifactCollection, IJobCollection jobCollection, ILogCollection logCollection, IGraphCollection graphs, ITestDataCollection testData, ITestDataCollectionV2 testDataV2, IJobStepRefCollection jobStepRefCollection, ITemplateCollection templateCollection, HttpClient httpClient, IClock clock, IOptionsSnapshot<BuildConfig> buildConfig, ILogger<JobRpcService> logger)
 		{
 			_aclService = aclService;
 			_jobService = jobService;
@@ -75,6 +76,7 @@ namespace HordeServer.Jobs
 			_logCollection = logCollection;
 			_graphs = graphs;
 			_testData = testData;
+			_testDataV2 = testDataV2;
 			_jobStepRefCollection = jobStepRefCollection;
 			_templateCollection = templateCollection;
 			_httpClient = httpClient;
@@ -832,6 +834,7 @@ namespace HordeServer.Jobs
 			if (job != null && jobStep != null)
 			{
 				await _testData.AddAsync(job, jobStep, data.ToArray(), context.CancellationToken);
+				await _testDataV2.AddAsync(job, jobStep, data.ToArray(), context.CancellationToken);
 			}
 			else
 			{
diff --git a/Engine/Source/Programs/Horde/Plugins/Build/HordeServer.Build/Jobs/TestData/ITestData.cs b/Engine/Source/Programs/Horde/Plugins/Build/HordeServer.Build/Jobs/TestData/ITestData.cs
--- a/Engine/Source/Programs/Horde/Plugins/Build/HordeServer.Build/Jobs/TestData/ITestData.cs
+++ b/Engine/Source/Programs/Horde/Plugins/Build/HordeServer.Build/Jobs/TestData/ITestData.cs
@@ -118,7 +118,7 @@ namespace HordeServer.Jobs.TestData
 		TestId TestId { get; }
 
 		/// <summary>
-		/// The ourcome of the suite test
+		/// The outcome of the suite test
 		/// </summary>
 		TestOutcome Outcome { get; }
 
@@ -261,6 +261,279 @@ namespace HordeServer.Jobs.TestData
 		IReadOnlyList<TestSuiteId> TestSuites { get; }
 	}
 
+	/// v2
+
+	/// <summary>
+	/// A test name reference
+	/// </summary>
+	public interface ITestNameRef
+	{
+		/// <summary>
+		/// The test id
+		/// </summary>
+		TestId Id { get; }
+
+		/// <summary>
+		/// The test data key that identify the test uniquely
+		/// </summary>
+		string Key { get; }
+
+		/// <summary>
+		/// The display name of the test
+		/// </summary>
+		string Name { get; }
+	}
+
+	/// <summary>
+	/// A test phase reference
+	/// </summary>
+	public interface ITestPhaseRef
+	{
+		/// <summary>
+		/// The test phase id
+		/// </summary>
+		TestPhaseId Id { get; }
+
+		/// <summary>
+		/// The test reference associated with the phase
+		/// </summary>
+		TestId TestNameRef { get; }
+
+		/// <summary>
+		/// The test phase key that identify the phase uniquely
+		/// </summary>
+		string Key { get; }
+
+		/// <summary>
+		/// The test phase (display) name
+		/// </summary>
+		string Name { get; }
+	}
+
+	/// <summary>
+	/// A metadata key/value pairs reference
+	/// </summary>
+	public interface ITestMetaRef
+	{
+		/// <summary>
+		/// The metadata id
+		/// </summary>
+		TestMetaId Id { get; }
+
+		/// <summary>
+		/// The key/value pairs that this metadata references
+		/// </summary>
+		IReadOnlyDictionary<string, string> KeyValuePairs { get; }
+	}
+
+	/// <summary>
+	/// A test tag reference
+	/// </summary>
+	public interface ITestTagRef
+	{
+		/// <summary>
+		/// The tag id
+		/// </summary>
+		TestTagId Id { get; }
+
+		/// <summary>
+		/// The tag name
+		/// </summary>
+		string Name { get; }
+	}
+
+	/// <summary>
+	/// Test session
+	/// </summary>
+	public interface ITestSession
+	{
+		/// <summary>
+		/// The test session id
+		/// </summary>
+		TestSessionId Id { get; }
+
+		/// <summary>
+		/// The metadata associated with the test session
+		/// </summary>
+		TestMetaId Metadata { get; }
+
+		/// <summary>
+		/// The tags associated with the test session
+		/// </summary>
+		IReadOnlyList<TestTagId>? Tags { get; }
+
+		/// <summary>
+		/// The test data id for this test session
+		/// </summary>
+		ObjectId TestDataId { get; }
+
+		/// <summary>
+		/// The build changelist upon which the test ran, may not correspond to the job changelist
+		/// </summary>
+		CommitIdWithOrder BuildCommitId { get; }
+
+		/// <summary>
+		/// Reference Id to test name
+		/// </summary>
+		TestId NameRef { get; }
+
+		/// <summary>
+		/// The associated stream
+		/// </summary>
+		StreamId StreamId { get; }
+
+		/// <summary>
+		/// The associated job
+		/// </summary>
+		JobId JobId { get; }
+
+		/// <summary>
+		/// The associated job step
+		/// </summary>
+		JobStepId StepId { get; }
+
+		/// <summary>
+		/// How long the test session ran
+		/// </summary>
+		TimeSpan Duration { get; }
+
+		/// <summary>
+		/// Date time at which the test session was started
+		/// </summary>
+		DateTime StartDateTime { get; }
+
+		/// <summary>
+		/// The outcome of the test session
+		/// if PhasesUndefinedCount > 0 => Undefiened
+		/// if PhasesFailedCount > 0 => Failure
+		/// if PhasesTotalCount > 0 and PhasesSucceededCount > 0 => Success
+		/// else => Skipped
+		/// </summary>
+		TestOutcome Outcome { get; }
+
+		/// <summary>
+		/// The total number the test session contains
+		/// </summary>
+		int PhasesTotalCount { get; }
+
+		/// <summary>
+		/// The number of phases that succeeded during the test session
+		/// </summary>
+		int PhasesSucceededCount { get; }
+
+		/// <summary>
+		/// The number of phases that were left undefined during the test session
+		/// </summary>
+		int PhasesUndefinedCount { get; }
+
+		/// <summary>
+		/// The number of phases that failed during the test session
+		/// </summary>
+		int PhasesFailedCount { get; }
+	}
+
+	/// <summary>
+	/// The test sessions running in a given stream
+	/// </summary>
+	public interface ITestSessionStream
+	{
+		/// <summary>
+		/// The stream id
+		/// </summary>
+		StreamId StreamId { get; }
+
+		/// <summary>
+		/// Test ids of tests running in the stream
+		/// </summary>
+		IReadOnlyList<TestId> Tests { get; }
+
+		/// <summary>
+		/// Test metadata ids from running tests in the stream 
+		/// </summary>
+		IReadOnlyList<TestMetaId> Metadata { get; }
+
+		/// <summary>
+		/// Test tag ids from running tests in the stream
+		/// </summary>
+		IReadOnlyList<TestTagId> Tags { get; }
+	}
+
+	/// <summary>
+	/// A Test phase outcome from a test session
+	/// </summary>
+	public interface ITestPhaseSession
+	{
+		/// <summary>
+		/// The test phase session id
+		/// </summary>
+		TestPhaseSessionId Id { get; }
+
+		/// <summary>
+		/// The test phase reference id
+		/// </summary>
+		TestPhaseId PhaseRef { get; }
+
+		/// <summary>
+		/// The phase outcome
+		/// </summary>
+		TestPhaseOutcome Outcome { get; }
+
+		/// <summary>
+		/// The test session that run that phase
+		/// </summary>
+		TestSessionId SessionId { get; }
+
+		/// <summary>
+		/// The metadata associated with this test session
+		/// </summary>
+		TestMetaId Metadata { get; }
+
+		/// <summary>
+		/// The path to the event stream
+		/// </summary>
+		string? EventStreamPath { get; }
+
+		/// <summary>
+		/// Whether the test phase encountered at least one warning
+		/// </summary>
+		bool? HasWarning { get;  }
+
+		/// <summary>
+		/// The error fingerprint from the event stream
+		/// </summary>
+		string? ErrorFingerprint { get; }
+
+		/// <summary>
+		/// The associated stream
+		/// </summary>
+		StreamId StreamId { get; }
+
+		/// <summary>
+		/// The associated job
+		/// </summary>
+		JobId JobId { get; }
+
+		/// <summary>
+		/// The associated job step
+		/// </summary>
+		JobStepId StepId { get; }
+
+		/// <summary>
+		/// The build changelist upon which the test ran, may not correspond to the job changelist
+		/// </summary>
+		CommitIdWithOrder BuildCommitId { get; }
+
+		/// <summary>
+		/// How long the test phase session ran
+		/// </summary>
+		TimeSpan Duration { get; }
+
+		/// <summary>
+		/// Date time at which the test phase session was started
+		/// </summary>
+		DateTime StartDateTime { get; }
+	}
+
 	/// <summary>
 	/// Stores information about the results of a test
 	/// </summary>
diff --git a/Engine/Source/Programs/Horde/Plugins/Build/HordeServer.Build/Jobs/TestData/ITestDataCollection.cs b/Engine/Source/Programs/Horde/Plugins/Build/HordeServer.Build/Jobs/TestData/ITestDataCollection.cs
--- a/Engine/Source/Programs/Horde/Plugins/Build/HordeServer.Build/Jobs/TestData/ITestDataCollection.cs
+++ b/Engine/Source/Programs/Horde/Plugins/Build/HordeServer.Build/Jobs/TestData/ITestDataCollection.cs
@@ -135,4 +135,114 @@ namespace HordeServer.Jobs.TestData
 		/// <param name="cancellationToken">Cancellation token for the operation</param>
 		public Task UpgradeAsync(CancellationToken cancellationToken = default);
 	}
+
+	/// <summary>
+	/// Collection of test data documents v2
+	/// </summary>
+	public interface ITestDataCollectionV2
+	{
+		/// <summary>
+		/// Gets the test sessions running in provided streams
+		/// </summary>
+		/// <param name="streamIds"></param>
+		/// <param name="cancellationToken">Cancellation token for the operation</param>
+		/// <returns></returns>
+		Task<IReadOnlyList<ITestSessionStream>> FindTestSessionStreamsAsync(StreamId[] streamIds, CancellationToken cancellationToken = default);
+
+		/// <summary>
+		/// Find test sessions
+		/// </summary>
+		/// <param name="streamIds"></param>
+		/// <param name="testIds"></param>
+		/// <param name="metaIds"></param>
+		/// <param name="minCreateTime"></param>
+		/// <param name="maxCreateTime"></param>
+		/// <param name="minCommitId"></param>
+		/// <param name="maxCommitId"></param>
+		/// <param name="cancellationToken">Cancellation token for the operation</param>
+		/// <returns></returns>
+		Task<IReadOnlyList<ITestSession>> FindTestSessionsAsync(StreamId[] streamIds, TestId[]? testIds = null, TestMetaId[]? metaIds = null, DateTime? minCreateTime = null, DateTime? maxCreateTime = null, CommitId? minCommitId = null, CommitId? maxCommitId = null, CancellationToken cancellationToken = default);
+
+		/// <summary>
+		/// Find test name references
+		/// </summary>
+		/// <param name="testIds"></param>
+		/// <param name="keys"></param>
+		/// <param name="cancellationToken">Cancellation token for the operation</param>
+		/// <returns></returns>
+		Task<IReadOnlyList<ITestNameRef>> FindTestNameRefsAsync(TestId[]? testIds = null, string[]? keys = null, CancellationToken cancellationToken = default);
+
+		/// <summary>
+		/// Creates a new test data document
+		/// </summary>
+		/// <param name="job">The job containing the step</param>
+		/// <param name="step">The step producing the data</param>
+		/// <param name="data">The test data generated by the job step</param>
+		/// <param name="cancellationToken">Cancellation token for the operation</param>
+		/// <returns>The new testdata document</returns>
+		Task<IReadOnlyList<ITestData>> AddAsync(IJob job, IJobStep step, (string key, BsonDocument value)[] data, CancellationToken cancellationToken = default);
+
+		/// <summary>
+		/// Gets a testdata by ID
+		/// </summary>
+		/// <param name="id">Unique id of the testdata</param>
+		/// <param name="cancellationToken">Cancellation token for the operation</param>
+		/// <returns>The testdata document</returns>
+		Task<ITestData?> GetAsync(ObjectId id, CancellationToken cancellationToken = default);
+
+		/// <summary>
+		/// Find test metadata set
+		/// </summary>
+		/// <param name="metaIds"></param>
+		/// <param name="keyValues"></param>
+		/// <param name="cancellationToken">Cancellation token for the operation</param>
+		/// <returns></returns>
+		Task<IReadOnlyList<ITestMetaRef>> FindTestMetaAsync(TestMetaId[]? metaIds = null, IReadOnlyDictionary<string, string>? keyValues = null, CancellationToken cancellationToken = default);
+
+		/// <summary>
+		/// Find test phase references associated with a set of tests
+		/// </summary>
+		/// <param name="testIds"></param>
+		/// <param name="keys"></param>
+		/// <param name="cancellationToken">Cancellation token for the operation</param>
+		/// <returns></returns>
+		Task<IReadOnlyList<ITestPhaseRef>> FindTestPhasesAsync(TestId[] testIds, string[]? keys = null, CancellationToken cancellationToken = default);
+
+		/// <summary>
+		/// Find test tag references
+		/// </summary>
+		/// <param name="tagIds"></param>
+		/// <param name="names"></param>
+		/// <param name="cancellationToken">Cancellation token for the operation</param>
+		/// <returns></returns>
+		Task<IReadOnlyList<ITestTagRef>> FindTestTagsAsync(TestTagId[]? tagIds = null, string[]? names = null, CancellationToken cancellationToken = default);
+
+		/// <summary>
+		/// Find test phase sessions from a set of test phases
+		/// </summary>
+		/// <param name="streamIds"></param>
+		/// <param name="phaseIds"></param>
+		/// <param name="minCreateTime"></param>
+		/// <param name="maxCreateTime"></param>
+		/// <param name="minCommitId"></param>
+		/// <param name="maxCommitId"></param>
+		/// <param name="cancellationToken">Cancellation token for the operation</param>
+		/// <returns></returns>
+		Task<IReadOnlyList<ITestPhaseSession>> FindTestPhaseSessionsAsync(StreamId[] streamIds, TestPhaseId[] phaseIds, DateTime? minCreateTime = null, DateTime? maxCreateTime = null, CommitId? minCommitId = null, CommitId? maxCommitId = null, CancellationToken cancellationToken = default);
+
+		/// <summary>
+		/// Prune test data from collection
+		/// </summary>
+		/// <param name="cancellationToken">Cancellation token for the operation</param>
+		/// <returns></returns>
+		Task PruneDataAsync(CancellationToken cancellationToken = default);
+
+		/// <summary>
+		/// Run a tick on the collection
+		/// </summary>
+		/// <param name="retainMonths"></param>
+		/// <param name="cancellationToken">Cancellation token for the operation</param>
+		/// <returns></returns>
+		Task<bool> UpdateAsync(int retainMonths, CancellationToken cancellationToken = default);
+	}
 }
diff --git a/Engine/Source/Programs/Horde/Plugins/Build/HordeServer.Build/Jobs/TestData/TestDataCollection.cs b/Engine/Source/Programs/Horde/Plugins/Build/HordeServer.Build/Jobs/TestData/TestDataCollection.cs
--- a/Engine/Source/Programs/Horde/Plugins/Build/HordeServer.Build/Jobs/TestData/TestDataCollection.cs
+++ b/Engine/Source/Programs/Horde/Plugins/Build/HordeServer.Build/Jobs/TestData/TestDataCollection.cs
@@ -1,5 +1,6 @@
 // Copyright Epic Games, Inc. All Rights Reserved.
 
+using EpicGames.Core;
 using EpicGames.Horde.Commits;
 using EpicGames.Horde.Jobs;
 using EpicGames.Horde.Jobs.Templates;
@@ -500,7 +501,6 @@ namespace HordeServer.Jobs.TestData
 		public async Task<IReadOnlyList<ITestData>> AddAsync(IJob job, IJobStep step, (string key, BsonDocument value)[] data, CancellationToken cancellationToken = default)
 		{
 			// detailed test data
-			bool skip = false;
 			List<TestDataDocument> documents = new List<TestDataDocument>();
 			for (int i = 0; i < data.Length; i++)
 			{
@@ -511,18 +511,13 @@ namespace HordeServer.Jobs.TestData
 				{
 					if (version > 1)
 					{
-						skip = true;
+						continue;
 					}
 				}
 
 				documents.Add(new TestDataDocument(job, step, key, document));
 			}
 
-			if (skip)
-			{
-				documents.Clear();
-			}
-
 			if (documents.Count == 0)
 			{
 				return documents;
diff --git a/Engine/Source/Programs/Horde/Plugins/Build/HordeServer.Build/Jobs/TestData/TestDataCollectionV2.cs b/Engine/Source/Programs/Horde/Plugins/Build/HordeServer.Build/Jobs/TestData/TestDataCollectionV2.cs
new file mode 100644
--- /dev/null
+++ b/Engine/Source/Programs/Horde/Plugins/Build/HordeServer.Build/Jobs/TestData/TestDataCollectionV2.cs
@@ -0,0 +1,1114 @@
+// Copyright Epic Games, Inc. All Rights Reserved.
+
+using EpicGames.Horde.Commits;
+using EpicGames.Horde.Jobs;
+using EpicGames.Horde.Jobs.Templates;
+using EpicGames.Horde.Jobs.TestData;
+using EpicGames.Horde.Streams;
+using HordeServer.Commits;
+using HordeServer.Server;
+using HordeServer.Utilities;
+using Microsoft.Extensions.Logging;
+using MongoDB.Bson;
+using MongoDB.Bson.Serialization;
+using MongoDB.Bson.Serialization.Attributes;
+using MongoDB.Driver;
+using OpenTelemetry.Trace;
+
+namespace HordeServer.Jobs.TestData
+{
+	/// <summary>
+	/// Collection of test data documents v2
+	/// </summary>
+	public class TestDataCollectionV2 : ITestDataCollectionV2
+	{
+		/// <summary>
+		/// Interface to handle expiration
+		/// </summary>
+		internal interface ITestExpire
+		{
+			DateTime LastSeenUtc { get; }
+		}
+
+		/// <summary>
+		/// Test metadata document
+		/// </summary>
+		class TestMetadataDocument : ITestMetaRef, ITestExpire
+		{
+			[BsonRequired, BsonId]
+			public TestMetaId Id { get; set; }
+
+			[BsonRequired, BsonElement("m")]
+			public Dictionary<string, string> KeyValuePairs { get; set; }
+			IReadOnlyDictionary<string, string> ITestMetaRef.KeyValuePairs => KeyValuePairs;
+
+			[BsonRequired, BsonElement("s")]
+			public DateTime LastSeenUtc { get; set; }
+
+			private TestMetadataDocument()
+			{
+				KeyValuePairs = new();
+			}
+
+			public TestMetadataDocument(IReadOnlyDictionary<string, string> keyValuePairs)
+			{
+				Id = TestMetaId.GenerateNewId();
+				KeyValuePairs = keyValuePairs.ToDictionary();
+				LastSeenUtc = DateTime.UtcNow;
+			}
+		}
+
+		/// <summary>
+		/// Test tag reference document
+		/// </summary>
+		class TestTagDocument : ITestTagRef, ITestExpire
+		{
+			[BsonRequired, BsonId] 
+			public TestTagId Id { get; set; }
+
+			[BsonRequired, BsonElement("n")] 
+			public string Name { get; set; }
+
+			[BsonRequired, BsonElement("s")]
+			public DateTime LastSeenUtc { get; set; }
+
+			private TestTagDocument()
+			{
+				Name = String.Empty;
+			}
+
+			public TestTagDocument(string name)
+			{
+				Id = TestTagId.GenerateNewId();
+				Name = name;
+				LastSeenUtc = DateTime.UtcNow;
+			}
+		}
+
+		/// <summary>
+		/// Test name reference document
+		/// </summary>
+		class TestNameDocument : ITestNameRef, ITestExpire
+		{
+			[BsonRequired, BsonId]
+			public TestId Id { get; set; }
+
+			[BsonRequired, BsonElement("n")]
+			public string Name { get; set; }
+
+			[BsonRequired, BsonElement("k")]
+			public string Key { get; set; }
+
+			[BsonRequired, BsonElement("s")]
+			public DateTime LastSeenUtc { get; set; }
+
+			private TestNameDocument()
+			{
+				Key = String.Empty;
+				Name = String.Empty;
+			}
+
+			public TestNameDocument(string key, string name)
+			{
+				Id = TestId.GenerateNewId();
+				Key = key;
+				Name = name;
+				LastSeenUtc = DateTime.UtcNow;
+			}
+		}
+
+		/// <summary>
+		/// Test phase reference document
+		/// </summary>
+		class TestPhaseDocument : ITestPhaseRef, ITestExpire
+		{
+			[BsonRequired, BsonId]
+			public TestPhaseId Id { get; set; }
+
+			[BsonRequired, BsonElement("tid")]
+			public TestId TestNameRef { get; set; }
+
+			[BsonRequired, BsonElement("n")]
+			public string Name { get; set; }
+
+			[BsonRequired, BsonElement("k")]
+			public string Key { get; set; }
+
+			[BsonRequired, BsonElement("s")]
+			public DateTime LastSeenUtc { get; set; }
+
+			private TestPhaseDocument()
+			{
+				Name = String.Empty;
+				Key = String.Empty;
+			}
+
+			public TestPhaseDocument(TestId testId, string key, string name)
+			{
+				Id = TestPhaseId.GenerateNewId();
+				TestNameRef = testId;
+				Name = name;
+				Key = key;
+				LastSeenUtc = DateTime.UtcNow;
+			}
+		}
+
+		/// <summary>
+		/// Stores the tests running in a stream
+		/// </summary>
+		class TestSessionStreamDocument : ITestSessionStream
+		{
+			[BsonRequired, BsonId]
+			public ObjectId Id { get; set; }
+
+			[BsonRequired, BsonElement("sid")]
+			public StreamId StreamId { get; set; }
+
+			[BsonRequired, BsonElement("t")]
+			public List<TestId> Tests { get; set; } = new List<TestId>();
+			IReadOnlyList<TestId> ITestSessionStream.Tests => Tests;
+
+			[BsonRequired, BsonElement("m")]
+			public List<TestMetaId> Metadata { get; set; } = new List<TestMetaId>();
+			IReadOnlyList<TestMetaId> ITestSessionStream.Metadata => Metadata;
+
+			[BsonRequired, BsonElement("tg")]
+			public List<TestTagId> Tags { get; set; } = new List<TestTagId>();
+			IReadOnlyList<TestTagId> ITestSessionStream.Tags => Tags;
+
+			private TestSessionStreamDocument()
+			{
+
+			}
+
+			public TestSessionStreamDocument(StreamId streamId)
+			{
+				Id = ObjectId.GenerateNewId();
+				StreamId = streamId;
+			}
+		}
+
+		/// <summary>
+		/// Store the test session overall result
+		/// </summary>
+		class TestSessionDocument : ITestSession
+		{
+			[BsonRequired, BsonId]
+			public TestSessionId Id { get; set; }
+
+			[BsonRequired, BsonElement("sid")]
+			public StreamId StreamId { get; set; }
+
+			[BsonRequired, BsonElement("m")]
+			public TestMetaId Metadata { get; set; }
+
+			[BsonIgnoreIfNull, BsonElement("tg")]
+			public IReadOnlyList<TestTagId>? Tags { get; set; }
+
+			[BsonRequired, BsonElement("did")]
+			public ObjectId TestDataId { get; set; }
+
+			public CommitIdWithOrder BuildCommitId
+			{
+				get => (BuildCommitName != null) ? new CommitIdWithOrder(BuildCommitName, BuildCommitOrder) : CommitIdWithOrder.FromPerforceChange(BuildCommitOrder);
+				set => (BuildCommitName, BuildCommitOrder) = (value.Name, value.Order);
+			}
+
+			[BsonIgnoreIfNull, BsonElement("bcn")]
+			public string? BuildCommitName { get; set; }
+
+			[BsonRequired, BsonElement("bcl")]
+			public int BuildCommitOrder { get; set; }
+
+			[BsonRequired, BsonElement("d")]
+			public TimeSpan Duration { get; set; }
+
+			[BsonRequired, BsonElement("s")]
+			public DateTime StartDateTime { get; set; }
+
+			[BsonRequired, BsonElement("tid")]
+			public TestId NameRef { get; set; }
+
+			[BsonRequired, BsonElement("o")]
+			public TestOutcome Outcome { get; set; }
+
+			[BsonRequired, BsonElement("ptc")]
+			public int PhasesTotalCount { get; set; }
+
+			[BsonRequired, BsonElement("psc")]
+			public int PhasesSucceededCount { get; set; }
+
+			[BsonRequired, BsonElement("puc")]
+			public int PhasesUndefinedCount { get; set; }
+
+			[BsonRequired, BsonElement("pfc")]
+			public int PhasesFailedCount { get; set; }
+
+			[BsonRequired, BsonElement("jid")]
+			public JobId JobId { get; set; }
+
+			[BsonRequired, BsonElement("sjid")]
+			public JobStepId StepId { get; set; }
+
+			private TestSessionDocument()
+			{
+
+			}
+
+			public TestSessionDocument(StreamId streamId, ObjectId testDataId, TestId testId, TestMetaId metaId)
+			{
+				Id = TestSessionId.GenerateNewId();
+				StreamId = streamId;
+				TestDataId = testDataId;
+				NameRef = testId;
+				Metadata = metaId;
+			}
+		}
+
+		/// <summary>
+		/// Store the test phase result
+		/// </summary>
+		class TestPhaseSessionDocument : ITestPhaseSession
+		{
+			[BsonRequired, BsonId]
+			public TestPhaseSessionId Id { get; set; }
+
+			[BsonRequired, BsonElement("tid")]
+			public TestPhaseId PhaseRef { get; set; }
+
+			[BsonRequired, BsonElement("tsid")]
+			public TestSessionId SessionId { get; set; }
+
+			[BsonRequired, BsonElement("sid")]
+			public StreamId StreamId { get; set; }
+
+			[BsonRequired, BsonElement("m")]
+			public TestMetaId Metadata { get; set; }
+
+			public CommitIdWithOrder BuildCommitId
+			{
+				get => (BuildCommitName != null) ? new CommitIdWithOrder(BuildCommitName, BuildCommitOrder) : CommitIdWithOrder.FromPerforceChange(BuildCommitOrder);
+				set => (BuildCommitName, BuildCommitOrder) = (value.Name, value.Order);
+			}
+
+			[BsonIgnoreIfNull, BsonElement("bcn")]
+			public string? BuildCommitName { get; set; }
+
+			[BsonRequired, BsonElement("bcl")]
+			public int BuildCommitOrder { get; set; }
+
+			[BsonRequired, BsonElement("d")]
+			public TimeSpan Duration { get; set; }
+
+			[BsonRequired, BsonElement("s")]
+			public DateTime StartDateTime { get; set; }
+
+			[BsonRequired, BsonElement("o")]
+			public TestPhaseOutcome Outcome { get; set; }
+
+			[BsonIgnoreIfNull, BsonElement("esp")]
+			public string? EventStreamPath { get; set; }
+
+			[BsonIgnoreIfNull, BsonElement("hw")]
+			public bool? HasWarning { get; set; }
+
+			[BsonIgnoreIfNull, BsonElement("ef")]
+			public string? ErrorFingerprint { get; set; }
+
+			[BsonRequired, BsonElement("jid")]
+			public JobId JobId { get; set; }
+
+			[BsonRequired, BsonElement("sjid")]
+			public JobStepId StepId { get; set; }
+
+			private TestPhaseSessionDocument()
+			{
+
+			}
+
+			public TestPhaseSessionDocument(StreamId streamId, TestSessionId sessionId, TestPhaseId phaseId, TestMetaId metaId)
+			{
+				Id = TestPhaseSessionId.GenerateNewId();
+				StreamId = streamId;
+				SessionId = sessionId;
+				PhaseRef = phaseId;
+				Metadata = metaId;
+			}
+		}
+
+		/// <summary>
+		/// Information about a test data document
+		/// </summary>
+		class TestDataDocument : ITestData
+		{
+			public ObjectId Id { get; set; }
+			public StreamId StreamId { get; set; }
+			public TemplateId TemplateRefId { get; set; }
+			public JobId JobId { get; set; }
+			public JobStepId StepId { get; set; }
+
+			[BsonIgnore]
+			public CommitIdWithOrder CommitId
+			{
+				get => (CommitName != null) ? new CommitIdWithOrder(CommitName, CommitOrder) : CommitIdWithOrder.FromPerforceChange(CommitOrder);
+				set => (CommitName, CommitOrder) = (value.Name, value.Order);
+			}
+
+			public string? CommitName { get; set; }
+
+			[BsonElement("Change")]
+			public int CommitOrder { get; set; }
+
+			public string Key { get; set; }
+			public BsonDocument Data { get; set; }
+
+			private TestDataDocument()
+			{
+				Key = String.Empty;
+				Data = new BsonDocument();
+			}
+
+			public TestDataDocument(IJob job, IJobStep jobStep, string key, BsonDocument value)
+			{
+				Id = ObjectId.GenerateNewId();
+				StreamId = job.StreamId;
+				TemplateRefId = job.TemplateId;
+				JobId = job.Id;
+				StepId = jobStep.Id;
+				CommitId = job.CommitId;
+				Key = key;
+				Data = value;
+			}
+		}
+
+		readonly IMongoCollection<TestDataDocument> _testDataDocuments;
+		readonly IMongoCollection<TestMetadataDocument> _testMeta;
+		readonly IMongoCollection<TestTagDocument> _testTags;
+		readonly IMongoCollection<TestNameDocument> _tests;
+		readonly IMongoCollection<TestPhaseDocument> _testPhases;
+		readonly IMongoCollection<TestSessionDocument> _testSessions;
+		readonly IMongoCollection<TestPhaseSessionDocument> _testPhaseSessions;
+		readonly IMongoCollection<TestSessionStreamDocument> _testStreams;
+		readonly ICommitService _commitService;
+		readonly Tracer _tracer;
+		readonly ILogger _logger;
+
+		/// <summary>
+		/// Constructor
+		/// </summary>
+		public TestDataCollectionV2(IMongoService mongoService, ICommitService commitService, Tracer tracer, ILogger<TestDataCollection> logger)
+		{
+			_commitService = commitService;
+			_tracer = tracer;
+			_logger = logger;
+
+			List<MongoIndex<TestDataDocument>> indexes = new List<MongoIndex<TestDataDocument>>();
+			_testDataDocuments = mongoService.GetCollection<TestDataDocument>("TestDataV2", indexes);
+
+			List<MongoIndex<TestMetadataDocument>> metaIndexes = new List<MongoIndex<TestMetadataDocument>>();
+			metaIndexes.Add(keys => keys.Wildcard(x => x.KeyValuePairs));
+			_testMeta = mongoService.GetCollection<TestMetadataDocument>("TestDataV2.Metadata", metaIndexes);
+
+			List<MongoIndex<TestTagDocument>> tagIndexes = new List<MongoIndex<TestTagDocument>>();
+			tagIndexes.Add(keys => keys.Ascending(x => x.Name), unique: true);
+			_testTags = mongoService.GetCollection<TestTagDocument>("TestDataV2.TestTags", tagIndexes);
+
+			List<MongoIndex<TestNameDocument>> testIndexes = new List<MongoIndex<TestNameDocument>>();
+			testIndexes.Add(keys => keys.Ascending(x => x.Key), unique: true);
+			_tests = mongoService.GetCollection<TestNameDocument>("TestDataV2.TestNames", testIndexes);
+
+			List<MongoIndex<TestPhaseDocument>> testPhaseIndexes = new List<MongoIndex<TestPhaseDocument>>();
+			testPhaseIndexes.Add(keys => keys.Ascending(x => x.TestNameRef).Ascending(x => x.Key), unique: true);
+			_testPhases = mongoService.GetCollection<TestPhaseDocument>("TestDataV2.TestPhases", testPhaseIndexes);
+
+			List<MongoIndex<TestSessionDocument>> testSessionIndexes = new List<MongoIndex<TestSessionDocument>>();
+			testSessionIndexes.Add(keys => keys.Ascending(x => x.StreamId).Ascending(x => x.NameRef).Descending(x => x.BuildCommitOrder));
+			_testSessions = mongoService.GetCollection<TestSessionDocument>("TestDataV2.TestSessions", testSessionIndexes);
+
+			List<MongoIndex<TestPhaseSessionDocument>> testPhaseSessionIndexes = new List<MongoIndex<TestPhaseSessionDocument>>();
+			testPhaseSessionIndexes.Add(keys => keys.Ascending(x => x.StreamId).Ascending(x => x.PhaseRef).Descending(x => x.BuildCommitOrder));
+			_testPhaseSessions = mongoService.GetCollection<TestPhaseSessionDocument>("TestDataV2.TestPhaseSessions", testPhaseSessionIndexes);
+
+			List<MongoIndex<TestSessionStreamDocument>> streamIndexes = new List<MongoIndex<TestSessionStreamDocument>>();
+			streamIndexes.Add(keys => keys.Ascending(x => x.StreamId), unique: true);
+			_testStreams = mongoService.GetCollection<TestSessionStreamDocument>("TestDataV2.TestStreams", streamIndexes);
+		}
+
+		/// <inheritdoc/>
+		public async Task<IReadOnlyList<ITestSessionStream>> FindTestSessionStreamsAsync(StreamId[] streamIds, CancellationToken cancellationToken = default)
+		{
+			return await _testStreams.Find(Builders<TestSessionStreamDocument>.Filter.In(x => x.StreamId, streamIds)).ToListAsync(cancellationToken);
+		}
+
+		/// <inheritdoc/>
+		public async Task<IReadOnlyList<ITestSession>> FindTestSessionsAsync(StreamId[] streamIds, TestId[]? testIds = null, TestMetaId[]? metaIds = null, DateTime? minCreateTime = null, DateTime? maxCreateTime = null, CommitId? minCommitId = null, CommitId? maxCommitId = null, CancellationToken cancellationToken = default)
+		{
+			FilterDefinition<TestSessionDocument> filter = FilterDefinition<TestSessionDocument>.Empty;
+			FilterDefinitionBuilder<TestSessionDocument> filterBuilder = Builders<TestSessionDocument>.Filter;
+
+			filter &= filterBuilder.In(x => x.StreamId, streamIds);
+
+			if (minCreateTime != null)
+			{
+				TestSessionId minTime = TestSessionId.GenerateNewId(minCreateTime.Value);
+				filter &= filterBuilder.Gte(x => x.Id!, minTime);
+
+			}
+			if (maxCreateTime != null)
+			{
+				TestSessionId maxTime = TestSessionId.GenerateNewId(maxCreateTime.Value);
+				filter &= filterBuilder.Lte(x => x.Id!, maxTime);
+			}
+
+			if (metaIds != null && metaIds.Length > 0)
+			{
+				filter &= filterBuilder.In(x => x.Metadata, metaIds);
+			}
+
+			if (testIds != null && testIds.Length > 0)
+			{
+				filter &= filterBuilder.In(x => x.NameRef!, testIds);
+			}
+
+			if (minCommitId != null)
+			{
+				int minCommitOrder = minCommitId.GetPerforceChange();
+				filter &= filterBuilder.Gte(x => x.BuildCommitOrder, minCommitOrder);
+			}
+			if (maxCommitId != null)
+			{
+				int maxCommitOrder = maxCommitId.GetPerforceChange();
+				filter &= filterBuilder.Lte(x => x.BuildCommitOrder, maxCommitOrder);
+			}
+
+			List<TestSessionDocument> results;
+
+			using (TelemetrySpan _ = _tracer.StartActiveSpan($"{nameof(TestDataCollectionV2)}.{nameof(FindTestSessionsAsync)}"))
+			{
+				results = await _testSessions.Find(filter).ToListAsync(cancellationToken);
+			}
+
+			return results.ConvertAll<ITestSession>(x => x);
+		}
+
+		/// <inheritdoc/>
+		public async Task<IReadOnlyList<ITestNameRef>> FindTestNameRefsAsync(TestId[]? testIds = null, string[]? keys = null, CancellationToken cancellationToken = default)
+		{
+			FilterDefinition<TestNameDocument> filter = FilterDefinition<TestNameDocument>.Empty;
+			FilterDefinitionBuilder<TestNameDocument> filterBuilder = Builders<TestNameDocument>.Filter;
+
+			if (testIds != null && testIds.Length > 0)
+			{
+				filter &= filterBuilder.In(x => x.Id, testIds);
+			}
+
+			if (keys != null && keys.Length > 0)
+			{
+				filter &= filterBuilder.In(x => x.Key, keys);
+			}
+
+			return await _tests.Find(filter).ToListAsync(cancellationToken);
+		}
+
+		/// <inheritdoc/>
+		public async Task<IReadOnlyList<ITestMetaRef>> FindTestMetaAsync(TestMetaId[]? metaIds = null, IReadOnlyDictionary<string, string>? keyValues = null, CancellationToken cancellationToken = default)
+		{
+			FilterDefinition<TestMetadataDocument> filter = FilterDefinition<TestMetadataDocument>.Empty;
+			FilterDefinitionBuilder<TestMetadataDocument> filterBuilder = Builders<TestMetadataDocument>.Filter;
+
+			if (keyValues != null && keyValues.Count > 0)
+			{
+				foreach (KeyValuePair<string, string> item in keyValues)
+				{
+					filter &= filterBuilder.Eq($"m.{item.Key}", item.Value);
+				}
+			}
+
+			if (metaIds != null && metaIds.Length > 0)
+			{
+				filter &= filterBuilder.In(x => x.Id, metaIds);
+			}
+
+			return await _testMeta.Find(filter).ToListAsync(cancellationToken);
+		}
+
+		/// <inheritdoc/>
+		public async Task<IReadOnlyList<ITestPhaseRef>> FindTestPhasesAsync(TestId[] testIds, string[]? keys = null, CancellationToken cancellationToken = default)
+		{
+			FilterDefinition<TestPhaseDocument> filter = FilterDefinition<TestPhaseDocument>.Empty;
+			FilterDefinitionBuilder<TestPhaseDocument> filterBuilder = Builders<TestPhaseDocument>.Filter;
+
+			filter &= filterBuilder.In(x => x.TestNameRef, testIds);
+
+			if (keys != null && keys.Length > 0)
+			{
+				filter &= filterBuilder.In(x => x.Key, keys);
+			}
+
+			return await _testPhases.Find(filter).ToListAsync(cancellationToken);
+		}
+
+		/// <inheritdoc/>
+		public async Task<IReadOnlyList<ITestTagRef>> FindTestTagsAsync(TestTagId[]? tagIds = null, string[]? names = null, CancellationToken cancellationToken = default)
+		{
+			FilterDefinition<TestTagDocument> filter = FilterDefinition<TestTagDocument>.Empty;
+			FilterDefinitionBuilder<TestTagDocument> filterBuilder = Builders<TestTagDocument>.Filter;
+
+			if (tagIds != null && tagIds.Length > 0)
+			{
+				filter &= filterBuilder.In(x => x.Id, tagIds);
+			}
+
+			if (names != null && names.Length > 0)
+			{
+				filter &= filterBuilder.In(x => x.Name, names);
+			}
+
+			return await _testTags.Find(filter).ToListAsync(cancellationToken);
+		}
+
+		/// <inheritdoc/>
+		public async Task<IReadOnlyList<ITestPhaseSession>> FindTestPhaseSessionsAsync(StreamId[] streamIds, TestPhaseId[] phaseIds, DateTime? minCreateTime = null, DateTime? maxCreateTime = null, CommitId? minCommitId = null, CommitId? maxCommitId = null, CancellationToken cancellationToken = default)
+		{
+			FilterDefinition<TestPhaseSessionDocument> filter = FilterDefinition<TestPhaseSessionDocument>.Empty;
+			FilterDefinitionBuilder<TestPhaseSessionDocument> filterBuilder = Builders<TestPhaseSessionDocument>.Filter;
+
+			filter &= filterBuilder.In(x => x.StreamId, streamIds);
+			filter &= filterBuilder.In(x => x.PhaseRef, phaseIds);
+
+			if (minCreateTime != null)
+			{
+				TestPhaseSessionId minTime = TestPhaseSessionId.GenerateNewId(minCreateTime.Value);
+				filter &= filterBuilder.Gte(x => x.Id!, minTime);
+
+			}
+			if (maxCreateTime != null)
+			{
+				TestPhaseSessionId maxTime = TestPhaseSessionId.GenerateNewId(maxCreateTime.Value);
+				filter &= filterBuilder.Lte(x => x.Id!, maxTime);
+			}
+
+			if (minCommitId != null)
+			{
+				int minCommitOrder = minCommitId.GetPerforceChange();
+				filter &= filterBuilder.Gte(x => x.BuildCommitOrder, minCommitOrder);
+			}
+			if (maxCommitId != null)
+			{
+				int maxCommitOrder = maxCommitId.GetPerforceChange();
+				filter &= filterBuilder.Lte(x => x.BuildCommitOrder, maxCommitOrder);
+			}
+
+			List<TestPhaseSessionDocument> results;
+
+			using (TelemetrySpan _ = _tracer.StartActiveSpan($"{nameof(TestDataCollectionV2)}.{nameof(FindTestPhaseSessionsAsync)}"))
+			{
+				results = await _testPhaseSessions.Find(filter).ToListAsync(cancellationToken);
+			}
+
+			return results.ConvertAll<ITestPhaseSession>(x => x);
+		}
+
+		/// <inheritdoc/>
+		public async Task<IReadOnlyList<ITestData>> AddAsync(IJob job, IJobStep step, (string key, BsonDocument value)[] data, CancellationToken cancellationToken = default)
+		{
+			// detailed test data
+			List <TestDataDocument> documents = new List<TestDataDocument>();
+			for (int i = 0; i < data.Length; i++)
+			{
+				(string key, BsonDocument document) = data[i];
+
+					int version;
+				if (document.TryGetInt32("version", out version) || document.TryGetInt32("Version", out version))
+				{
+					if (version > 1)
+					{
+						documents.Add(new TestDataDocument(job, step, key, document));
+					}
+				}
+			}
+
+			if (documents.Count == 0)
+			{
+				return documents;
+			}
+
+			await _testDataDocuments.InsertManyAsync(documents, null, cancellationToken);
+
+			try
+			{
+				await AddTestReportDataAsync(job, step, documents, cancellationToken);
+			}
+			catch (Exception ex)
+			{
+				_logger.LogWarning(ex, "Exception while adding test data  report, jobId: {JobId} stepId: {StepId}", job.Id, step.Id);
+			}
+
+			return documents;
+		}
+
+		/// <inheritdoc/>
+		public async Task<ITestData?> GetAsync(ObjectId id, CancellationToken cancellationToken)
+		{
+			return await _testDataDocuments.Find<TestDataDocument>(x => x.Id == id).FirstOrDefaultAsync(cancellationToken);
+		}
+
+		// --- (Minimal) Parsing of Automation framework tests, geared towards generating indexes for aggregate queries
+		class AutomatedTestSessionData
+		{
+			public class TestSessionSummary
+			{
+				[BsonRequired, BsonElement("testName")]
+				public string TestName { get; set; } = String.Empty;
+				[BsonRequired, BsonElement("dateTime")]
+				public string DateTime { get; set; } = String.Empty;
+				[BsonRequired, BsonElement("timeElapseSec")]
+				public double TimeElapseSec { get; set; }
+				[BsonRequired, BsonElement("phasesTotalCount")]
+				public int PhasesTotalCount { get; set; }
+				[BsonRequired, BsonElement("phasesSucceededCount")]
+				public int PhasesSucceededCount { get; set; }
+				[BsonRequired, BsonElement("phasesUndefinedCount")]
+				public int PhasesUndefinedCount { get; set; }
+				[BsonRequired, BsonElement("phasesFailedCount")]
+				public int PhasesFailedCount { get; set; }
+			}
+
+			public class TestPhase
+			{
+				[BsonRequired, BsonElement("name")]
+				public string Name { get; set; } = String.Empty;
+				[BsonRequired, BsonElement("key")]
+				public string Key { get; set; } = String.Empty;
+				[BsonRequired, BsonElement("dateTime")]
+				public string DateTime { get; set; } = String.Empty;
+				[BsonRequired, BsonElement("timeElapseSec")]
+				public double TimeElapseSec { get; set; }
+				[BsonRequired, BsonElement("outcome")]
+				public TestPhaseOutcome Outcome { get; set; }
+				[BsonIgnoreIfNull, BsonElement("errorFingerprint")]
+				public string? ErrorFingerprint { get; set; }
+				[BsonIgnoreIfNull, BsonElement("hasWarning")]
+				public bool? HasWarning { get; set; }
+				[BsonIgnoreIfNull, BsonElement("eventStreamPath")]
+				public string? EventStreamPath { get; set; }
+			}
+
+			[BsonRequired, BsonElement("summary")]
+			public TestSessionSummary? Summary { get; set; }
+			[BsonRequired, BsonElement("phases")]
+			public List<TestPhase> Phases { get; set; } = new List<TestPhase>();
+			[BsonRequired, BsonElement("metadata")]
+			public Dictionary<string, string> Metadata { get; set; } = new Dictionary<string, string>();
+			[BsonIgnoreIfNull, BsonElement("tags")]
+			public List<string>? Tags { get; set; }
+		}
+
+		private async Task AddTestReportDataAsync(IJob job, IJobStep step, IReadOnlyList<TestDataDocument> data, CancellationToken cancellationToken = default)
+		{
+			// do not add preflight to temporal data
+			if (job.PreflightCommitId != null)
+			{
+				return;
+			}
+
+			foreach (TestDataDocument document in data)
+			{
+				AutomatedTestSessionData testSession = BsonSerializer.Deserialize<AutomatedTestSessionData>(document.Data);
+
+				if (testSession.Summary == null)
+				{
+					throw new Exception($"Missing summary field in test session {document.Key}.");
+				}
+
+				if (testSession.Phases == null)
+				{
+					throw new Exception($"Missing phases field in test session {testSession.Summary.TestName}.");
+				}
+
+				if (testSession.Metadata == null)
+				{
+					throw new Exception($"Missing metadata field in test session {testSession.Summary.TestName}.");
+				}
+
+				IReadOnlyList<ITestTagRef>? tags = await AddTestTagsAsync(testSession.Tags, cancellationToken);
+				ITestMetaRef metadata = await AddTestMetadataAsync(testSession.Metadata, cancellationToken);
+				ITestNameRef testRef = await AddTestAsync(document.Key, testSession.Summary.TestName, metadata.Id, tags?.Select(t => t.Id).ToList(), document.StreamId, cancellationToken);
+
+				// test session
+				TestSessionDocument sessionDoc = new TestSessionDocument(document.StreamId, document.Id, testRef.Id, metadata.Id);
+				// job and commit
+				sessionDoc.JobId = job.Id;
+				sessionDoc.StepId = step.Id;
+				sessionDoc.BuildCommitId = document.CommitId;
+				// phase outcomes and summary
+				sessionDoc.Outcome = testSession.Summary.PhasesUndefinedCount > 0 ? TestOutcome.Unspecified :
+					testSession.Summary.PhasesFailedCount > 0 ? TestOutcome.Failure :
+						testSession.Summary.PhasesTotalCount > 0 && testSession.Summary.PhasesSucceededCount > 0 ? TestOutcome.Success : TestOutcome.Skipped;
+				sessionDoc.PhasesFailedCount = testSession.Summary.PhasesFailedCount;
+				sessionDoc.PhasesSucceededCount = testSession.Summary.PhasesSucceededCount;
+				sessionDoc.PhasesTotalCount = testSession.Summary.PhasesTotalCount;
+				sessionDoc.PhasesUndefinedCount = testSession.Summary.PhasesUndefinedCount;
+				// timing
+				DateTime startDateTime;
+				if (DateTime.TryParse(testSession.Summary.DateTime, out startDateTime))
+				{
+					sessionDoc.StartDateTime = startDateTime;
+				}
+				sessionDoc.Duration = TimeSpan.FromSeconds(testSession.Summary.TimeElapseSec);
+
+				// phase session
+				List<TestPhaseSessionDocument> phaseSessions = new();
+				foreach (AutomatedTestSessionData.TestPhase phase in testSession.Phases)
+				{
+					ITestPhaseRef phaseRef = await AddTestPhaseAsync(testRef.Id, phase.Key, phase.Name, cancellationToken);
+
+					TestPhaseSessionDocument phaseDoc = new TestPhaseSessionDocument(document.StreamId, sessionDoc.Id, phaseRef.Id, metadata.Id);
+					// job and commit
+					phaseDoc.JobId = job.Id;
+					phaseDoc.StepId = step.Id;
+					phaseDoc.BuildCommitId = document.CommitId;
+					// outcome and event summary
+					phaseDoc.Outcome = phase.Outcome;
+					phaseDoc.ErrorFingerprint = phase.ErrorFingerprint;
+					phaseDoc.EventStreamPath = phase.EventStreamPath;
+					phaseDoc.HasWarning = phase.HasWarning;
+					// timing
+					if (DateTime.TryParse(phase.DateTime, out startDateTime))
+					{
+						phaseDoc.StartDateTime = startDateTime;
+					}
+					phaseDoc.Duration = TimeSpan.FromSeconds(phase.TimeElapseSec);
+
+					phaseSessions.Add(phaseDoc);
+				}
+				if (phaseSessions.Count > 0)
+				{
+					await _testPhaseSessions.InsertManyAsync(phaseSessions, null, cancellationToken);
+				}
+
+				await _testSessions.InsertOneAsync(sessionDoc, null, cancellationToken);
+			}
+		}
+
+		private async Task<ITestMetaRef> AddTestMetadataAsync(IReadOnlyDictionary<string, string> metaData, CancellationToken cancellationToken)
+		{
+			FilterDefinitionBuilder<TestMetadataDocument> filterBuilder = Builders<TestMetadataDocument>.Filter;
+			FilterDefinition<TestMetadataDocument> filter = FilterDefinition<TestMetadataDocument>.Empty;
+
+			foreach (KeyValuePair<string, string> item in metaData)
+			{
+				filter &= filterBuilder.Eq($"m.{item.Key}", item.Value);
+			}
+
+			List<TestMetadataDocument> metaList = await _testMeta.Find(filter).ToListAsync(cancellationToken);
+			TestMetadataDocument? meta = metaList.OrderBy(m => m.KeyValuePairs.Count).FirstOrDefault();
+
+			if (meta == null || meta.KeyValuePairs.Count > metaData.Count)
+			{
+				meta = new TestMetadataDocument(metaData);
+				await _testMeta.InsertOneAsync(meta, null, cancellationToken);
+			}
+			// update only once per day
+			else if (meta.LastSeenUtc < DateTime.UtcNow.AddDays(-1))
+			{
+				UpdateDefinitionBuilder<TestMetadataDocument> updateBuilder = Builders<TestMetadataDocument>.Update;
+				List<UpdateDefinition<TestMetadataDocument>> updates = new List<UpdateDefinition<TestMetadataDocument>>();
+				updates.Add(updateBuilder.Set(x => x.LastSeenUtc, DateTime.UtcNow));
+				await _testMeta.FindOneAndUpdateAsync(x => x.Id == meta.Id, updateBuilder.Combine(updates), null, cancellationToken);
+			}
+
+			return meta;
+		}
+
+		private async Task<IReadOnlyList<ITestTagRef>?> AddTestTagsAsync(IReadOnlyList<string>? tags, CancellationToken cancellationToken)
+		{
+			if (tags == null)
+			{
+				return null;
+			}
+
+			FilterDefinitionBuilder<TestTagDocument> filterBuilder = Builders<TestTagDocument>.Filter;
+			FilterDefinition<TestTagDocument> filter = filterBuilder.In(i => i.Name, tags);
+
+			List<TestTagDocument> storedTags = await _testTags.Find(filter).ToListAsync(cancellationToken);
+			if (storedTags.Count > 0)
+			{
+				// update only once per day
+				List<TestTagId> needUpdateTags = storedTags.Where(i => i.LastSeenUtc < DateTime.UtcNow.AddDays(-1)).Select(i => i.Id).ToList();
+				if (needUpdateTags.Count > 0)
+				{
+					UpdateDefinitionBuilder<TestTagDocument> updateBuilder = Builders<TestTagDocument>.Update;
+					List<UpdateDefinition<TestTagDocument>> updates = new List<UpdateDefinition<TestTagDocument>>();
+					updates.Add(updateBuilder.Set(x => x.LastSeenUtc, DateTime.UtcNow));
+					FilterDefinition<TestTagDocument> filter2 = filterBuilder.In(i => i.Id, needUpdateTags);
+					await _testTags.UpdateManyAsync(filter2, updateBuilder.Combine(updates), null, cancellationToken);
+				}
+			}
+			// add the missing ones
+			List<TestTagDocument> missingTags = tags.Except(storedTags.Select(i => i.Name)).Select(t => new TestTagDocument(t)).ToList();
+			if (missingTags.Count > 0 )
+			{
+				await _testTags.InsertManyAsync(missingTags, null, cancellationToken);
+				storedTags.AddRange(missingTags);
+			}
+
+			return storedTags;
+		}
+
+		private async Task<ITestNameRef> AddTestAsync(string key, string name, TestMetaId metadataId, IReadOnlyList<TestTagId>? tags, StreamId streamId, CancellationToken cancellationToken)
+		{
+			FilterDefinitionBuilder<TestNameDocument> filterBuilder = Builders<TestNameDocument>.Filter;
+
+			TestNameDocument? test = await _tests.Find(filterBuilder.Eq(i => i.Key, key)).FirstOrDefaultAsync(cancellationToken);
+
+			if (test == null)
+			{
+				test = new TestNameDocument(key, name);
+				await _tests.InsertOneAsync(test, null, cancellationToken);
+			}
+			// update only once per day
+			else if (test.LastSeenUtc < DateTime.UtcNow.AddDays(-1))
+			{
+				UpdateDefinitionBuilder<TestNameDocument> updateBuilder = Builders<TestNameDocument>.Update;
+				List<UpdateDefinition<TestNameDocument>> updates = new List<UpdateDefinition<TestNameDocument>>();
+				updates.Add(updateBuilder.Set(x => x.LastSeenUtc, DateTime.UtcNow));
+				await _tests.FindOneAndUpdateAsync(x => x.Id == test.Id, updateBuilder.Combine(updates), null, cancellationToken);
+			}
+
+			await AddTestToStreamSessionAsync(test.Id, metadataId, tags, streamId, cancellationToken);
+
+			return test;
+		}
+
+		private async Task AddTestToStreamSessionAsync(TestId testRef, TestMetaId metadataId, IReadOnlyList<TestTagId>? tags, StreamId streamId, CancellationToken cancellationToken)
+		{
+			List<TestSessionStreamDocument> streams = await _testStreams.Find(Builders<TestSessionStreamDocument>.Filter.Eq(x => x.StreamId, streamId)).ToListAsync(cancellationToken);
+
+			if (streams.Count == 0)
+			{
+				TestSessionStreamDocument streamDoc = new TestSessionStreamDocument(streamId);
+				streamDoc.Tests.Add(testRef);
+				streamDoc.Metadata.Add(metadataId);
+				if (tags != null)
+				{
+					streamDoc.Tags.AddRange(tags);
+				}
+				await _testStreams.InsertOneAsync(streamDoc, null, cancellationToken);
+			}
+			else
+			{
+				UpdateDefinitionBuilder<TestSessionStreamDocument> updateBuilder = Builders<TestSessionStreamDocument>.Update;
+				List<UpdateDefinition<TestSessionStreamDocument>> updates = new List<UpdateDefinition<TestSessionStreamDocument>>();
+
+				TestSessionStreamDocument streamDoc = streams[0];
+
+				if (!streamDoc.Tests.Contains(testRef))
+				{
+					streamDoc.Tests.Add(testRef);
+					updates.Add(updateBuilder.Set(x => x.Tests, streamDoc.Tests));
+				}
+				if (!streamDoc.Metadata.Contains(metadataId))
+				{
+					streamDoc.Metadata.Add(metadataId);
+					updates.Add(updateBuilder.Set(x => x.Metadata, streamDoc.Metadata));
+				}
+				if (tags != null)
+				{
+					List<TestTagId> missingTags = tags.Except(streamDoc.Tags).ToList();
+					if (missingTags.Count > 0)
+					{
+						streamDoc.Tags.AddRange(missingTags);
+						updates.Add(updateBuilder.Set(x => x.Tags, streamDoc.Tags));
+					}
+				}
+
+				if (updates.Count > 0)
+				{
+					FilterDefinitionBuilder<TestSessionStreamDocument> ebuilder = Builders<TestSessionStreamDocument>.Filter;
+					FilterDefinition<TestSessionStreamDocument> efilter = ebuilder.Eq(x => x.StreamId, streamDoc.StreamId);
+					await _testStreams.UpdateOneAsync(efilter, updateBuilder.Combine(updates), null, cancellationToken);
+				}
+			}
+		}
+
+		private async Task<ITestPhaseRef> AddTestPhaseAsync(TestId testId, string key, string name, CancellationToken cancellationToken)
+		{
+			FilterDefinitionBuilder<TestPhaseDocument> filterBuilder = Builders<TestPhaseDocument>.Filter;
+			FilterDefinition<TestPhaseDocument> filter = FilterDefinition<TestPhaseDocument>.Empty;
+
+			filter &= filterBuilder.Eq(i => i.TestNameRef, testId);
+			filter &= filterBuilder.Eq(i => i.Key, key);
+
+			TestPhaseDocument? phase = await _testPhases.Find(filter).FirstOrDefaultAsync(cancellationToken);
+
+			if (phase == null)
+			{
+				phase = new TestPhaseDocument(testId, key, name);
+				await _testPhases.InsertOneAsync(phase, null, cancellationToken);
+			}
+			// update only once per day
+			else if (phase.LastSeenUtc < DateTime.UtcNow.AddDays(-1))
+			{
+				UpdateDefinitionBuilder<TestPhaseDocument> updateBuilder = Builders<TestPhaseDocument>.Update;
+				List<UpdateDefinition<TestPhaseDocument>> updates = new List<UpdateDefinition<TestPhaseDocument>>();
+				updates.Add(updateBuilder.Set(x => x.LastSeenUtc, DateTime.UtcNow));
+				await _testPhases.FindOneAndUpdateAsync(x => x.Id == phase.Id, updateBuilder.Combine(updates), null, cancellationToken);
+			}
+
+			return phase;
+		}
+
+		/// <inheritdoc/>
+		public async Task<bool> UpdateAsync(int retainMonths, CancellationToken cancellationToken)
+		{
+			await ExpireTestDataAsync(retainMonths, cancellationToken);
+
+			return true;
+		}
+
+		static async Task<long> ExpireCollectionAsync<DocType>(IMongoCollection<DocType> collection, DateTime expireTime, CancellationToken cancellationToken) where DocType : ITestExpire
+		{
+			FilterDefinitionBuilder<DocType> filterBuilder = Builders<DocType>.Filter;
+			FilterDefinition<DocType> filter = filterBuilder.Empty;
+			filter &= filterBuilder.Lte(x => x.LastSeenUtc, expireTime);
+			DeleteResult result = await collection.DeleteManyAsync(filter, cancellationToken);
+
+			return result.DeletedCount;
+		}
+
+		private async Task ExpireTestDataAsync(int retainMonths, CancellationToken cancellationToken)
+		{
+			if (retainMonths <= 0)
+			{
+				return;
+			}
+
+			DateTime expireTime = DateTime.Now.AddMonths(-retainMonths);
+
+			// expire and prune root test data
+			long testsExpired = await ExpireCollectionAsync(_tests, expireTime, cancellationToken);
+			long phasesExpired = await ExpireCollectionAsync(_testPhases, expireTime, cancellationToken);
+			long metadataExpired = await ExpireCollectionAsync(_testMeta, expireTime, cancellationToken);
+			long tagsExpired = await ExpireCollectionAsync(_testTags, expireTime, cancellationToken);
+
+			// Check whether we need to prune
+			if (testsExpired > 0 || phasesExpired > 0 || metadataExpired > 0 || tagsExpired > 0)
+			{
+				_logger.LogInformation("Expired {TestsExpired} tests, {PhasesExpired} suites, {MetadataExpired} meta, {TagsExpired} tags", testsExpired, phasesExpired, metadataExpired, tagsExpired);
+				await PruneDataAsync(cancellationToken);
+			}
+
+			// expire test data documents
+			{
+				ObjectId dataExpireTime = ObjectId.GenerateNewId(expireTime);
+
+				FilterDefinition<TestDataDocument> filter = FilterDefinition<TestDataDocument>.Empty;
+				FilterDefinitionBuilder<TestDataDocument> filterBuilder = Builders<TestDataDocument>.Filter;
+				filter &= filterBuilder.Lte(x => x.Id!, dataExpireTime);
+				DeleteResult result = await _testDataDocuments.DeleteManyAsync(filter, null, cancellationToken);
+				if (result.DeletedCount > 0)
+				{
+					_logger.LogInformation("Expired {NumDeleted} test data documents", result.DeletedCount);
+				}
+			}
+
+			// expire test sessions
+			{
+				TestSessionId sessionExpireTime = TestSessionId.GenerateNewId(expireTime);
+
+				FilterDefinition<TestSessionDocument> filter = FilterDefinition<TestSessionDocument>.Empty;
+				FilterDefinitionBuilder<TestSessionDocument> filterBuilder = Builders<TestSessionDocument>.Filter;
+				filter &= filterBuilder.Lte(x => x.Id!, sessionExpireTime);
+				DeleteResult result = await _testSessions.DeleteManyAsync(filter, cancellationToken);
+				if (result.DeletedCount > 0)
+				{
+					_logger.LogInformation("Expired {NumDeleted} test session documents", result.DeletedCount);
+				}
+			}
+
+			// expire test phase sessions
+			{
+				TestPhaseSessionId sessionExpireTime = TestPhaseSessionId.GenerateNewId(expireTime);
+				FilterDefinition<TestPhaseSessionDocument> filter = FilterDefinition<TestPhaseSessionDocument>.Empty;
+				FilterDefinitionBuilder<TestPhaseSessionDocument> filterBuilder = Builders<TestPhaseSessionDocument>.Filter;
+				filter &= filterBuilder.Lte(x => x.Id!, sessionExpireTime);
+				DeleteResult result = await _testPhaseSessions.DeleteManyAsync(filter, cancellationToken);
+				if (result.DeletedCount > 0)
+				{
+					_logger.LogInformation("Expired {NumDeleted} test phase session documents", result.DeletedCount);
+				}
+			}
+		}
+
+		/// <inheritdoc/>
+		public async Task PruneDataAsync(CancellationToken cancellationToken = default)
+		{
+			HashSet<TestId> testIds = new HashSet<TestId>();
+			foreach (TestId testId in await (await _tests.DistinctAsync(x => x.Id, Builders<TestNameDocument>.Filter.Empty, null, cancellationToken)).ToListAsync(cancellationToken))
+			{
+				testIds.Add(testId);
+			}
+			_logger.LogInformation("Test data pruning to {TestIdCount} test name ids", testIds.Count);
+
+			HashSet<TestMetaId> metadataIds = new HashSet<TestMetaId>();
+			foreach (TestMetaId metaId in await (await _testMeta.DistinctAsync(x => x.Id, Builders<TestMetadataDocument>.Filter.Empty, null, cancellationToken)).ToListAsync(cancellationToken))
+			{
+				metadataIds.Add(metaId);
+			}
+			_logger.LogInformation("Test data pruning to {MetadataIdCount} test metadata ids", metadataIds.Count);
+
+			HashSet<TestTagId> tagIds = new HashSet<TestTagId>();
+			foreach (TestTagId tagId in await (await _testTags.DistinctAsync(x => x.Id, Builders<TestTagDocument>.Filter.Empty, null, cancellationToken)).ToListAsync(cancellationToken))
+			{
+				tagIds.Add(tagId);
+			}
+			_logger.LogInformation("Test data pruning to {TagIdCount} test tag ids", tagIds.Count);
+
+			int testStreamsDeleted = 0;
+
+			// update test session streams
+			{
+				List<TestSessionStreamDocument> testStreams = await _testStreams.Find(Builders<TestSessionStreamDocument>.Filter.Empty).ToListAsync(cancellationToken);
+				foreach (TestSessionStreamDocument stream in testStreams)
+				{
+					List<TestId> tests = stream.Tests.Where(x => testIds.Contains(x)).ToList();
+					if (tests.Count == 0)
+					{
+						await _testStreams.DeleteOneAsync(x => x.Id == stream.Id, cancellationToken);
+						testStreamsDeleted++;
+						continue;
+					}
+
+					List<TestMetaId> metadata = stream.Metadata.Where(x => metadataIds.Contains(x)).ToList();
+					List<TestTagId> tags = stream.Tags.Where(x => tagIds.Contains(x)).ToList();
+					if (tests.Count != stream.Tests.Count || metadata.Count != stream.Metadata.Count || tags.Count != stream.Tags.Count)
+					{
+						UpdateDefinitionBuilder<TestSessionStreamDocument> updateBuilder = Builders<TestSessionStreamDocument>.Update;
+						List<UpdateDefinition<TestSessionStreamDocument>> updates = new List<UpdateDefinition<TestSessionStreamDocument>>();
+
+						if (tests.Count != stream.Tests.Count)
+						{
+							updates.Add(updateBuilder.Set(x => x.Tests, tests));
+						}
+						if (metadata.Count != stream.Metadata.Count)
+						{
+							updates.Add(updateBuilder.Set(x => x.Metadata, metadata));
+						}
+						if (tags.Count != stream.Tags.Count)
+						{
+							updates.Add(updateBuilder.Set(x => x.Tags, tags));
+						}
+
+						await _testStreams.FindOneAndUpdateAsync(x => x.Id == stream.Id, updateBuilder.Combine(updates), null, cancellationToken);
+					}
+				}
+			}
+
+			if (testStreamsDeleted > 0)
+			{
+				_logger.LogInformation("Pruned {TestStreamDeleteCount} test session streams", testStreamsDeleted);
+			}
+		}
+	}
+}
diff --git a/Engine/Source/Programs/Horde/Plugins/Build/HordeServer.Build/Jobs/TestData/TestDataController.cs b/Engine/Source/Programs/Horde/Plugins/Build/HordeServer.Build/Jobs/TestData/TestDataController.cs
--- a/Engine/Source/Programs/Horde/Plugins/Build/HordeServer.Build/Jobs/TestData/TestDataController.cs
+++ b/Engine/Source/Programs/Horde/Plugins/Build/HordeServer.Build/Jobs/TestData/TestDataController.cs
@@ -463,5 +463,374 @@ namespace HordeServer.Jobs.TestData
 				Data = BsonSerializer.Deserialize<Dictionary<string, object>>(testData.Data)
 			}, filter);
 		}
+
+		/// <summary>
+		/// Retrieve information about a specific issue
+		/// </summary>
+		/// <param name="testDataId">Id of the document to get information about</param>
+		/// <param name="filter">Filter for the properties to return</param>
+		/// <param name="cancellationToken">Cancellation token for the operation</param>
+		/// <returns>List of matching agents</returns>
+		[HttpGet]
+		[Route("/api/v3/testdata/{testDataId}")]
+		[ProducesResponseType(typeof(GetTestDataResponse), 200)]
+		public async Task<ActionResult<object>> GetTestDataV2Async(string testDataId, [FromQuery] PropertyFilter? filter = null, CancellationToken cancellationToken = default)
+		{
+			ITestData? testData = await _testDataService.GetTestDataV2Async(ObjectId.Parse(testDataId), cancellationToken);
+			if (testData == null)
+			{
+				return NotFound();
+			}
+			if (!await _jobService.AuthorizeAsync(testData.JobId, JobAclAction.ViewJob, User, _buildConfig.Value, cancellationToken))
+			{
+				return Forbid();
+			}
+
+			return PropertyFilter.Apply(new GetTestDataResponse
+			{
+				Id = testData.Id.ToString(),
+				StreamId = testData.StreamId.ToString(),
+				TemplateRefId = testData.TemplateRefId.ToString(),
+				JobId = testData.JobId.ToString(),
+				StepId = testData.StepId.ToString(),
+				CommitId = testData.CommitId,
+				Key = testData.Key,
+				Data = BsonSerializer.Deserialize<Dictionary<string, object>>(testData.Data)
+			}, filter);
+		}
+
+		/// <summary>
+		/// Get stream test data for the provided ids
+		/// </summary>
+		/// <param name="streamIds"></param>
+		/// <returns></returns>
+		[HttpGet]
+		[Route("/api/v3/testdata/streams")]
+		[ProducesResponseType(typeof(List<GetTestSessionStreamResponse>), 200)]
+		public async Task<ActionResult<List<GetTestSessionStreamResponse>>> GetTestStreamsV2Async([FromQuery(Name = "Id")] string[] streamIds)
+		{
+			StreamId[] streamIdValues = Array.ConvertAll(streamIds, x => new StreamId(x));
+
+			List<StreamId> queryStreams = new List<StreamId>();
+			List<GetTestSessionStreamResponse> responses = new List<GetTestSessionStreamResponse>();
+
+			// authorize streams
+			foreach (StreamId streamId in streamIdValues)
+			{
+				if (_buildConfig.Value.TryGetStream(streamId, out StreamConfig? streamConfig) && streamConfig.Authorize(JobAclAction.ViewJob, User))
+				{
+					queryStreams.Add(streamId);
+				}
+			}
+
+			if (queryStreams.Count == 0)
+			{
+				return responses;
+			}
+
+			IReadOnlyList<ITestSessionStream> streams = await _testDataService.FindTestSessionStreamsV2Async(queryStreams.ToArray());
+
+			// flatten requested streams to single service queries		
+			HashSet<TestId> testIds = new HashSet<TestId>();
+			HashSet<TestMetaId> metaIds = new HashSet<TestMetaId>();
+			HashSet<TestTagId> tagIds = new HashSet<TestTagId>();
+			foreach (ITestSessionStream stream in streams)
+			{
+				testIds.UnionWith(stream.Tests);
+				metaIds.UnionWith(stream.Metadata);
+				tagIds.UnionWith(stream.Tags);
+			}
+
+			// Gather all relevant test names
+			IReadOnlyList<ITestNameRef> tests = new List<ITestNameRef>();
+			if (testIds.Count > 0)
+			{
+				tests = await _testDataService.FindTestNameRefsV2Async(testIds.ToArray());
+			}
+
+			// Gather all relevant metadata
+			IReadOnlyList<ITestMetaRef> metadata = new List<ITestMetaRef>();
+			if (metaIds.Count > 0)
+			{
+				metadata = await _testDataService.FindTestMetaV2Async(metaIds.ToArray());
+			}
+
+			// Gather all relevant tags
+			IReadOnlyList<ITestTagRef> tags = new List<ITestTagRef>();
+			if (tagIds.Count > 0)
+			{
+				tags = await _testDataService.FindTestTagsV2Async(tagIds.ToArray());
+			}
+
+			// generate individual stream responses
+			foreach (ITestSessionStream s in streams)
+			{
+				IEnumerable<ITestNameRef> filteredTests = tests.Where(x => s.Tests.Contains(x.Id));
+				IEnumerable<ITestMetaRef> filteredMetadata = metadata.Where(x => s.Metadata.Contains(x.Id));
+				IEnumerable<ITestTagRef> filteredTags = tags.Where(x => s.Tags.Contains(x.Id));
+
+				responses.Add(new GetTestSessionStreamResponse
+				{
+					StreamId = s.StreamId.ToString(),
+					Tests = filteredTests.Select(test => new GetTestNameResponse
+					{
+						Id = test.Id.ToString(),
+						Name = test.Name,
+						Key = test.Key
+					}).ToList(),
+					TestMetadata = filteredMetadata.Select(meta => new GetTestMetadataResponse
+					{
+						Id = meta.Id.ToString(),
+						KeyValuePairs = new (meta.KeyValuePairs)
+					}).ToList(),
+					TestTags = filteredTags.Select(tag => new GetTestTagResponse
+					{
+						Id = tag.Id.ToString(),
+						Name = new (tag.Name)
+					}).ToList()
+				});
+			}
+
+			return responses;
+		}
+
+		/// <summary>
+		/// Get test names from provided ids
+		/// </summary>
+		/// <param name="request"></param>
+		/// <returns></returns>
+		[HttpPost]
+		[Route("/api/v3/testdata/tests")]
+		[ProducesResponseType(typeof(List<GetTestNameResponse>), 200)]
+		public async Task<ActionResult<List<GetTestNameResponse>>> GetTestNameRefV2Async([FromBody] GetTestsRequest request)
+		{
+			HashSet<string> testIds = new HashSet<string>(request.TestIds);
+
+			IReadOnlyList<ITestNameRef> testValues = await _testDataService.FindTestNameRefsV2Async(testIds.Select(x => TestId.Parse(x)).ToArray());
+
+			return testValues.Select(x => new GetTestNameResponse
+			{
+				Id = x.Id.ToString(),
+				Name = x.Name,
+				Key = x.Key
+			}).ToList();
+		}
+
+		/// <summary>
+		/// Get test metadata from provided ids
+		/// </summary>
+		/// <param name="request"></param>
+		/// <returns></returns>
+		[HttpPost]
+		[Route("/api/v3/testdata/metadata")]
+		[ProducesResponseType(typeof(List<GetTestMetadataResponse>), 200)]
+		public async Task<ActionResult<List<GetTestMetadataResponse>>> GetTestMetadataV2Async([FromBody] GetMetadataRequest request)
+		{
+			TestMetaId[]? metaIds = request.MetadataIds?.Select(x => TestMetaId.Parse(x)).ToHashSet().ToArray();
+
+			IReadOnlyList<ITestMetaRef> metadataValues = await _testDataService.FindTestMetaV2Async(metaIds, request.KeyValuePairs);
+
+			return metadataValues.Select(x => new GetTestMetadataResponse
+			{
+				Id = x.Id.ToString(),
+				KeyValuePairs = new(x.KeyValuePairs)
+			}).ToList();
+		}
+
+		/// <summary>
+		/// Get test tags from provided ids
+		/// </summary>
+		/// <param name="request"></param>
+		/// <returns></returns>
+		[HttpPost]
+		[Route("/api/v3/testdata/tags")]
+		[ProducesResponseType(typeof(List<GetTestTagResponse>), 200)]
+		public async Task<ActionResult<List<GetTestTagResponse>>> GetTestTagRefV2Async([FromBody] GetTestTagRequest request)
+		{
+			TestTagId[]? tagIds = request.TagIds?.Select(x => TestTagId.Parse(x)).ToHashSet().ToArray();
+
+			IReadOnlyList<ITestTagRef> testValues = await _testDataService.FindTestTagsV2Async(tagIds, request.TagNames?.ToArray());
+
+			return testValues.Select(x => new GetTestTagResponse
+			{
+				Id = x.Id.ToString(),
+				Name = x.Name
+			}).ToList();
+		}
+
+		/// <summary>
+		/// Get test phases from provide test keys and refs
+		/// </summary>
+		/// <param name="request"></param>
+		/// <returns></returns>
+		[HttpPost]
+		[Route("/api/v3/testdata/phases")]
+		[ProducesResponseType(typeof(List<GetTestPhaseResponse>), 200)]
+		public async Task<ActionResult<List<GetTestPhaseResponse>>> GetTestPhaseRefFromTestV2Async([FromBody] GetTestPhasesRequest request)
+		{
+			string[]? testKeys = request.TestKeys?.ToHashSet().ToArray();
+			TestId[]? testIds = request.TestIds?.Select(x => TestId.Parse(x)).ToHashSet().ToArray();
+			IReadOnlyList<ITestNameRef> testNameRefs = await _testDataService.FindTestNameRefsV2Async(testIds, testKeys);
+			IReadOnlyList<ITestPhaseRef> phaseValues = await _testDataService.FindTestPhasesV2Async(testNameRefs.Select(t => t.Id).ToArray());
+
+			List<GetTestPhaseResponse> response = new List<GetTestPhaseResponse>();
+
+			foreach(ITestNameRef test in testNameRefs)
+			{
+				response.Add(new GetTestPhaseResponse
+				{
+					TestId = test.Id.ToString(),
+					TestKey = test.Key.ToString(),
+					TestName = test.Name.ToString(),
+					Phases = phaseValues.Where(p => p.TestNameRef == test.Id)
+						.Select(p => new GetPhaseResponse
+							{
+								Id = p.Id.ToString(),
+								Key = p.Key,
+								Name = p.Name
+							}).ToList()
+				});
+			}
+
+			return response;
+		}
+
+		/// <summary>
+		/// Gets test data sessions 
+		/// </summary>
+		/// <param name="streamIds"></param>
+		/// <param name="testIds"></param>
+		/// <param name="metaIds"></param>
+		/// <param name="minCreateTime"></param>
+		/// <param name="maxCreateTime"></param>
+		/// <param name="minChange"></param>
+		/// <param name="maxChange"></param>
+		/// <returns></returns>
+		[HttpGet]
+		[Route("/api/v3/testdata/session/tests")]
+		[ProducesResponseType(typeof(List<GetTestSessionResponse>), 200)]
+		public async Task<ActionResult<List<GetTestSessionResponse>>> GetTestSessionsV2Async(
+			[FromQuery(Name = "Id")] string[] streamIds,
+			[FromQuery(Name = "Tid")] string[]? testIds = null,
+			[FromQuery(Name = "Mid")] string[]? metaIds = null,
+			[FromQuery] DateTimeOffset? minCreateTime = null,
+			[FromQuery] DateTimeOffset? maxCreateTime = null,
+			[FromQuery] CommitId? minChange = null,
+			[FromQuery] CommitId? maxChange = null)
+		{
+			StreamId[] streamIdValues = Array.ConvertAll(streamIds, x => new StreamId(x));
+
+			List<StreamId> queryStreams = new List<StreamId>();
+			List<GetTestSessionResponse> responses = new List<GetTestSessionResponse>();
+
+			// authorize streams
+			foreach (StreamId streamId in streamIdValues)
+			{
+				if (_buildConfig.Value.TryGetStream(streamId, out StreamConfig? streamConfig) && streamConfig.Authorize(JobAclAction.ViewJob, User))
+				{
+					queryStreams.Add(streamId);
+				}
+			}
+
+			if (queryStreams.Count == 0)
+			{
+				return responses;
+			}
+
+			TestMetaId[]? queryMetaIds = metaIds?.ConvertAll(x => TestMetaId.Parse(x)).ToArray();
+			TestId[]? queryTestIds = testIds?.ConvertAll(x => TestId.Parse(x)).ToArray();
+
+			IReadOnlyList<ITestSession> sessions = await _testDataService.FindTestSessionsV2Async(queryStreams.ToArray(), queryTestIds, queryMetaIds, minCreateTime?.UtcDateTime, maxCreateTime?.UtcDateTime, minChange, maxChange);
+			foreach (ITestSession testData in sessions)
+			{
+				responses.Add(new GetTestSessionResponse
+				{
+					Id = testData.Id.ToString(),
+					StreamId = testData.StreamId.ToString(),
+					JobId = testData.JobId.ToString(),
+					StepId = testData.StepId.ToString(),
+					Duration = testData.Duration,
+					StartDateTime = testData.StartDateTime,
+					BuildCommitId = testData.BuildCommitId,
+					MetadataId = testData.Metadata.ToString(),
+					NameRef = testData.NameRef.ToString(),
+					TestDataId = testData.TestDataId.ToString(),
+					Outcome = testData.Outcome,
+					PhasesTotalCount = testData.PhasesTotalCount,
+					PhasesSucceededCount = testData.PhasesSucceededCount,
+					PhasesUndefinedCount = testData.PhasesUndefinedCount,
+					PhasesFailedCount = testData.PhasesFailedCount
+				});
+			}
+
+			return responses;
+		}
+
+		/// <summary>
+		/// Gets test data phase sessions 
+		/// </summary>
+		/// <param name="streamIds"></param>
+		/// <param name="phaseIds"></param>
+		/// <param name="minCreateTime"></param>
+		/// <param name="maxCreateTime"></param>
+		/// <param name="minChange"></param>
+		/// <param name="maxChange"></param>
+		/// <returns></returns>
+		[HttpGet]
+		[Route("/api/v3/testdata/session/phases")]
+		[ProducesResponseType(typeof(List<GetTestPhaseSessionResponse>), 200)]
+		public async Task<ActionResult<List<GetTestPhaseSessionResponse>>> GetTestPhaseSessionsV2Async(
+			[FromQuery(Name = "Id")] string[] streamIds,
+			[FromQuery(Name = "Pid")] string[] phaseIds,
+			[FromQuery] DateTimeOffset? minCreateTime = null,
+			[FromQuery] DateTimeOffset? maxCreateTime = null,
+			[FromQuery] CommitId? minChange = null,
+			[FromQuery] CommitId? maxChange = null)
+		{
+			StreamId[] streamIdValues = Array.ConvertAll(streamIds, x => new StreamId(x));
+
+			List<StreamId> queryStreams = new List<StreamId>();
+			List<GetTestPhaseSessionResponse> responses = new List<GetTestPhaseSessionResponse>();
+
+			// authorize streams
+			foreach (StreamId streamId in streamIdValues)
+			{
+				if (_buildConfig.Value.TryGetStream(streamId, out StreamConfig? streamConfig) && streamConfig.Authorize(JobAclAction.ViewJob, User))
+				{
+					queryStreams.Add(streamId);
+				}
+			}
+
+			if (queryStreams.Count == 0)
+			{
+				return responses;
+			}
+
+			TestPhaseId[] queryPhaseIds = phaseIds.ConvertAll(x => TestPhaseId.Parse(x)).ToArray();
+
+			IReadOnlyList<ITestPhaseSession> sessions = await _testDataService.FindTestPhaseSessionsV2Async(queryStreams.ToArray(), queryPhaseIds, minCreateTime?.UtcDateTime, maxCreateTime?.UtcDateTime, minChange, maxChange);
+			foreach (ITestPhaseSession testData in sessions)
+			{
+				responses.Add(new GetTestPhaseSessionResponse
+				{
+					Id = testData.Id.ToString(),
+					StreamId = testData.StreamId.ToString(),
+					JobId = testData.JobId.ToString(),
+					StepId = testData.StepId.ToString(),
+					Duration = testData.Duration,
+					StartDateTime = testData.StartDateTime,
+					BuildCommitId = testData.BuildCommitId,
+					MetadataId = testData.Metadata.ToString(),
+					PhaseRef = testData.PhaseRef.ToString(),
+					SessionsId = testData.SessionId.ToString(),
+					Outcome = testData.Outcome,
+					HasWarning = testData.HasWarning,
+					ErrorFingerprint = testData.ErrorFingerprint,
+					EventStreamPath = testData.EventStreamPath
+				});
+			}
+
+			return responses;
+		}
 	}
 }
diff --git a/Engine/Source/Programs/Horde/Plugins/Build/HordeServer.Build/Jobs/TestData/TestDataService.cs b/Engine/Source/Programs/Horde/Plugins/Build/HordeServer.Build/Jobs/TestData/TestDataService.cs
--- a/Engine/Source/Programs/Horde/Plugins/Build/HordeServer.Build/Jobs/TestData/TestDataService.cs
+++ b/Engine/Source/Programs/Horde/Plugins/Build/HordeServer.Build/Jobs/TestData/TestDataService.cs
@@ -6,26 +6,29 @@ using EpicGames.Horde.Streams;
 using Microsoft.Extensions.Hosting;
 using Microsoft.Extensions.Logging;
 using Microsoft.Extensions.Options;
+using MongoDB.Bson;
 
 namespace HordeServer.Jobs.TestData
 {
 	/// <summary>
-	/// Device management service
+	/// TestData management service
 	/// </summary>
 	public sealed class TestDataService : IHostedService, IAsyncDisposable
 	{
 
 		readonly ITestDataCollection _testData;
+		readonly ITestDataCollectionV2 _testDataV2;
 		readonly IOptionsMonitor<BuildServerConfig> _settings;
 		readonly ITicker _ticker;
 		readonly ILogger<TestDataService> _logger;
 
 		/// <summary>
-		/// Device service constructor
+		/// TestData service constructor
 		/// </summary>
-		public TestDataService(ITestDataCollection testData, IOptionsMonitor<BuildServerConfig> settings, IClock clock, ILogger<TestDataService> logger)
+		public TestDataService(ITestDataCollection testData, ITestDataCollectionV2 testDataV2, IOptionsMonitor<BuildServerConfig> settings, IClock clock, ILogger<TestDataService> logger)
 		{
 			_testData = testData;
+			_testDataV2 = testDataV2;
 			_settings = settings;
 			_logger = logger;
 			_ticker = clock.AddSharedTicker<TestDataService>(TimeSpan.FromMinutes(10.0), TickAsync, logger);
@@ -73,6 +76,14 @@ namespace HordeServer.Jobs.TestData
 				{
 					_logger.LogError(ex, "Exception while ticking test data collection: {Message}", ex.Message);
 				}
+				try
+				{
+					await _testDataV2.UpdateAsync(_settings.CurrentValue.TestDataRetainMonths, stoppingToken);
+				}
+				catch (Exception ex)
+				{
+					_logger.LogError(ex, "Exception while ticking test data collection v2: {Message}", ex.Message);
+				}
 			}
 		}
 
@@ -162,5 +173,108 @@ namespace HordeServer.Jobs.TestData
 		{
 			return await _testData.FindTestDetailsAsync(ids, cancellationToken);
 		}
+
+		/// <summary>
+		/// Gets a testdata by ID
+		/// </summary>
+		/// <param name="id">Unique id of the testdata</param>
+		/// <param name="cancellationToken">Cancellation token for the operation</param>
+		/// <returns>The testdata document</returns>
+		public async Task<ITestData?> GetTestDataV2Async(ObjectId id, CancellationToken cancellationToken)
+		{
+			return await _testDataV2.GetAsync(id, cancellationToken);
+		}
+
+		/// <summary>
+		/// Gets the test sessions running in provided streams
+		/// </summary>
+		/// <param name="streamIds"></param>
+		/// <param name="cancellationToken">Cancellation token for the operation</param>
+		/// <returns></returns>
+		public async Task<IReadOnlyList<ITestSessionStream>> FindTestSessionStreamsV2Async(StreamId[] streamIds, CancellationToken cancellationToken = default)
+		{
+			return await _testDataV2.FindTestSessionStreamsAsync(streamIds, cancellationToken);
+		}
+
+		/// <summary>
+		/// Find test sessions
+		/// </summary>
+		/// <param name="streamIds"></param>
+		/// <param name="testIds"></param>
+		/// <param name="metaIds"></param>
+		/// <param name="minCreateTime"></param>
+		/// <param name="maxCreateTime"></param>
+		/// <param name="minCommitId"></param>
+		/// <param name="maxCommitId"></param>
+		/// <param name="cancellationToken">Cancellation token for the operation</param>
+		/// <returns></returns>
+		public async Task<IReadOnlyList<ITestSession>> FindTestSessionsV2Async(StreamId[] streamIds, TestId[]? testIds = null, TestMetaId[]? metaIds = null, DateTime? minCreateTime = null, DateTime? maxCreateTime = null, CommitId? minCommitId = null, CommitId? maxCommitId = null, CancellationToken cancellationToken = default)
+		{
+			return await _testDataV2.FindTestSessionsAsync(streamIds, testIds, metaIds, minCreateTime, maxCreateTime, minCommitId, maxCommitId, cancellationToken);
+		}
+
+		/// <summary>
+		/// Find test name references
+		/// </summary>
+		/// <param name="testIds"></param>
+		/// <param name="keys"></param>
+		/// <param name="cancellationToken">Cancellation token for the operation</param>
+		/// <returns></returns>
+		public async Task<IReadOnlyList<ITestNameRef>> FindTestNameRefsV2Async(TestId[]? testIds = null, string[]? keys = null, CancellationToken cancellationToken = default)
+		{
+			return await _testDataV2.FindTestNameRefsAsync(testIds, keys, cancellationToken);
+		}
+
+		/// <summary>
+		/// Find test metadata set
+		/// </summary>
+		/// <param name="metaIds"></param>
+		/// <param name="keyValues"></param>
+		/// <param name="cancellationToken">Cancellation token for the operation</param>
+		/// <returns></returns>
+		public async Task<IReadOnlyList<ITestMetaRef>> FindTestMetaV2Async(TestMetaId[]? metaIds = null, IReadOnlyDictionary<string, string>? keyValues = null, CancellationToken cancellationToken = default)
+		{
+			return await _testDataV2.FindTestMetaAsync(metaIds, keyValues, cancellationToken);
+		}
+
+		/// <summary>
+		/// Find test phase references associated with a set of tests
+		/// </summary>
+		/// <param name="testIds"></param>
+		/// <param name="keys"></param>
+		/// <param name="cancellationToken">Cancellation token for the operation</param>
+		/// <returns></returns>
+		public async Task<IReadOnlyList<ITestPhaseRef>> FindTestPhasesV2Async(TestId[] testIds, string[]? keys = null, CancellationToken cancellationToken = default)
+		{
+			return await _testDataV2.FindTestPhasesAsync(testIds, keys, cancellationToken);
+		}
+
+		/// <summary>
+		/// Find test tag references
+		/// </summary>
+		/// <param name="tagIds"></param>
+		/// <param name="names"></param>
+		/// <param name="cancellationToken">Cancellation token for the operation</param>
+		/// <returns></returns>
+		public async Task<IReadOnlyList<ITestTagRef>> FindTestTagsV2Async(TestTagId[]? tagIds, string[]? names = null, CancellationToken cancellationToken = default)
+		{
+			return await _testDataV2.FindTestTagsAsync(tagIds, names, cancellationToken);
+		}
+
+		/// <summary>
+		/// Find test phase sessions from a set of test phases
+		/// </summary>
+		/// <param name="streamIds"></param>
+		/// <param name="phaseIds"></param>
+		/// <param name="minCreateTime"></param>
+		/// <param name="maxCreateTime"></param>
+		/// <param name="minCommitId"></param>
+		/// <param name="maxCommitId"></param>
+		/// <param name="cancellationToken">Cancellation token for the operation</param>
+		/// <returns></returns>
+		public async Task<IReadOnlyList<ITestPhaseSession>> FindTestPhaseSessionsV2Async(StreamId[] streamIds, TestPhaseId[] phaseIds, DateTime? minCreateTime = null, DateTime? maxCreateTime = null, CommitId? minCommitId = null, CommitId? maxCommitId = null, CancellationToken cancellationToken = default)
+		{
+			return await _testDataV2.FindTestPhaseSessionsAsync(streamIds, phaseIds, minCreateTime, maxCreateTime, minCommitId, maxCommitId, cancellationToken);
+		}
 	}
 }
\ No newline at end of file
diff --git a/Engine/Source/Programs/Horde/Plugins/Build/HordeServer.Build/Jobs/TestData/TestPhaseId.cs b/Engine/Source/Programs/Horde/Plugins/Build/HordeServer.Build/Jobs/TestData/TestPhaseId.cs
new file mode 100644
--- /dev/null
+++ b/Engine/Source/Programs/Horde/Plugins/Build/HordeServer.Build/Jobs/TestData/TestPhaseId.cs
@@ -0,0 +1,60 @@
+// Copyright Epic Games, Inc. All Rights Reserved.
+
+using System.ComponentModel;
+using EpicGames.Core;
+using HordeServer.Utilities;
+using MongoDB.Bson;
+
+namespace HordeServer.Jobs.TestData
+{
+	/// <summary>
+	/// Identifier for a test phase
+	/// </summary>
+	/// <param name="Id">Id to construct from</param>
+	[JsonSchemaString]
+	[TypeConverter(typeof(ObjectIdTypeConverter<TestPhaseId, TestPhaseIdConverter>))]
+	[ObjectIdConverter(typeof(TestPhaseIdConverter))]
+	public record struct TestPhaseId(ObjectId Id)
+	{
+		/// <summary>
+		/// Default empty value for phase id
+		/// </summary>
+		public static TestPhaseId Empty { get; } = default;
+
+		/// <summary>
+		/// Creates a new <see cref="TestPhaseId"/>
+		/// </summary>
+		public static TestPhaseId GenerateNewId() => new TestPhaseId(ObjectId.GenerateNewId());
+
+		/// <inheritdoc cref="ObjectId.Parse(System.String)"/>
+		public static TestPhaseId Parse(string text) => new TestPhaseId(ObjectId.Parse(text));
+
+		/// <inheritdoc cref="ObjectId.TryParse(System.String, out ObjectId)"/>
+		public static bool TryParse(string text, out TestPhaseId id)
+		{
+			if (ObjectId.TryParse(text, out ObjectId objectId))
+			{
+				id = new TestPhaseId(objectId);
+				return true;
+			}
+
+			id = default;
+			return false;
+		}
+
+		/// <inheritdoc/>
+		public override string ToString() => Id.ToString();
+	}
+
+	/// <summary>
+	/// Converter to and from <see cref="ObjectId"/> instances.
+	/// </summary>
+	class TestPhaseIdConverter : ObjectIdConverter<TestPhaseId>
+	{
+		/// <inheritdoc/>
+		public override TestPhaseId FromObjectId(ObjectId id) => new TestPhaseId(id);
+
+		/// <inheritdoc/>
+		public override ObjectId ToObjectId(TestPhaseId value) => value.Id;
+	}
+}
diff --git a/Engine/Source/Programs/Horde/Plugins/Build/HordeServer.Build/Jobs/TestData/TestPhaseSessionId.cs b/Engine/Source/Programs/Horde/Plugins/Build/HordeServer.Build/Jobs/TestData/TestPhaseSessionId.cs
new file mode 100644
--- /dev/null
+++ b/Engine/Source/Programs/Horde/Plugins/Build/HordeServer.Build/Jobs/TestData/TestPhaseSessionId.cs
@@ -0,0 +1,65 @@
+// Copyright Epic Games, Inc. All Rights Reserved.
+
+using System.ComponentModel;
+using EpicGames.Core;
+using HordeServer.Utilities;
+using MongoDB.Bson;
+
+namespace HordeServer.Jobs.TestData
+{
+	/// <summary>
+	/// Identifier for a test phase session
+	/// </summary>
+	/// <param name="Id">Id to construct from</param>
+	[JsonSchemaString]
+	[TypeConverter(typeof(ObjectIdTypeConverter<TestPhaseSessionId, TestPhaseSessionIdConverter>))]
+	[ObjectIdConverter(typeof(TestPhaseSessionIdConverter))]
+	public record struct TestPhaseSessionId(ObjectId Id)
+	{
+		/// <summary>
+		/// Default empty value for session id
+		/// </summary>
+		public static TestPhaseSessionId Empty { get; } = default;
+
+		/// <summary>
+		/// Creates a new <see cref="TestPhaseSessionId"/>
+		/// </summary>
+		public static TestPhaseSessionId GenerateNewId() => new TestPhaseSessionId(ObjectId.GenerateNewId());
+
+		/// <summary>
+		/// Creates a new <see cref="TestPhaseSessionId"/>
+		/// </summary>
+		public static TestPhaseSessionId GenerateNewId(DateTime time) => new TestPhaseSessionId(ObjectId.GenerateNewId(time));
+
+		/// <inheritdoc cref="ObjectId.Parse(System.String)"/>
+		public static TestPhaseSessionId Parse(string text) => new TestPhaseSessionId(ObjectId.Parse(text));
+
+		/// <inheritdoc cref="ObjectId.TryParse(System.String, out ObjectId)"/>
+		public static bool TryParse(string text, out TestPhaseSessionId id)
+		{
+			if (ObjectId.TryParse(text, out ObjectId objectId))
+			{
+				id = new TestPhaseSessionId(objectId);
+				return true;
+			}
+
+			id = default;
+			return false;
+		}
+
+		/// <inheritdoc/>
+		public override string ToString() => Id.ToString();
+	}
+
+	/// <summary>
+	/// Converter to and from <see cref="ObjectId"/> instances.
+	/// </summary>
+	class TestPhaseSessionIdConverter : ObjectIdConverter<TestPhaseSessionId>
+	{
+		/// <inheritdoc/>
+		public override TestPhaseSessionId FromObjectId(ObjectId id) => new TestPhaseSessionId(id);
+
+		/// <inheritdoc/>
+		public override ObjectId ToObjectId(TestPhaseSessionId value) => value.Id;
+	}
+}
diff --git a/Engine/Source/Programs/Horde/Plugins/Build/HordeServer.Build/Jobs/TestData/TestSessionId.cs b/Engine/Source/Programs/Horde/Plugins/Build/HordeServer.Build/Jobs/TestData/TestSessionId.cs
new file mode 100644
--- /dev/null
+++ b/Engine/Source/Programs/Horde/Plugins/Build/HordeServer.Build/Jobs/TestData/TestSessionId.cs
@@ -0,0 +1,65 @@
+// Copyright Epic Games, Inc. All Rights Reserved.
+
+using System.ComponentModel;
+using EpicGames.Core;
+using HordeServer.Utilities;
+using MongoDB.Bson;
+
+namespace HordeServer.Jobs.TestData
+{
+	/// <summary>
+	/// Identifier for a test session
+	/// </summary>
+	/// <param name="Id">Id to construct from</param>
+	[JsonSchemaString]
+	[TypeConverter(typeof(ObjectIdTypeConverter<TestSessionId, TestSessionIdConverter>))]
+	[ObjectIdConverter(typeof(TestSessionIdConverter))]
+	public record struct TestSessionId(ObjectId Id)
+	{
+		/// <summary>
+		/// Default empty value for session id
+		/// </summary>
+		public static TestSessionId Empty { get; } = default;
+
+		/// <summary>
+		/// Creates a new <see cref="TestSessionId"/>
+		/// </summary>
+		public static TestSessionId GenerateNewId() => new TestSessionId(ObjectId.GenerateNewId());
+
+		/// <summary>
+		/// Creates a new <see cref="TestSessionId"/>
+		/// </summary>
+		public static TestSessionId GenerateNewId(DateTime time) => new TestSessionId(ObjectId.GenerateNewId(time));
+
+		/// <inheritdoc cref="ObjectId.Parse(System.String)"/>
+		public static TestSessionId Parse(string text) => new TestSessionId(ObjectId.Parse(text));
+
+		/// <inheritdoc cref="ObjectId.TryParse(System.String, out ObjectId)"/>
+		public static bool TryParse(string text, out TestSessionId id)
+		{
+			if (ObjectId.TryParse(text, out ObjectId objectId))
+			{
+				id = new TestSessionId(objectId);
+				return true;
+			}
+
+			id = default;
+			return false;
+		}
+
+		/// <inheritdoc/>
+		public override string ToString() => Id.ToString();
+	}
+
+	/// <summary>
+	/// Converter to and from <see cref="ObjectId"/> instances.
+	/// </summary>
+	class TestSessionIdConverter : ObjectIdConverter<TestSessionId>
+	{
+		/// <inheritdoc/>
+		public override TestSessionId FromObjectId(ObjectId id) => new TestSessionId(id);
+
+		/// <inheritdoc/>
+		public override ObjectId ToObjectId(TestSessionId value) => value.Id;
+	}
+}
diff --git a/Engine/Source/Programs/Horde/Plugins/Build/HordeServer.Build/Jobs/TestData/TestTagId.cs b/Engine/Source/Programs/Horde/Plugins/Build/HordeServer.Build/Jobs/TestData/TestTagId.cs
new file mode 100644
--- /dev/null
+++ b/Engine/Source/Programs/Horde/Plugins/Build/HordeServer.Build/Jobs/TestData/TestTagId.cs
@@ -0,0 +1,60 @@
+// Copyright Epic Games, Inc. All Rights Reserved.
+
+using System.ComponentModel;
+using EpicGames.Core;
+using HordeServer.Utilities;
+using MongoDB.Bson;
+
+namespace HordeServer.Jobs.TestData
+{
+	/// <summary>
+	/// Identifier for a tag
+	/// </summary>
+	/// <param name="Id">Id to construct from</param>
+	[JsonSchemaString]
+	[TypeConverter(typeof(ObjectIdTypeConverter<TestTagId, TestTagIdConverter>))]
+	[ObjectIdConverter(typeof(TestTagIdConverter))]
+	public record struct TestTagId(ObjectId Id)
+	{
+		/// <summary>
+		/// Default empty value for session id
+		/// </summary>
+		public static TestTagId Empty { get; } = default;
+
+		/// <summary>
+		/// Creates a new <see cref="TestTagId"/>
+		/// </summary>
+		public static TestTagId GenerateNewId() => new TestTagId(ObjectId.GenerateNewId());
+
+		/// <inheritdoc cref="ObjectId.Parse(System.String)"/>
+		public static TestTagId Parse(string text) => new TestTagId(ObjectId.Parse(text));
+
+		/// <inheritdoc cref="ObjectId.TryParse(System.String, out ObjectId)"/>
+		public static bool TryParse(string text, out TestTagId id)
+		{
+			if (ObjectId.TryParse(text, out ObjectId objectId))
+			{
+				id = new TestTagId(objectId);
+				return true;
+			}
+
+			id = default;
+			return false;
+		}
+
+		/// <inheritdoc/>
+		public override string ToString() => Id.ToString();
+	}
+
+	/// <summary>
+	/// Converter to and from <see cref="ObjectId"/> instances.
+	/// </summary>
+	class TestTagIdConverter : ObjectIdConverter<TestTagId>
+	{
+		/// <inheritdoc/>
+		public override TestTagId FromObjectId(ObjectId id) => new TestTagId(id);
+
+		/// <inheritdoc/>
+		public override ObjectId ToObjectId(TestTagId value) => value.Id;
+	}
+}
diff --git a/Engine/Source/Programs/Shared/EpicGames.Horde/Jobs/TestData/TestDataMessages.cs b/Engine/Source/Programs/Shared/EpicGames.Horde/Jobs/TestData/TestDataMessages.cs
--- a/Engine/Source/Programs/Shared/EpicGames.Horde/Jobs/TestData/TestDataMessages.cs
+++ b/Engine/Source/Programs/Shared/EpicGames.Horde/Jobs/TestData/TestDataMessages.cs
@@ -32,6 +32,37 @@ namespace EpicGames.Horde.Jobs.TestData
 		Unspecified
 	}
 
+	/// <summary>
+	/// Test phase outcome
+	/// </summary>
+	public enum TestPhaseOutcome
+	{
+		/// <summary>
+		/// The phase has failed
+		/// </summary>
+		Failed,
+		/// <summary>
+		/// The phase was successful
+		/// </summary>
+		Success,
+		/// <summary>
+		/// The phase was not run
+		/// </summary>
+		NotRun,
+		/// <summary>
+		/// The phase was interrupted in its execution
+		/// </summary>
+		Interrupted,
+		/// <summary>
+		/// The phase was skipped
+		/// </summary>
+		Skipped,
+		/// <summary>
+		/// The phase outcome is unknown
+		/// </summary>
+		Unknown
+	}
+
 	/// <summary>
 	/// Response object describing test data to store
 	/// </summary>
@@ -226,6 +257,54 @@ namespace EpicGames.Horde.Jobs.TestData
 		public List<string> TestIds { get; set; } = new List<string>();
 	}
 
+	/// <summary>
+	/// Get test metadata request
+	/// </summary>
+	public class GetMetadataRequest
+	{
+		/// <summary>
+		/// Test metadata ids to get
+		/// </summary>
+		public List<string>? MetadataIds { get; set; }
+
+		/// <summary>
+		/// Matching key/value pair Test metadata to get
+		/// </summary>
+		public Dictionary<string, string>? KeyValuePairs { get; set; }
+	}
+
+	/// <summary>
+	/// Get test tag request
+	/// </summary>
+	public class GetTestTagRequest
+	{
+		/// <summary>
+		/// Test tag ids to get
+		/// </summary>
+		public List<string>? TagIds { get; set; }
+
+		/// <summary>
+		/// Test tag name to get
+		/// </summary>
+		public List<string>? TagNames { get; set; }
+	}
+
+	/// <summary>
+	/// Get test phases from test keys and ids request
+	/// </summary>
+	public class GetTestPhasesRequest
+	{
+		/// <summary>
+		/// Test keys to get
+		/// </summary>
+		public List<string>? TestKeys { get; set; }
+
+		/// <summary>
+		/// Test ids to get
+		/// </summary>
+		public List<string>? TestIds { get; set; }
+	}
+
 	/// <summary>
 	/// A test suite that runs in a stream, contain subtests
 	/// </summary>
@@ -424,4 +503,287 @@ namespace EpicGames.Horde.Jobs.TestData
 		/// </summary>
 		public int? SuiteSuccessCount { get; set; }
 	}
+
+	/// <summary>
+	/// A test name/key ref
+	/// </summary>
+	public class GetTestNameResponse
+	{
+		/// <summary>
+		/// The id of the test
+		/// </summary>
+		public string Id { get; set; } = String.Empty;
+
+		/// <summary>
+		/// The name of the test 
+		/// </summary>
+		public string Name { get; set; } = String.Empty;
+
+		/// <summary>
+		/// The key of the test
+		/// </summary>
+		public string Key { get; set; } = String.Empty;
+	}
+
+	/// <summary>
+	/// A set of environment conditions in which a test was run
+	/// </summary>
+	public class GetTestMetadataResponse
+	{
+		/// <summary>
+		/// Metadata unique id for environment 
+		/// </summary>
+		public string Id { get; set; } = String.Empty;
+
+		/// <summary>
+		/// The set of key/value pairs this environment represent
+		/// </summary>
+		public Dictionary<string, string> KeyValuePairs { get; set; } = new Dictionary<string, string>();
+	}
+
+	/// <summary>
+	/// A test tag ref
+	/// </summary>
+	public class GetTestTagResponse
+	{
+		/// <summary>
+		/// Tag unique id
+		/// </summary>
+		public string Id { get; set; } = String.Empty;
+
+		/// <summary>
+		/// The name of the tag
+		/// </summary>
+		public string Name { get; set; } = String.Empty;
+	}
+
+	/// <summary>
+	/// A test phase from a test
+	/// </summary>
+	public class GetPhaseResponse
+	{
+		/// <summary>
+		/// The id of the test phase
+		/// </summary>
+		public string Id { get; set; } = String.Empty;
+
+		/// <summary>
+		/// The name of the test phase
+		/// </summary>
+		public string Name { get; set; } = String.Empty;
+
+		/// <summary>
+		/// The key of the test phase
+		/// </summary>
+		public string Key { get; set; } = String.Empty;
+	}
+
+	/// <summary>
+	/// All test phases from a test
+	/// </summary>
+	public class GetTestPhaseResponse
+	{
+		/// <summary>
+		/// The id of the test phase
+		/// </summary>
+		public string TestId { get; set; } = String.Empty;
+
+		/// <summary>
+		/// The name of the test phase
+		/// </summary>
+		public string TestName { get; set; } = String.Empty;
+
+		/// <summary>
+		/// The key of the test phase
+		/// </summary>
+		public string TestKey { get; set; } = String.Empty;
+
+		/// <summary>
+		/// The phases associated with the test
+		/// </summary>
+		public List<GetPhaseResponse> Phases { get; set; } = new List<GetPhaseResponse>();
+	}
+
+	/// <summary>
+	/// Response object describing test results
+	/// </summary>
+	public class GetTestSessionStreamResponse
+	{
+		/// <summary>
+		/// The stream id
+		/// </summary>
+		public string StreamId { get; set; } = String.Empty;
+
+		/// <summary>
+		/// Individual tests which run in the stream
+		/// </summary>
+		public List<GetTestNameResponse> Tests { get; set; } = new List<GetTestNameResponse>();
+
+		/// <summary>
+		/// Known test metadata
+		/// </summary>
+		public List<GetTestMetadataResponse> TestMetadata { get; set; } = new List<GetTestMetadataResponse>();
+
+		/// <summary>
+		/// Known test tags
+		/// </summary>
+		public List<GetTestTagResponse> TestTags { get; set; } = new List<GetTestTagResponse>();
+	}
+
+	/// <summary>
+	/// Response object for a test session
+	/// </summary>
+	public class GetTestSessionResponse
+	{
+		/// <summary>
+		/// The test session id
+		/// </summary>
+		public string Id { get; set; } = String.Empty;
+
+		/// <summary>
+		/// The stream id associated with the test session
+		/// </summary>
+		public string StreamId { get; set; } = String.Empty;
+
+		/// <summary>
+		/// Test metadata id associated with the test session
+		/// </summary>
+		public string MetadataId { get; set; } = String.Empty;
+
+		/// <summary>
+		/// The test data id associated with the test session
+		/// </summary>
+		public string TestDataId { get; set; } = String.Empty;
+
+		/// <summary>
+		/// The commit id associated with the test session
+		/// </summary>
+		public CommitId BuildCommitId { get; set; } = CommitId.Empty;
+
+		/// <summary>
+		/// The duration of the test session
+		/// </summary>
+		public TimeSpan Duration { get; set; }
+
+		/// <summary>
+		/// The date time at which the test session started
+		/// </summary>
+		public DateTime StartDateTime { get; set; }
+
+		/// <summary>
+		/// The test name reference associated with the test session
+		/// </summary>
+		public string NameRef { get; set; } = String.Empty;
+
+		/// <summary>
+		/// The outcome of the test session
+		/// </summary>
+		public TestOutcome Outcome { get; set; }
+
+		/// <summary>
+		/// The total count of phases
+		/// </summary>
+		public int PhasesTotalCount { get; set; }
+
+		/// <summary>
+		/// The count of succeeded phases
+		/// </summary>
+		public int PhasesSucceededCount { get; set; }
+
+		/// <summary>
+		/// The count of undefined phases
+		/// </summary>
+		public int PhasesUndefinedCount { get; set; }
+
+		/// <summary>
+		/// The count of failed phases
+		/// </summary>
+		public int PhasesFailedCount { get; set; }
+
+		/// <summary>
+		/// The job id associated with the test session
+		/// </summary>
+		public string JobId { get; set; } = String.Empty;
+
+		/// <summary>
+		/// The step job id associated with the test session
+		/// </summary>
+		public string StepId { get; set; } = String.Empty;
+	}
+
+	/// <summary>
+	/// Response object for a test phase session
+	/// </summary>
+	public class GetTestPhaseSessionResponse
+	{
+		/// <summary>
+		/// The phase session id
+		/// </summary>
+		public string Id { get; set; } = String.Empty;
+
+		/// <summary>
+		/// The stream id associated with the phase session
+		/// </summary>
+		public string StreamId { get; set; } = String.Empty;
+
+		/// <summary>
+		/// Test metadata id associated with the phase session
+		/// </summary>
+		public string MetadataId { get; set; } = String.Empty;
+
+		/// <summary>
+		/// The test session id associated with the phase session
+		/// </summary>
+		public string SessionsId { get; set; } = String.Empty;
+
+		/// <summary>
+		/// The commit id associated with the phase session
+		/// </summary>
+		public CommitId BuildCommitId { get; set; } = CommitId.Empty;
+
+		/// <summary>
+		/// The duration of the phase session
+		/// </summary>
+		public TimeSpan Duration { get; set; }
+
+		/// <summary>
+		/// The date time at which the phase session started
+		/// </summary>
+		public DateTime StartDateTime { get; set; }
+
+		/// <summary>
+		/// The test phase reference associated with the phase session
+		/// </summary>
+		public string PhaseRef { get; set; } = String.Empty;
+
+		/// <summary>
+		/// The outcome of the phase session
+		/// </summary>
+		public TestPhaseOutcome Outcome { get; set; }
+
+		/// <summary>
+		/// The job id associated with the phase session
+		/// </summary>
+		public string JobId { get; set; } = String.Empty;
+
+		/// <summary>
+		/// The step job id associated with the phase session
+		/// </summary>
+		public string StepId { get; set; } = String.Empty;
+
+		/// <summary>
+		/// The path to the event stream
+		/// </summary>
+		public string? EventStreamPath { get; set; }
+
+		/// <summary>
+		/// Whether the test phase encountered at least one warning
+		/// </summary>
+		public bool? HasWarning { get; set; }
+
+		/// <summary>
+		/// The error fingerprint from the event stream
+		/// </summary>
+		public string? ErrorFingerprint { get; set; }
+	}
 }
