[TestDataV2] Replace Wildcard indexing by Multikey indexing in metadata

#rnx

diff --git a/Engine/Source/Programs/Horde/Plugins/Build/HordeServer.Build.Tests/Jobs/TestData/TestDataTests.cs b/Engine/Source/Programs/Horde/Plugins/Build/HordeServer.Build.Tests/Jobs/TestData/TestDataTests.cs
--- a/Engine/Source/Programs/Horde/Plugins/Build/HordeServer.Build.Tests/Jobs/TestData/TestDataTests.cs
+++ b/Engine/Source/Programs/Horde/Plugins/Build/HordeServer.Build.Tests/Jobs/TestData/TestDataTests.cs
@@ -393,15 +393,15 @@ namespace HordeServer.Tests.Jobs.TestData
 			Assert.IsTrue(testnames.Any(t => t.Key == "UE.Automation(Group:Editor) Lyra"));
 			Assert.AreEqual("Editor", testnames[0].Name);
 
-			GetMetadataRequest metaRequest = new GetMetadataRequest() { KeyValuePairs = new() { { "Platform", "Win64" } } };
+			GetMetadataRequest metaRequest = new GetMetadataRequest() { Entries = new() { { "Platform", "Win64" } } };
 			ActionResult<List<GetTestMetadataResponse>> metaResult = await TestDataController.GetTestMetadataV2Async(metaRequest);
 			Assert.IsNotNull(metaResult);
 			Assert.IsNotNull(metaResult.Value);
 			List<GetTestMetadataResponse> meta = metaResult.Value;
 			Assert.AreEqual(2, meta.Count);
-			Assert.IsTrue(meta.All(t => t.KeyValuePairs.Any(i => i.Key == "BuildTarget" && i.Value == "Editor")));
-			Assert.IsTrue(meta.Any(t => t.KeyValuePairs.Any(i => i.Key == "Project" && i.Value == "EngineTest")));
-			Assert.IsTrue(meta.Any(t => t.KeyValuePairs.Any(i => i.Key == "Project" && i.Value == "Lyra")));
+			Assert.IsTrue(meta.All(t => t.Entries.Any(i => i.Key == "BuildTarget" && i.Value == "Editor")));
+			Assert.IsTrue(meta.Any(t => t.Entries.Any(i => i.Key == "Project" && i.Value == "EngineTest")));
+			Assert.IsTrue(meta.Any(t => t.Entries.Any(i => i.Key == "Project" && i.Value == "Lyra")));
 
 			GetTestTagRequest tagRequest = new GetTestTagRequest() { TagIds = streams[0].TestTags.Select(t => t.Id).ToList() };
 			ActionResult<List<GetTestTagResponse>> tagResult = await TestDataController.GetTestTagRefV2Async(tagRequest);
diff --git a/Engine/Source/Programs/Horde/Plugins/Build/HordeServer.Build/Jobs/TestData/ITestData.cs b/Engine/Source/Programs/Horde/Plugins/Build/HordeServer.Build/Jobs/TestData/ITestData.cs
--- a/Engine/Source/Programs/Horde/Plugins/Build/HordeServer.Build/Jobs/TestData/ITestData.cs
+++ b/Engine/Source/Programs/Horde/Plugins/Build/HordeServer.Build/Jobs/TestData/ITestData.cs
@@ -310,6 +310,22 @@ namespace HordeServer.Jobs.TestData
 		string Name { get; }
 	}
 
+	/// <summary>
+	/// A metadata entry key value pair
+	/// </summary>
+	public interface ITestMetaEntry
+	{
+		/// <summary>
+		/// The meta key
+		/// </summary>
+		string Key { get; }
+
+		/// <summary>
+		/// The meta value
+		/// </summary>
+		string Value { get; }
+	}
+
 	/// <summary>
 	/// A metadata key/value pairs reference
 	/// </summary>
@@ -323,7 +339,7 @@ namespace HordeServer.Jobs.TestData
 		/// <summary>
 		/// The key/value pairs that this metadata references
 		/// </summary>
-		IReadOnlyDictionary<string, string> KeyValuePairs { get; }
+		IReadOnlyList<ITestMetaEntry> Entries { get; }
 	}
 
 	/// <summary>
diff --git a/Engine/Source/Programs/Horde/Plugins/Build/HordeServer.Build/Jobs/TestData/TestDataCollectionV2.cs b/Engine/Source/Programs/Horde/Plugins/Build/HordeServer.Build/Jobs/TestData/TestDataCollectionV2.cs
--- a/Engine/Source/Programs/Horde/Plugins/Build/HordeServer.Build/Jobs/TestData/TestDataCollectionV2.cs
+++ b/Engine/Source/Programs/Horde/Plugins/Build/HordeServer.Build/Jobs/TestData/TestDataCollectionV2.cs
@@ -30,6 +30,30 @@ namespace HordeServer.Jobs.TestData
 			DateTime LastSeenUtc { get; }
 		}
 
+		/// <summary>
+		/// Metadata key value pair
+		/// </summary>
+		class TestMetaEntry : ITestMetaEntry
+		{
+			[BsonRequired, BsonElement("k")]
+			public string Key { get; set; }
+
+			[BsonRequired, BsonElement("v")]
+			public string Value { get; set; }
+
+			private TestMetaEntry()
+			{
+				Key = String.Empty;
+				Value = String.Empty;
+			}
+
+			public TestMetaEntry(string key, string data)
+			{
+				Key = key;
+				Value = data;
+			}
+		}
+
 		/// <summary>
 		/// Test metadata document
 		/// </summary>
@@ -39,21 +63,21 @@ namespace HordeServer.Jobs.TestData
 			public TestMetaId Id { get; set; }
 
 			[BsonRequired, BsonElement("m")]
-			public Dictionary<string, string> KeyValuePairs { get; set; }
-			IReadOnlyDictionary<string, string> ITestMetaRef.KeyValuePairs => KeyValuePairs;
+			public List<TestMetaEntry> Entries { get; set; }
+			IReadOnlyList<ITestMetaEntry> ITestMetaRef.Entries => Entries;
 
 			[BsonRequired, BsonElement("s")]
 			public DateTime LastSeenUtc { get; set; }
 
 			private TestMetadataDocument()
 			{
-				KeyValuePairs = new();
+				Entries = new();
 			}
 
-			public TestMetadataDocument(IReadOnlyDictionary<string, string> keyValuePairs)
+			public TestMetadataDocument(IReadOnlyDictionary<string, string> entries)
 			{
 				Id = TestMetaId.GenerateNewId();
-				KeyValuePairs = keyValuePairs.ToDictionary();
+				Entries = entries.Select(i => new TestMetaEntry(i.Key, i.Value)).ToList();
 				LastSeenUtc = DateTime.UtcNow;
 			}
 		}
@@ -406,7 +430,7 @@ namespace HordeServer.Jobs.TestData
 			_testDataDocuments = mongoService.GetCollection<TestDataDocument>("TestDataV2", indexes);
 
 			List<MongoIndex<TestMetadataDocument>> metaIndexes = new List<MongoIndex<TestMetadataDocument>>();
-			metaIndexes.Add(keys => keys.Wildcard(x => x.KeyValuePairs));
+			metaIndexes.Add(keys => keys.Ascending("m.k").Ascending("m.v"));
 			_testMeta = mongoService.GetCollection<TestMetadataDocument>("TestDataV2.Metadata", metaIndexes);
 
 			List<MongoIndex<TestTagDocument>> tagIndexes = new List<MongoIndex<TestTagDocument>>();
@@ -514,19 +538,20 @@ namespace HordeServer.Jobs.TestData
 		public async Task<IReadOnlyList<ITestMetaRef>> FindTestMetaAsync(TestMetaId[]? metaIds = null, IReadOnlyDictionary<string, string>? keyValues = null, CancellationToken cancellationToken = default)
 		{
 			FilterDefinition<TestMetadataDocument> filter = FilterDefinition<TestMetadataDocument>.Empty;
-			FilterDefinitionBuilder<TestMetadataDocument> filterBuilder = Builders<TestMetadataDocument>.Filter;
+			FilterDefinitionBuilder<TestMetadataDocument> filterDocBuilder = Builders<TestMetadataDocument>.Filter;
+			FilterDefinitionBuilder<TestMetaEntry> metaFilterBuilder = Builders<TestMetaEntry>.Filter;
 
 			if (keyValues != null && keyValues.Count > 0)
 			{
 				foreach (KeyValuePair<string, string> item in keyValues)
 				{
-					filter &= filterBuilder.Eq($"m.{item.Key}", item.Value);
+					filter &= filterDocBuilder.ElemMatch(x => x.Entries, metaFilterBuilder.Eq(i => i.Key, item.Key) & metaFilterBuilder.Eq(i => i.Value, item.Value));
 				}
 			}
 
 			if (metaIds != null && metaIds.Length > 0)
 			{
-				filter &= filterBuilder.In(x => x.Id, metaIds);
+				filter &= filterDocBuilder.In(x => x.Id, metaIds);
 			}
 
 			return await _testMeta.Find(filter).ToListAsync(cancellationToken);
@@ -793,18 +818,20 @@ namespace HordeServer.Jobs.TestData
 
 		private async Task<ITestMetaRef> AddTestMetadataAsync(IReadOnlyDictionary<string, string> metaData, CancellationToken cancellationToken)
 		{
-			FilterDefinitionBuilder<TestMetadataDocument> filterBuilder = Builders<TestMetadataDocument>.Filter;
 			FilterDefinition<TestMetadataDocument> filter = FilterDefinition<TestMetadataDocument>.Empty;
+			FilterDefinitionBuilder<TestMetadataDocument> filterDocBuilder = Builders<TestMetadataDocument>.Filter;
+			FilterDefinitionBuilder<TestMetaEntry> metaFilterBuilder = Builders<TestMetaEntry>.Filter;
 
 			foreach (KeyValuePair<string, string> item in metaData)
 			{
-				filter &= filterBuilder.Eq($"m.{item.Key}", item.Value);
+				filter &= filterDocBuilder.ElemMatch(x => x.Entries, metaFilterBuilder.Eq(i => i.Key, item.Key) & metaFilterBuilder.Eq(i => i.Value, item.Value));
 			}
+			// Match the exact number of entries
+			filter &= filterDocBuilder.Size(x => x.Entries, metaData.Count);
 
-			List<TestMetadataDocument> metaList = await _testMeta.Find(filter).ToListAsync(cancellationToken);
-			TestMetadataDocument? meta = metaList.OrderBy(m => m.KeyValuePairs.Count).FirstOrDefault();
+			TestMetadataDocument? meta = await _testMeta.Find(filter).FirstOrDefaultAsync(cancellationToken);
 
-			if (meta == null || meta.KeyValuePairs.Count > metaData.Count)
+			if (meta == null)
 			{
 				meta = new TestMetadataDocument(metaData);
 				await _testMeta.InsertOneAsync(meta, null, cancellationToken);
diff --git a/Engine/Source/Programs/Horde/Plugins/Build/HordeServer.Build/Jobs/TestData/TestDataController.cs b/Engine/Source/Programs/Horde/Plugins/Build/HordeServer.Build/Jobs/TestData/TestDataController.cs
--- a/Engine/Source/Programs/Horde/Plugins/Build/HordeServer.Build/Jobs/TestData/TestDataController.cs
+++ b/Engine/Source/Programs/Horde/Plugins/Build/HordeServer.Build/Jobs/TestData/TestDataController.cs
@@ -507,7 +507,7 @@ namespace HordeServer.Jobs.TestData
 		[HttpGet]
 		[Route("/api/v3/testdata/streams")]
 		[ProducesResponseType(typeof(List<GetTestSessionStreamResponse>), 200)]
-		public async Task<ActionResult<List<GetTestSessionStreamResponse>>> GetTestStreamsV2Async([FromQuery(Name = "Id")] string[] streamIds)
+		public async Task<ActionResult<List<GetTestSessionStreamResponse>>> GetTestStreamsV2Async([FromQuery(Name = "id")] string[] streamIds)
 		{
 			StreamId[] streamIdValues = Array.ConvertAll(streamIds, x => new StreamId(x));
 
@@ -581,7 +581,7 @@ namespace HordeServer.Jobs.TestData
 					TestMetadata = filteredMetadata.Select(meta => new GetTestMetadataResponse
 					{
 						Id = meta.Id.ToString(),
-						KeyValuePairs = new (meta.KeyValuePairs)
+						Entries = meta.Entries.ToDictionary(i => i.Key, i => i.Value)
 					}).ToList(),
 					TestTags = filteredTags.Select(tag => new GetTestTagResponse
 					{
@@ -628,12 +628,12 @@ namespace HordeServer.Jobs.TestData
 		{
 			TestMetaId[]? metaIds = request.MetadataIds?.Select(x => TestMetaId.Parse(x)).ToHashSet().ToArray();
 
-			IReadOnlyList<ITestMetaRef> metadataValues = await _testDataService.FindTestMetaV2Async(metaIds, request.KeyValuePairs);
+			IReadOnlyList<ITestMetaRef> metadataValues = await _testDataService.FindTestMetaV2Async(metaIds, request.Entries);
 
 			return metadataValues.Select(x => new GetTestMetadataResponse
 			{
 				Id = x.Id.ToString(),
-				KeyValuePairs = new(x.KeyValuePairs)
+				Entries = x.Entries.ToDictionary(i => i.Key, i => i.Value)
 			}).ToList();
 		}
 
@@ -710,9 +710,9 @@ namespace HordeServer.Jobs.TestData
 		[Route("/api/v3/testdata/session/tests")]
 		[ProducesResponseType(typeof(List<GetTestSessionResponse>), 200)]
 		public async Task<ActionResult<List<GetTestSessionResponse>>> GetTestSessionsV2Async(
-			[FromQuery(Name = "Id")] string[] streamIds,
-			[FromQuery(Name = "Tid")] string[]? testIds = null,
-			[FromQuery(Name = "Mid")] string[]? metaIds = null,
+			[FromQuery(Name = "id")] string[] streamIds,
+			[FromQuery(Name = "tid")] string[]? testIds = null,
+			[FromQuery(Name = "mid")] string[]? metaIds = null,
 			[FromQuery] DateTimeOffset? minCreateTime = null,
 			[FromQuery] DateTimeOffset? maxCreateTime = null,
 			[FromQuery] CommitId? minChange = null,
@@ -780,8 +780,8 @@ namespace HordeServer.Jobs.TestData
 		[Route("/api/v3/testdata/session/phases")]
 		[ProducesResponseType(typeof(List<GetTestPhaseSessionResponse>), 200)]
 		public async Task<ActionResult<List<GetTestPhaseSessionResponse>>> GetTestPhaseSessionsV2Async(
-			[FromQuery(Name = "Id")] string[] streamIds,
-			[FromQuery(Name = "Pid")] string[] phaseIds,
+			[FromQuery(Name = "id")] string[] streamIds,
+			[FromQuery(Name = "pid")] string[] phaseIds,
 			[FromQuery] DateTimeOffset? minCreateTime = null,
 			[FromQuery] DateTimeOffset? maxCreateTime = null,
 			[FromQuery] CommitId? minChange = null,
diff --git a/Engine/Source/Programs/Shared/EpicGames.Horde/Jobs/TestData/TestDataMessages.cs b/Engine/Source/Programs/Shared/EpicGames.Horde/Jobs/TestData/TestDataMessages.cs
--- a/Engine/Source/Programs/Shared/EpicGames.Horde/Jobs/TestData/TestDataMessages.cs
+++ b/Engine/Source/Programs/Shared/EpicGames.Horde/Jobs/TestData/TestDataMessages.cs
@@ -175,7 +175,7 @@ namespace EpicGames.Horde.Jobs.TestData
 	}
 
 	/// <summary>
-	/// A test emvironment running in a stream
+	/// A test environment running in a stream
 	/// </summary>
 	public class GetTestMetaResponse
 	{
@@ -210,7 +210,7 @@ namespace EpicGames.Horde.Jobs.TestData
 		public string RHI { get; set; } = String.Empty;
 
 		/// <summary>
-		/// The varation of the test meta data, for example address sanitizing
+		/// The variation of the test meta data, for example address sanitizing
 		/// </summary>
 		public string Variation { get; set; } = String.Empty;
 	}
@@ -270,7 +270,7 @@ namespace EpicGames.Horde.Jobs.TestData
 		/// <summary>
 		/// Matching key/value pair Test metadata to get
 		/// </summary>
-		public Dictionary<string, string>? KeyValuePairs { get; set; }
+		public Dictionary<string, string>? Entries { get; set; }
 	}
 
 	/// <summary>
@@ -363,7 +363,7 @@ namespace EpicGames.Horde.Jobs.TestData
 		public string TestId { get; set; } = String.Empty;
 
 		/// <summary>
-		/// The ourcome of the suite test
+		/// The outcome of the suite test
 		/// </summary>
 		public TestOutcome Outcome { get; set; }
 
@@ -538,7 +538,7 @@ namespace EpicGames.Horde.Jobs.TestData
 		/// <summary>
 		/// The set of key/value pairs this environment represent
 		/// </summary>
-		public Dictionary<string, string> KeyValuePairs { get; set; } = new Dictionary<string, string>();
+		public Dictionary<string, string> Entries { get; set; } = new Dictionary<string, string>();
 	}
 
 	/// <summary>
