Horde: Dashboard side of device cleaning suipport
#skipci

diff --git a/Engine/Source/Programs/Horde/HordeDashboard/src/backend/Api.ts b/Engine/Source/Programs/Horde/HordeDashboard/src/backend/Api.ts
--- a/Engine/Source/Programs/Horde/HordeDashboard/src/backend/Api.ts
+++ b/Engine/Source/Programs/Horde/HordeDashboard/src/backend/Api.ts
@@ -3726,6 +3726,15 @@ export type GetDeviceResponse = {
 	/** Utilization info */
 	utilization?: GetDeviceUtilizationResponse[];
 
+	/** The time the cleaning started */
+	cleanStartTime?: Date | string;
+
+	/** The time the cleaning finished */
+	cleanFinishTime?: Date | string;
+
+	/** Whether the device requires cleaning */
+	cleanRequird?: boolean;
+
 }
 
 /** Device checkout request object */
diff --git a/Engine/Source/Programs/Horde/HordeDashboard/src/backend/DeviceHandler.ts b/Engine/Source/Programs/Horde/HordeDashboard/src/backend/DeviceHandler.ts
--- a/Engine/Source/Programs/Horde/HordeDashboard/src/backend/DeviceHandler.ts
+++ b/Engine/Source/Programs/Horde/HordeDashboard/src/backend/DeviceHandler.ts
@@ -151,6 +151,10 @@ export class DeviceHandler extends PollBase {
          return DeviceStatus.Maintenance;
       }
 
+      if (device.cleanStartTime && !device.cleanFinishTime) {
+         return DeviceStatus.Maintenance;
+      }
+
       if (device.problemTime) {
 
          const end = moment.utc();
diff --git a/Engine/Source/Programs/Horde/HordeDashboard/src/components/devices/DeviceView.tsx b/Engine/Source/Programs/Horde/HordeDashboard/src/components/devices/DeviceView.tsx
--- a/Engine/Source/Programs/Horde/HordeDashboard/src/components/devices/DeviceView.tsx
+++ b/Engine/Source/Programs/Horde/HordeDashboard/src/components/devices/DeviceView.tsx
@@ -383,7 +383,6 @@ const DevicePanel: React.FC = observer(() => {
       let backgroundColor = StatusColors.get(status)!
       let text = DeviceStatusLabels.get(status)!;
 
-
       if (device.checkedOutByUserId) {
          if (status === DeviceStatus.Available) {
             text = "Checked Out";
@@ -391,6 +390,10 @@ const DevicePanel: React.FC = observer(() => {
          }
       }
 
+      if (item.device.cleanStartTime && !item.device.cleanFinishTime) {
+         text = "Cleaning";
+      }
+
       // @todo: support deleting reservations
       return <Stack verticalFill={true} verticalAlign="center" horizontal>
          <PrimaryButton styles={{ root: { border: "0px", width: 84, height: 24, fontSize: 12, fontFamily: "Horde Open Sans SemiBold !important", selectors: { ".ms-Button-label": { color: color } }, backgroundColor: backgroundColor } }} text={text}
