Horde: Perforce server page visual fixes & linking from status page
Tweaked column widths, made detail column multiline, and added a tooltip.

diff --git a/Engine/Source/Programs/Horde/HordeDashboard/src/components/PerforceView.tsx b/Engine/Source/Programs/Horde/HordeDashboard/src/components/PerforceView.tsx
--- a/Engine/Source/Programs/Horde/HordeDashboard/src/components/PerforceView.tsx
+++ b/Engine/Source/Programs/Horde/HordeDashboard/src/components/PerforceView.tsx
@@ -1,6 +1,6 @@
 // Copyright Epic Games, Inc. All Rights Reserved.
 
-import { DetailsList, DetailsListLayoutMode, IColumn, SelectionMode, Stack, Text } from "@fluentui/react";
+import { DetailsList, DetailsListLayoutMode, IColumn, SelectionMode, Stack, Text, TooltipHost } from "@fluentui/react";
 import { observer } from "mobx-react-lite";
 import React, { useEffect } from "react";
 import backend from "../backend";
@@ -61,11 +61,11 @@ const ServerPanel: React.FC = observer(() => {
    if (handler.updated) { };
 
    const columns = [
-      { key: 'column_cluster', name: 'Cluster', fieldName: 'cluster', minWidth: 120, maxWidth: 120, isResizable: false },
-      { key: 'column_numleases', name: 'Leases', fieldName: 'numLeases', minWidth: 120, maxWidth: 120, isResizable: false },
+      { key: 'column_cluster', name: 'Cluster', fieldName: 'cluster', minWidth: 140, maxWidth: 140, isResizable: false },
+      { key: 'column_numleases', name: 'Leases', fieldName: 'numLeases', minWidth: 100, maxWidth: 100, isResizable: false },
       { key: 'column_status', name: 'Status', fieldName: 'status', minWidth: 120, maxWidth: 120, isResizable: false },
-      { key: 'column_server', name: 'Server', fieldName: 'serverAndPort', minWidth: 300, maxWidth: 300, isResizable: false },
-      { key: 'column_detail', name: 'Detail', fieldName: 'detail', minWidth: 200, maxWidth: 200, isResizable: false },
+      { key: 'column_server', name: 'Server', fieldName: 'serverAndPort', minWidth: 330, maxWidth: 330, isResizable: false },
+      { key: 'column_detail', name: 'Detail', fieldName: 'detail', minWidth: 200, maxWidth: 200, isResizable: false, isMultiline: true },
    ];
 
    let status = [...handler.status];
@@ -96,7 +96,13 @@ const ServerPanel: React.FC = observer(() => {
       if (!column?.fieldName) {
          return null;
       }
-      return <Text style={{ color: modeColors.text }}>{item[column?.fieldName]}</Text>
+
+      return <TooltipHost content={column?.fieldName === "detail" ? item[column?.fieldName] : ''} styles={{ root: { display: 'inline-block' } }}>
+               <Stack horizontalAlign="start" verticalAlign="center" verticalFill>
+                  <Text style={{ color: modeColors.text }}>{item[column?.fieldName]}</Text>
+               </Stack>
+         </TooltipHost>
+         
    };
 
    return <Stack>
diff --git a/Engine/Source/Programs/Horde/HordeDashboard/src/components/server/ServerStatus.tsx b/Engine/Source/Programs/Horde/HordeDashboard/src/components/server/ServerStatus.tsx
--- a/Engine/Source/Programs/Horde/HordeDashboard/src/components/server/ServerStatus.tsx
+++ b/Engine/Source/Programs/Horde/HordeDashboard/src/components/server/ServerStatus.tsx
@@ -12,6 +12,7 @@ import { getHordeStyling } from "../../styles/Styles";
 import { Breadcrumbs } from "../Breadcrumbs";
 import { TopNav } from "../TopNav";
 import dashboard, { StatusColor } from "../../backend/Dashboard";
+import { Link } from 'react-router-dom';
 
 class ServerHandler extends PollBase {
 
@@ -67,6 +68,7 @@ const ServerPanel: React.FC = observer(() => {
    }
 
    const statusColors = dashboard.getStatusColors();
+   const perforceServerUrl = "/perforce/servers"
 
    const columns = [
       {
@@ -89,7 +91,7 @@ const ServerPanel: React.FC = observer(() => {
 
             return <Stack horizontal verticalAlign="center" tokens={{ childrenGap: 6 }} verticalFill>
                <FontIcon style={{ color: color, fontSize: 13 }} iconName="Square" />
-               <Text variant="small">{status.name}</Text>
+               {status.name === "Perforce" ? <Link to={perforceServerUrl}><Text variant="small">{status.name}</Text></Link> : <Text variant="small">{status.name}</Text>}
             </Stack>
 
          }
@@ -131,7 +133,7 @@ const ServerPanel: React.FC = observer(() => {
             }
 
             return <TooltipHost content={update.message} styles={{ root: { display: 'inline-block' } }} ><Stack horizontalAlign="start" verticalAlign="center" verticalFill>
-               <Text> {update.message ?? ""}</Text>
+               {status.name === "Perforce" ? <Text>{update.message.replace(/\.$/, ",") ?? ""} see <Link to={perforceServerUrl}>server status</Link>.</Text> : <Text>{update.message ?? ""}</Text>}
             </Stack></TooltipHost>;
          }
       }
