Horde: Add issue summary badges to job label overview, and cap at 20.

diff --git a/Engine/Source/Programs/Horde/HordeDashboard/src/components/jobDetailsV2/JobDetailLabels.tsx b/Engine/Source/Programs/Horde/HordeDashboard/src/components/jobDetailsV2/JobDetailLabels.tsx
--- a/Engine/Source/Programs/Horde/HordeDashboard/src/components/jobDetailsV2/JobDetailLabels.tsx
+++ b/Engine/Source/Programs/Horde/HordeDashboard/src/components/jobDetailsV2/JobDetailLabels.tsx
@@ -1,4 +1,4 @@
-import { ConstrainMode, DefaultButton, DetailsList, DetailsListLayoutMode, IColumn, SelectionMode, Stack, Text } from "@fluentui/react";
+import { ConstrainMode, DefaultButton, DetailsList, DetailsListLayoutMode, DirectionalHint, IColumn, SelectionMode, Stack, Text, TooltipHost } from "@fluentui/react";
 import { observer } from "mobx-react-lite";
 import { useNavigate, useLocation } from "react-router-dom";
 import { GetStepResponse, JobData, JobStepOutcome, LabelOutcome, LabelState, StepData } from "../../backend/Api";
@@ -7,7 +7,7 @@ import { getLabelColor } from "../../styles/colors";
 import { useQuery } from "../JobDetailCommon";
 import { JobDataView, JobDetailsV2 } from "./JobDetailsViewCommon";
 import { getHordeStyling } from "../../styles/Styles";
-import dashboard from "../../backend/Dashboard";
+import dashboard, { StatusColor } from "../../backend/Dashboard";
 import { Link } from "react-router-dom";
 import { StepStatusIcon } from "../StatusIcon";
 
@@ -142,8 +142,12 @@ export const LabelsPanelV2: React.FC<{ jobDetails: JobDetailsV2, dataView: JobDa
 
    let steps: GetStepResponse[] = [];
 
+   steps = job.batches.map(b => b.steps).flat().filter(step => !!step.startTime && (step.outcome === JobStepOutcome.Warnings || step.outcome === JobStepOutcome.Failure));
+   
+   // sort by severity
+   steps = steps.sort((a, b) => a.outcome.localeCompare(b.outcome));
+
    if (label) {
-      steps = job.batches.map(b => b.steps).flat().filter(step => !!step.startTime && (step.outcome === JobStepOutcome.Warnings || step.outcome === JobStepOutcome.Failure));
       steps = steps.filter(s => label?.steps.indexOf(s.id) === -1 ? false : true);
    }
 
@@ -153,6 +157,18 @@ export const LabelsPanelV2: React.FC<{ jobDetails: JobDetailsV2, dataView: JobDa
          return null;
       }
 
+      const maxStepsShown = 20
+      let remainingSteps
+      let remainingErrorsCount = 0
+      let remainingWarningsCount = 0
+
+      if (steps.length > maxStepsShown) {
+         remainingSteps =  steps.slice(maxStepsShown, steps.length);
+         remainingErrorsCount = remainingSteps.filter(s => s.outcome == "Failure").length
+         remainingWarningsCount = remainingSteps.filter(s => s.outcome == "Warnings").length
+         steps = steps.slice(0, maxStepsShown);
+      }
+
       const onRenderCell = (stepItem: StepItem): JSX.Element => {
 
          const step = stepItem.step;
@@ -164,7 +180,13 @@ export const LabelsPanelV2: React.FC<{ jobDetails: JobDetailsV2, dataView: JobDa
                <Link to={stepUrl} onClick={(ev) => { ev.stopPropagation(); }}><div style={{ cursor: "pointer" }}>
                   <Stack horizontal verticalAlign="center">
                      <StepStatusIcon step={step} style={{ fontSize: 13 }} />
-                     <Text styles={{ root: { fontSize: 13, paddingRight: 4, paddingTop: 2, userSelect: "none" } }}>{`${stepItem.name}`}</Text>
+                     <TooltipHost
+                        content={stepItem.name.length > 30 ? stepItem.name : ''}
+                        directionalHint={DirectionalHint.bottomCenter}>
+                           <Text styles={{ root: { fontSize: 12, paddingRight: 4, paddingTop: 2, userSelect: "none" } }}>
+                              {stepItem.name.length <= 30 ? stepItem.name : `${stepItem.name.slice(0, 27)}...` }
+                           </Text>
+                     </TooltipHost>
                   </Stack>
                </div></Link>
             </Stack>
@@ -186,7 +208,14 @@ export const LabelsPanelV2: React.FC<{ jobDetails: JobDetailsV2, dataView: JobDa
 
       const render = items.map(s => onRenderCell(s!));
 
-      return <Stack horizontal wrap tokens={{ childrenGap: 12 }}>{render}</Stack>;
+      return <Stack>
+               <Stack horizontal wrap tokens={{ childrenGap: 12}} style={{paddingBottom: 10}}>{render}</Stack>
+               {!!remainingSteps && <Text style={{ fontSize: 13 }}>
+                  {!!remainingErrorsCount && <Text>{remainingErrorsCount} {remainingErrorsCount === 1 ? "error" : "errors"} </Text>}
+                  {!!remainingErrorsCount && !!remainingWarningsCount ? "and " : ""}
+                  {!!remainingWarningsCount && <Text>{remainingWarningsCount} {remainingWarningsCount === 1 ? "warning" : "warnings"} </Text>}
+                  not shown. Please select individual labels to view additional issues.</Text>}
+            </Stack>;
 
    }
 
@@ -269,7 +298,7 @@ export const LabelsPanelV2: React.FC<{ jobDetails: JobDetailsV2, dataView: JobDa
 
    return <Stack styles={{ root: { paddingTop: 18, paddingRight: 12 } }}>
       <Stack style={{ paddingLeft: 12 }}>
-         <Stack styles={{ root: { selectors: { '.ms-DetailsRow': { backgroundColor: dashboard.darktheme ? modeColors.content : undefined } } } }}>
+         <Stack styles={{ root: { paddingBottom: 12, selectors: { '.ms-DetailsRow': { backgroundColor: dashboard.darktheme ? modeColors.content : undefined } } } }}>
             <DetailsList
                isHeaderVisible={false}
                indentWidth={0}
@@ -283,14 +312,14 @@ export const LabelsPanelV2: React.FC<{ jobDetails: JobDetailsV2, dataView: JobDa
                onShouldVirtualize={() => { return false; }}
             />
          </Stack>
-         <Stack style={{ paddingTop: 12, paddingLeft: 12 }} horizontal>
+         {!!steps.length && <Stack style={{ paddingLeft: 12 }} horizontal>
             <Stack>
                <Text style={{ fontFamily: "Horde Open Sans SemiBold" }}>Issues</Text>
             </Stack>
             <Stack style={{paddingLeft: 100}}>
                {renderSteps()}
             </Stack>
-         </Stack>
+         </Stack>}
       </Stack>
    </Stack>;
 
