Horde: Release Notes
Fix getting the current change number of server and agent.

#rnx

diff --git a/Engine/Source/Programs/Horde/Tools/ReleaseNotes/Program.cs b/Engine/Source/Programs/Horde/Tools/ReleaseNotes/Program.cs
--- a/Engine/Source/Programs/Horde/Tools/ReleaseNotes/Program.cs
+++ b/Engine/Source/Programs/Horde/Tools/ReleaseNotes/Program.cs
@@ -89,12 +89,12 @@ namespace CreateReleaseNotes
 				logger.LogInformation("Querying Horde for currently deployed version...");
 
 				GetServerInfoResponse info = await hordeHttpClient.GetServerInfoAsync();
-				if (serverChange == 0 && !TryParseMinChange(info.ServerVersion, out serverChange))
+				if (serverChange == 0 && !TryParseChange(info.ServerVersion, out serverChange))
 				{
 					logger.LogError("Unexpected server version format: {Version}", info.ServerVersion);
 					return 1;
 				}
-				if (agentChange == 0 && !TryParseMinChange(info.AgentVersion, out agentChange))
+				if (agentChange == 0 && !TryParseChange(info.AgentVersion, out agentChange))
 				{
 					logger.LogError("Unexpected agent version format: {Version}", info.AgentVersion);
 					return 1;
@@ -234,7 +234,7 @@ namespace CreateReleaseNotes
 			return 0;
 		}
 
-		static bool TryParseMinChange(string? name, out int version)
+		static bool TryParseChange(string? name, out int version)
 		{
 			if (name == null)
 			{
@@ -250,7 +250,7 @@ namespace CreateReleaseNotes
 			}
 			else
 			{
-				version = Int32.Parse(match.Groups[1].Value) + 1;
+				version = Int32.Parse(match.Groups[1].Value);
 				return true;
 			}
 		}
