[TestDataV2] Change parameter input type for clarity and efficiency in testdata session fetching endpoints

#rnx

diff --git a/Engine/Source/Programs/Horde/Plugins/Build/HordeServer.Build.Tests/Jobs/TestData/TestDataTests.cs b/Engine/Source/Programs/Horde/Plugins/Build/HordeServer.Build.Tests/Jobs/TestData/TestDataTests.cs
--- a/Engine/Source/Programs/Horde/Plugins/Build/HordeServer.Build.Tests/Jobs/TestData/TestDataTests.cs
+++ b/Engine/Source/Programs/Horde/Plugins/Build/HordeServer.Build.Tests/Jobs/TestData/TestDataTests.cs
@@ -371,7 +371,8 @@ namespace HordeServer.Tests.Jobs.TestData
 
 			await TestDataService.TickForTestingAsync();
 
-			ActionResult<List<GetTestSessionStreamResponse>> streamResult = await TestDataController.GetTestStreamsV2Async(streamIds);
+			GetStreamTestsRequest streamRequest = new GetStreamTestsRequest() { StreamIds = streamIds.ToList() };
+			ActionResult<List<GetTestSessionStreamResponse>> streamResult = await TestDataController.GetTestStreamsV2Async(streamRequest);
 			Assert.IsNotNull(streamResult);
 			Assert.IsNotNull(streamResult.Value);
 			List<GetTestSessionStreamResponse> streams = streamResult.Value;
@@ -425,14 +426,16 @@ namespace HordeServer.Tests.Jobs.TestData
 			Assert.IsTrue(tests[0].Phases.Any(i => i.Name == "Editor.Failing.Test"));
 			string phaseId = tests[0].Phases.First(i => i.Name == "Editor.Failing.Test").Id.ToString();
 
-			ActionResult<List<GetTestSessionResponse>> testSessionResult = await TestDataController.GetTestSessionsV2Async(streamIds, testnames.Select(t => t.Id.ToString()).ToArray());
+			GetTestSessionsRequest sessionsRequest = new GetTestSessionsRequest() { StreamIds = streamIds.ToList(), TestIds = testnames.Select(t => t.Id.ToString()).ToList() };
+			ActionResult<List<GetTestSessionResponse>> testSessionResult = await TestDataController.GetTestSessionsV2Async(sessionsRequest);
 			Assert.IsNotNull(testSessionResult);
 			Assert.IsNotNull(testSessionResult.Value);
 			List<GetTestSessionResponse> sessions = testSessionResult.Value;
 			Assert.AreEqual(5, sessions.Count);
 			Assert.AreEqual(TestOutcome.Failure, sessions[0].Outcome);
 
-			ActionResult<List<GetTestPhaseSessionResponse>> phaseSessionResult = await TestDataController.GetTestPhaseSessionsV2Async(streamIds, tests.SelectMany(t => t.Phases.Select(p => p.Id)).ToArray());
+			GetPhaseSessionsRequest phaseSessionRequest = new GetPhaseSessionsRequest() { StreamIds = streamIds.ToList(), PhaseIds = tests.SelectMany(t => t.Phases.Select(p => p.Id)).ToList() };
+			ActionResult<List<GetTestPhaseSessionResponse>> phaseSessionResult = await TestDataController.GetTestPhaseSessionsV2Async(phaseSessionRequest);
 			Assert.IsNotNull(phaseSessionResult);
 			Assert.IsNotNull(phaseSessionResult.Value);
 			List<GetTestPhaseSessionResponse> phaseSessions = phaseSessionResult.Value;
@@ -440,7 +443,8 @@ namespace HordeServer.Tests.Jobs.TestData
 			Assert.AreEqual(TestPhaseOutcome.Failed, phaseSessions.First(p => p.PhaseRef == phaseId).Outcome);
 
 			// Get the sessions for one phase in one stream
-			phaseSessionResult = await TestDataController.GetTestPhaseSessionsV2Async([streamIds[0]], [tests[0].Phases[0].Id]);
+			phaseSessionRequest = new GetPhaseSessionsRequest() { StreamIds = [streamIds[0]], PhaseIds = [tests[0].Phases[0].Id] };
+			phaseSessionResult = await TestDataController.GetTestPhaseSessionsV2Async(phaseSessionRequest);
 			Assert.IsNotNull(phaseSessionResult);
 			Assert.IsNotNull(phaseSessionResult.Value);
 			phaseSessions = phaseSessionResult.Value;
diff --git a/Engine/Source/Programs/Horde/Plugins/Build/HordeServer.Build/Jobs/TestData/TestDataController.cs b/Engine/Source/Programs/Horde/Plugins/Build/HordeServer.Build/Jobs/TestData/TestDataController.cs
--- a/Engine/Source/Programs/Horde/Plugins/Build/HordeServer.Build/Jobs/TestData/TestDataController.cs
+++ b/Engine/Source/Programs/Horde/Plugins/Build/HordeServer.Build/Jobs/TestData/TestDataController.cs
@@ -502,14 +502,14 @@ namespace HordeServer.Jobs.TestData
 		/// <summary>
 		/// Get stream test data for the provided ids
 		/// </summary>
-		/// <param name="streamIds"></param>
+		/// <param name="request"></param>
 		/// <returns></returns>
-		[HttpGet]
+		[HttpPost]
 		[Route("/api/v3/testdata/streams")]
 		[ProducesResponseType(typeof(List<GetTestSessionStreamResponse>), 200)]
-		public async Task<ActionResult<List<GetTestSessionStreamResponse>>> GetTestStreamsV2Async([FromQuery(Name = "id")] string[] streamIds)
+		public async Task<ActionResult<List<GetTestSessionStreamResponse>>> GetTestStreamsV2Async([FromBody] GetStreamTestsRequest request)
 		{
-			StreamId[] streamIdValues = Array.ConvertAll(streamIds, x => new StreamId(x));
+			StreamId[] streamIdValues = request.StreamIds.Select(x => new StreamId(x)).ToArray();
 
 			List<StreamId> queryStreams = new List<StreamId>();
 			List<GetTestSessionStreamResponse> responses = new List<GetTestSessionStreamResponse>();
@@ -698,27 +698,23 @@ namespace HordeServer.Jobs.TestData
 		/// <summary>
 		/// Gets test data sessions 
 		/// </summary>
-		/// <param name="streamIds"></param>
-		/// <param name="testIds"></param>
-		/// <param name="metaIds"></param>
+		/// <param name="request"></param>
 		/// <param name="minCreateTime"></param>
 		/// <param name="maxCreateTime"></param>
 		/// <param name="minChange"></param>
 		/// <param name="maxChange"></param>
 		/// <returns></returns>
-		[HttpGet]
+		[HttpPost]
 		[Route("/api/v3/testdata/session/tests")]
 		[ProducesResponseType(typeof(List<GetTestSessionResponse>), 200)]
 		public async Task<ActionResult<List<GetTestSessionResponse>>> GetTestSessionsV2Async(
-			[FromQuery(Name = "id")] string[] streamIds,
-			[FromQuery(Name = "tid")] string[]? testIds = null,
-			[FromQuery(Name = "mid")] string[]? metaIds = null,
+			[FromBody] GetTestSessionsRequest request,
 			[FromQuery] DateTimeOffset? minCreateTime = null,
 			[FromQuery] DateTimeOffset? maxCreateTime = null,
 			[FromQuery] CommitId? minChange = null,
 			[FromQuery] CommitId? maxChange = null)
 		{
-			StreamId[] streamIdValues = Array.ConvertAll(streamIds, x => new StreamId(x));
+			StreamId[] streamIdValues = request.StreamIds.Select(x => new StreamId(x)).ToArray();
 
 			List<StreamId> queryStreams = new List<StreamId>();
 			List<GetTestSessionResponse> responses = new List<GetTestSessionResponse>();
@@ -737,8 +733,8 @@ namespace HordeServer.Jobs.TestData
 				return responses;
 			}
 
-			TestMetaId[]? queryMetaIds = metaIds?.ConvertAll(x => TestMetaId.Parse(x)).ToArray();
-			TestId[]? queryTestIds = testIds?.ConvertAll(x => TestId.Parse(x)).ToArray();
+			TestMetaId[]? queryMetaIds = request.MetaIds?.Select(x => TestMetaId.Parse(x)).ToArray();
+			TestId[]? queryTestIds = request.TestIds?.Select(x => TestId.Parse(x)).ToArray();
 
 			IReadOnlyList<ITestSession> sessions = await _testDataService.FindTestSessionsV2Async(queryStreams.ToArray(), queryTestIds, queryMetaIds, minCreateTime?.UtcDateTime, maxCreateTime?.UtcDateTime, minChange, maxChange);
 			foreach (ITestSession testData in sessions)
@@ -769,25 +765,23 @@ namespace HordeServer.Jobs.TestData
 		/// <summary>
 		/// Gets test data phase sessions 
 		/// </summary>
-		/// <param name="streamIds"></param>
-		/// <param name="phaseIds"></param>
+		/// <param name="request"></param>
 		/// <param name="minCreateTime"></param>
 		/// <param name="maxCreateTime"></param>
 		/// <param name="minChange"></param>
 		/// <param name="maxChange"></param>
 		/// <returns></returns>
-		[HttpGet]
+		[HttpPost]
 		[Route("/api/v3/testdata/session/phases")]
 		[ProducesResponseType(typeof(List<GetTestPhaseSessionResponse>), 200)]
 		public async Task<ActionResult<List<GetTestPhaseSessionResponse>>> GetTestPhaseSessionsV2Async(
-			[FromQuery(Name = "id")] string[] streamIds,
-			[FromQuery(Name = "pid")] string[] phaseIds,
+			[FromBody] GetPhaseSessionsRequest request,
 			[FromQuery] DateTimeOffset? minCreateTime = null,
 			[FromQuery] DateTimeOffset? maxCreateTime = null,
 			[FromQuery] CommitId? minChange = null,
 			[FromQuery] CommitId? maxChange = null)
 		{
-			StreamId[] streamIdValues = Array.ConvertAll(streamIds, x => new StreamId(x));
+			StreamId[] streamIdValues = request.StreamIds.Select(x => new StreamId(x)).ToArray();
 
 			List<StreamId> queryStreams = new List<StreamId>();
 			List<GetTestPhaseSessionResponse> responses = new List<GetTestPhaseSessionResponse>();
@@ -806,7 +800,7 @@ namespace HordeServer.Jobs.TestData
 				return responses;
 			}
 
-			TestPhaseId[] queryPhaseIds = phaseIds.ConvertAll(x => TestPhaseId.Parse(x)).ToArray();
+			TestPhaseId[] queryPhaseIds = request.PhaseIds.Select(x => TestPhaseId.Parse(x)).ToArray();
 
 			IReadOnlyList<ITestPhaseSession> sessions = await _testDataService.FindTestPhaseSessionsV2Async(queryStreams.ToArray(), queryPhaseIds, minCreateTime?.UtcDateTime, maxCreateTime?.UtcDateTime, minChange, maxChange);
 			foreach (ITestPhaseSession testData in sessions)
diff --git a/Engine/Source/Programs/Shared/EpicGames.Horde/Jobs/TestData/TestDataMessages.cs b/Engine/Source/Programs/Shared/EpicGames.Horde/Jobs/TestData/TestDataMessages.cs
--- a/Engine/Source/Programs/Shared/EpicGames.Horde/Jobs/TestData/TestDataMessages.cs
+++ b/Engine/Source/Programs/Shared/EpicGames.Horde/Jobs/TestData/TestDataMessages.cs
@@ -257,6 +257,17 @@ namespace EpicGames.Horde.Jobs.TestData
 		public List<string> TestIds { get; set; } = new List<string>();
 	}
 
+	/// <summary>
+	/// Get Tests for streams request
+	/// </summary>
+	public class GetStreamTestsRequest
+	{
+		/// <summary>
+		/// Stream ids to get
+		/// </summary>
+		public List<string> StreamIds { get; set; } = new List<string>();
+	}
+
 	/// <summary>
 	/// Get test metadata request
 	/// </summary>
@@ -305,6 +316,43 @@ namespace EpicGames.Horde.Jobs.TestData
 		public List<string>? TestIds { get; set; }
 	}
 
+	/// <summary>
+	/// Get test sessions from streams, tests and optionally meta ids
+	/// </summary>
+	public class GetTestSessionsRequest
+	{
+		/// <summary>
+		/// Stream ids to get
+		/// </summary>
+		public List<string> StreamIds { get; set; } = new List<string>();
+
+		/// <summary>
+		/// Test ids to get
+		/// </summary>
+		public List<string>? TestIds { get; set; }
+
+		/// <summary>
+		/// Metadata ids to get
+		/// </summary>
+		public List<string>? MetaIds { get; set; }
+	}
+
+	/// <summary>
+	/// Get phase sessions from stream and phase ids
+	/// </summary>
+	public class GetPhaseSessionsRequest
+	{
+		/// <summary>
+		/// Stream ids to get
+		/// </summary>
+		public List<string> StreamIds { get; set; } = new List<string>();
+
+		/// <summary>
+		/// Phase ids to get
+		/// </summary>
+		public List<string> PhaseIds { get; set; } = new List<string>();
+	}
+
 	/// <summary>
 	/// A test suite that runs in a stream, contain subtests
 	/// </summary>
