Horde: Add support to get secrets from HashiCorp Vault

diff --git a/Engine/Source/Programs/Horde/Docs/Config/Secrets.md b/Engine/Source/Programs/Horde/Docs/Config/Secrets.md
--- a/Engine/Source/Programs/Horde/Docs/Config/Secrets.md
+++ b/Engine/Source/Programs/Horde/Docs/Config/Secrets.md
@@ -6,7 +6,7 @@ Horde implements an API for retrieving secrets that can be stored in its own con
 an external source. Marshalling data through Horde allows access to be controlled using Horde's permissions model,
 and for automated processes to impersonate the user that requested them.
 
-At the moment, Horde only supports the AWS Parameter Store as an external secret provider out of the box, though 
+At the moment, Horde only supports the AWS Parameter Store and HashiCorp Vault as external secret providers out of the box, though 
 other implementations can be added through the `ISecretProvider` interface.
 
 ## Configuring secrets
@@ -76,6 +76,76 @@ An example secret may be configured as follows:
 
 Secrets are queried from the external provider when requested by a user, and are not cached by Horde.
 
+Secrets from external providers that require additional options, for example HashiCorp Vault to set the endpoint and
+credentials, will use a `providerconfig` property instead of `name` property to provide the data. For example to fetch
+two secrets from HashiCorp Vault one using an AWS role and the other a pre-shared token would be configured as follows:
+
+   ```json
+    "secrets": [
+        {
+            "id": "horde-secrets",
+
+            // Some values read from HashiCorp Vault
+            "sources": [
+
+                {
+                    "providerconfig": "hcp-vault-aws-role",
+                    // The Vault API path to the secret e.g. /v1/secret/data/my-secret
+                    "path": "path-of-secret-in-vault",
+                    // For Vault the format is always json because a secret can contain multiple items
+                    "format": "json",
+                },
+
+                {
+                    "providerconfig": "hcp-vault-pre-shared-key",
+                    "format": "json",
+                    "path": "path-of-secret-in-vault"
+                }
+            ],
+
+            "providers": [
+
+                {
+                    "name": "hcp-vault-aws-role",
+                    "provider": "HcpVault",
+                    "hcpvault": {
+                        "credentials": "awsauth",
+                        "endpoint": "https://vault.example.com",
+                        "awsiamserverid": "vault.example.com",
+                        "awsarnrole": "arn:aws:iam::1234567890:role/vault-example-role",
+                        "role": "vault-role-name"
+                    }
+                },
+
+                {
+                    "name": "hcp-vault-pre-shared-key",
+                    "provider": "HcpVault",
+                    "hcpvault": {
+                        "credentials": "presharedkey",
+                        "endpoint": "https://vault.example.com",
+                        "presharedkey": "the-key-provided-by-vault"
+                    }
+                }
+            ]
+
+            // Only allow Horde agents to access this
+            "acl": {
+                "entries": [
+                    {
+                        "claim": {
+                            "type": "http://epicgames.com/ue/horde/role",
+                            "value": "agent"
+                        },
+                        "actions": [
+                            "ViewSecret"
+                        ]
+                    }
+                ]
+            }
+        }
+    ]
+ ```
+
 ## Using secrets
 
 The most common use case for secrets is during build automation pipelines. In this scenario, the Horde Server
diff --git a/Engine/Source/Programs/Horde/Plugins/Secrets/HordeServer.Secrets.Tests/HordeServer.Secrets.Tests.csproj b/Engine/Source/Programs/Horde/Plugins/Secrets/HordeServer.Secrets.Tests/HordeServer.Secrets.Tests.csproj
new file mode 100644
--- /dev/null
+++ b/Engine/Source/Programs/Horde/Plugins/Secrets/HordeServer.Secrets.Tests/HordeServer.Secrets.Tests.csproj
@@ -0,0 +1,43 @@
+<Project Sdk="Microsoft.NET.Sdk">
+
+  <PropertyGroup>
+    <TargetFramework>net8.0</TargetFramework>
+    <ImplicitUsings>enable</ImplicitUsings>
+    <Nullable>enable</Nullable>
+    <IsPackable>false</IsPackable>
+    <IsTestProject>true</IsTestProject>
+    <RootNamespace>HordeServer.Tests</RootNamespace>
+    <RunAnalyzersDuringBuild>False</RunAnalyzersDuringBuild>
+    <Configurations>Debug;Release;Analyze</Configurations>
+    <NoWarn>1701;1702;1591</NoWarn>
+  </PropertyGroup>
+
+  <PropertyGroup Condition=" '$(Configuration)|$(Platform)' == 'Analyze|AnyCPU' ">
+    <RunAnalyzersDuringBuild>True</RunAnalyzersDuringBuild>
+    <TreatWarningsAsErrors>false</TreatWarningsAsErrors>
+    <EnforceCodeStyleInBuild>True</EnforceCodeStyleInBuild>
+    <GenerateDocumentationFile>True</GenerateDocumentationFile>
+  </PropertyGroup>
+
+  <ItemGroup>
+    <PackageReference Include="AWSSDK.SSO" Version="3.7.400.108" />
+    <PackageReference Include="AWSSDK.SSOOIDC" Version="3.7.400.109" />
+    <PackageReference Include="coverlet.collector" Version="6.0.2">
+      <PrivateAssets>all</PrivateAssets>
+      <IncludeAssets>runtime; build; native; contentfiles; analyzers; buildtransitive</IncludeAssets>
+    </PackageReference>
+    <PackageReference Include="Microsoft.NET.Test.Sdk" Version="17.11.1" />
+    <PackageReference Include="MSTest.TestAdapter" Version="3.6.2" />
+    <PackageReference Include="MSTest.TestFramework" Version="3.6.2" />
+  </ItemGroup>
+
+  <ItemGroup>
+    <ProjectReference Include="..\HordeServer.Secrets\HordeServer.Secrets.csproj" />
+    <ProjectReference Include="..\..\..\HordeServer.Tests.Shared\HordeServer.Tests.Shared.csproj" />
+  </ItemGroup>
+
+  <ItemGroup>
+    <Using Include="Microsoft.VisualStudio.TestTools.UnitTesting" />
+  </ItemGroup>
+
+</Project>
diff --git a/Engine/Source/Programs/Horde/Plugins/Secrets/HordeServer.Secrets.Tests/Secrets/SecretProviderTests.cs b/Engine/Source/Programs/Horde/Plugins/Secrets/HordeServer.Secrets.Tests/Secrets/SecretProviderTests.cs
new file mode 100644
--- /dev/null
+++ b/Engine/Source/Programs/Horde/Plugins/Secrets/HordeServer.Secrets.Tests/Secrets/SecretProviderTests.cs
@@ -0,0 +1,295 @@
+// Copyright Epic Games, Inc. All Rights Reserved.
+
+using System.Text.Json;
+using EpicGames.Horde.Secrets;
+using HordeServer.Secrets;
+using HordeServer.Secrets.Providers;
+using HordeServer.Server;
+using HordeServer.Utilities;
+using Microsoft.Extensions.Options;
+
+namespace HordeServer.Tests.Secrets
+{
+	[TestClass]
+	public class SecretProviderTests
+	{
+		private class OptionsSnapshotStub<T>(T value) : IOptionsSnapshot<T> where T : class
+		{
+			public T Value { get; } = value;
+			public T Get(string? name)
+			{
+				return Value;
+			}
+		}
+
+		private class InMemoryProvider : ISecretProvider
+		{
+			private Dictionary<string, string> Secrets { get; } = new()
+			{
+				{ "secret/username", "user01" },
+				{ "secret/password", "password01" }
+			};
+			public string Name { get; } = "InMemory";
+
+			public Task<string> GetSecretAsync(string path, SecretProviderConfig? config, CancellationToken cancellationToken)
+			{
+				return Task.FromResult(Secrets[path]);
+			}
+		}
+
+		private class StubHttpClientFactory : IHttpClientFactory
+		{
+			public HttpClient CreateClient(string name)
+			{
+				return new HttpClient();
+			}
+		}
+		private static SecretCollection CreateSecretCollection()
+		{
+			SecretCollectionInternal secretCollectionInternal = new([
+				new InMemoryProvider(),
+				new HcpVaultSecretProvider(new StubHttpClientFactory())
+			]);
+			SecretsConfig secretsConfig = new();
+			secretsConfig.Secrets.Add(new SecretConfig()
+			{
+				Id = new SecretId("Secret"),
+				Sources = [
+					new ExternalSecretConfig()
+					{
+						Provider = "InMemory",
+						Path = "secret/username",
+						Key = "username"
+					},
+					new ExternalSecretConfig()
+					{
+						Provider = "InMemory",
+						Path = "secret/password",
+						Key = "password"
+					}
+				]
+			});
+			secretsConfig.Secrets.Add(new SecretConfig()
+			{
+				Id = new SecretId("SecretUsingProviderConfig"),
+				Sources = [
+					new ExternalSecretConfig()
+					{
+						ProviderConfig = "config-id",
+						Path = "secret/username",
+						Key = "username",
+					},
+					new ExternalSecretConfig()
+					{
+						ProviderConfig = "config-id",
+						Path = "secret/password",
+						Key = "password",
+					}
+				],
+				Providers = [
+					new SecretProviderConfig()
+					{
+						Name = "config-id",
+						Provider = "InMemory"
+					}
+				]
+			});
+			secretsConfig.Secrets.Add(new SecretConfig()
+			{
+				Id = new SecretId("SecretUsingHcpVault"),
+				Sources = [
+					new ExternalSecretConfig
+					{
+						ProviderConfig = "pre-shared-key",
+						Path = "/v1/secret/data/example",
+						Format = ExternalSecretFormat.Json
+					}
+				],
+				Providers = [
+					new SecretProviderConfig()
+					{
+						Name = "pre-shared-key",
+						Provider = "HcpVault",
+						HcpVault = new HcpVaultConfig()
+						{
+							Credentials = HcpVaultCredentialsType.PreSharedKey,
+							EndPoint = "http://127.0.0.1:8200",
+							PreSharedKey = "replace"
+						}
+					}
+				]
+			});
+			secretsConfig.Secrets.Add(new SecretConfig()
+			{
+				Id = new SecretId("SecretUsingHcpVaultAwsAuth"),
+				Sources = [
+					new ExternalSecretConfig
+					{
+						ProviderConfig = "vault-aws",
+						Path = "/v1/secret/data/example",
+						Format = ExternalSecretFormat.Json
+					}
+				],
+				Providers = [
+					new SecretProviderConfig()
+					{
+						Name = "vault-aws",
+						Provider = "HcpVault",
+						HcpVault = new HcpVaultConfig(){
+							Credentials = HcpVaultCredentialsType.AwsAuth,
+							EndPoint = "https://vault.example.com/",
+							AwsIamServerId = "vault.example.com",
+							AwsArnRole = "arn:aws:iam::1234567890:role/vault-example-role",
+							Role = "vault-example"
+						}
+					}
+				]
+			});
+			secretsConfig.Secrets.Add(new SecretConfig()
+			{
+				Id = new SecretId("EmptySecret"),
+				Sources = [new ExternalSecretConfig() ]
+			});
+			SecretCollection secretCollection = new(secretCollectionInternal, new OptionsSnapshotStub<SecretsConfig>(secretsConfig));
+			GlobalConfig config = new();
+			config.Plugins.AddSecretsConfig(secretsConfig);
+			config.PostLoad(new ServerSettings(), [], []);
+			return secretCollection;
+		}
+
+		[TestMethod]
+		public async Task GetSecretAsync()
+		{
+			SecretCollection secretCollection = CreateSecretCollection(); 
+			ISecret? secret = await secretCollection.GetAsync(new SecretId("Secret"), CancellationToken.None);
+			Assert.AreEqual(secret?.Data["username"], "user01");
+			Assert.AreEqual(secret?.Data["password"], "password01");
+		}
+
+		[TestMethod]
+		public async Task GetSecretFromProviderConfigAsync()
+		{
+			SecretCollection secretCollection = CreateSecretCollection(); 
+			ISecret? secret = await secretCollection.GetAsync(new SecretId("SecretUsingProviderConfig"), CancellationToken.None);
+			Assert.AreEqual(secret?.Data["username"], "user01");
+			Assert.AreEqual(secret?.Data["password"], "password01");
+		}
+
+		[TestMethod]
+		public async Task InvalidGetSecretAsync()
+		{
+			SecretCollection secretCollection = CreateSecretCollection();
+			await Assert.ThrowsExceptionAsync<InvalidOperationException>(() => secretCollection.GetAsync(new SecretId("EmptySecret"), CancellationToken.None));
+		}
+
+		[Ignore("Requires access to a running HCP Vault server and a pre-shared key")]
+		[TestMethod]
+		public async Task GetSecretFromVaultAsync()
+		{
+			// To set up a local dev server with test data:
+			// 1. install the vault executable https://developer.hashicorp.com/vault/install
+			// 2. in a shell run: vault server -dev
+			//    make a note of the token value called 'Root Token' printed to stdout 
+			// 3. in a second shell set the environment variable VAULT_ADDR
+			//    the addr will likely be http://127.0.0.1:8200
+			// 4. in the second shell login in with the token: vault login
+			// 5. in the second shell add the data: vault kv put secret/example username="User01" password="Password123"
+			// 6. update the SecretProviderConfig options for the SecretId used in this test 
+			SecretCollection secretCollection = CreateSecretCollection(); 
+			ISecret? secret = await secretCollection.GetAsync(new SecretId("SecretUsingHcpVault"), CancellationToken.None);
+			Assert.AreEqual(secret?.Data["username"], "User01");
+			Assert.AreEqual(secret?.Data["password"], "Password123");
+		}
+
+		[Ignore("Requires access to a running HCP Vault server and supported AWS IAM")]
+		[TestMethod]
+		public async Task GetSecretFromVaultWithAwsAuthAsync()
+		{
+			// To run this test from your workstation
+			// 1. log into AWS with: aws sso login
+			// 2. set the environment variable AWS_PROFILE to a profile that can call AssumeRole for the required ARN
+			// 3. restart IDE to pick up the environment change
+			// 4. update the SecretProviderConfig options for the SecretId used in this test 
+			SecretCollection secretCollection = CreateSecretCollection(); 
+			ISecret? secret = await secretCollection.GetAsync(new SecretId("SecretUsingHcpVaultAwsAuth"), CancellationToken.None);
+			Assert.AreEqual(secret?.Data["test"], "test");
+		}
+
+		[TestMethod]
+		public void GetSecretConfigFromJson()
+		{
+			string json = """
+				{
+					"secrets": [
+						{
+							"id": "my-secrets",
+							"sources":[
+								{
+									"ProviderConfig": "vault-aws-provider",
+									"path": "path-to-secret",
+									"format": "json"
+								},
+								{
+									"ProviderConfig": "vault-pre-shared-key-provider",
+									"path": "path-to-secret",
+									"format": "json"
+								}
+							],
+							"providers": [
+								{
+									"Name": "vault-aws-provider",
+									"Provider": "HcpVault",
+									"HcpVault": {
+										"Credentials": "awsauth",
+										"EndPoint": "https://vault.example.com",
+										"AwsIamServerId": "vault.example.com",
+										"AwsArnRole": "arn:aws:iam::1234567890:role/vault-example-role",
+										"Role": "vault-role-name"
+									}
+								},
+								{
+									"Name": "vault-pre-shared-key-provider",
+									"Provider": "HcpVault",
+									"HcpVault": {
+										"credentials": "presharedkey",
+										"endpoint": "https://vault.example.com",
+										"presharedkey": "Abc123Def"
+									}
+								}
+							]
+						}
+					]
+				}
+			""";
+			SecretsConfig secretsConfig = JsonSerializer.Deserialize<SecretsConfig>(json, JsonUtils.DefaultSerializerOptions)!;
+			SecretConfig secretConfig = secretsConfig.Secrets[0];
+			Assert.AreEqual(secretConfig.Id.ToString(), "my-secrets");
+
+			ExternalSecretConfig externalSecretConfig0 = secretConfig.Sources[0];
+			SecretProviderConfig secretProviderConfig0 = secretConfig.Providers[0];
+			HcpVaultConfig hcpVaultConfig0 = secretProviderConfig0.HcpVault!;
+			Assert.AreEqual(externalSecretConfig0.Format, ExternalSecretFormat.Json);
+			Assert.AreEqual(externalSecretConfig0.Path, "path-to-secret");
+			Assert.AreEqual(externalSecretConfig0.ProviderConfig, "vault-aws-provider");
+			Assert.AreEqual(secretProviderConfig0.Name, "vault-aws-provider");
+			Assert.AreEqual(secretProviderConfig0.Provider, "HcpVault");
+			Assert.AreEqual(hcpVaultConfig0.EndPoint, "https://vault.example.com");
+			Assert.AreEqual(hcpVaultConfig0.Credentials, HcpVaultCredentialsType.AwsAuth);
+			Assert.AreEqual(hcpVaultConfig0.AwsArnRole, "arn:aws:iam::1234567890:role/vault-example-role");
+			Assert.AreEqual(hcpVaultConfig0.AwsIamServerId, "vault.example.com");
+			Assert.AreEqual(hcpVaultConfig0.Role, "vault-role-name");
+
+			ExternalSecretConfig externalSecretConfig1 = secretConfig.Sources[1];
+			SecretProviderConfig secretProviderConfig1 = secretConfig.Providers[1];
+			HcpVaultConfig hcpVaultConfig1 = secretProviderConfig1.HcpVault!;
+			Assert.AreEqual(externalSecretConfig1.Format, ExternalSecretFormat.Json);
+			Assert.AreEqual(externalSecretConfig1.Path, "path-to-secret");
+			Assert.AreEqual(externalSecretConfig1.ProviderConfig, "vault-pre-shared-key-provider");
+			Assert.AreEqual(secretProviderConfig1.Name, "vault-pre-shared-key-provider");
+			Assert.AreEqual(secretProviderConfig1.Provider, "HcpVault");
+			Assert.AreEqual(hcpVaultConfig1.EndPoint, "https://vault.example.com");
+			Assert.AreEqual(hcpVaultConfig1.Credentials, HcpVaultCredentialsType.PreSharedKey);
+			Assert.AreEqual(hcpVaultConfig1.PreSharedKey, "Abc123Def");
+		}
+	}
+}
diff --git a/Engine/Source/Programs/Horde/Plugins/Secrets/HordeServer.Secrets/Secrets/ISecretProvider.cs b/Engine/Source/Programs/Horde/Plugins/Secrets/HordeServer.Secrets/Secrets/ISecretProvider.cs
--- a/Engine/Source/Programs/Horde/Plugins/Secrets/HordeServer.Secrets/Secrets/ISecretProvider.cs
+++ b/Engine/Source/Programs/Horde/Plugins/Secrets/HordeServer.Secrets/Secrets/ISecretProvider.cs
@@ -16,7 +16,8 @@ namespace HordeServer.Secrets
 		/// Gets values for a particular secret
 		/// </summary>
 		/// <param name="path">Path to the secret to fetch</param>
+		/// <param name="config">Configuration for the secret provider</param>
 		/// <param name="cancellationToken">Cancellation token for the operation</param>
-		Task<string> GetSecretAsync(string path, CancellationToken cancellationToken);
+		Task<string> GetSecretAsync(string path, SecretProviderConfig? config, CancellationToken cancellationToken);
 	}
 }
diff --git a/Engine/Source/Programs/Horde/Plugins/Secrets/HordeServer.Secrets/Secrets/Providers/AwsParameterStoreSecretProvider.cs b/Engine/Source/Programs/Horde/Plugins/Secrets/HordeServer.Secrets/Secrets/Providers/AwsParameterStoreSecretProvider.cs
--- a/Engine/Source/Programs/Horde/Plugins/Secrets/HordeServer.Secrets/Secrets/Providers/AwsParameterStoreSecretProvider.cs
+++ b/Engine/Source/Programs/Horde/Plugins/Secrets/HordeServer.Secrets/Secrets/Providers/AwsParameterStoreSecretProvider.cs
@@ -26,7 +26,7 @@ namespace HordeServer.Secrets.Providers
 		}
 
 		/// <inheritdoc/>
-		public async Task<string> GetSecretAsync(string path, CancellationToken cancellationToken)
+		public async Task<string> GetSecretAsync(string path, SecretProviderConfig? config, CancellationToken cancellationToken)
 		{
 			GetParameterResponse response = await _systemsManagement.GetParameterAsync(new GetParameterRequest { Name = path, WithDecryption = true }, cancellationToken);
 			if (response.Parameter == null)
diff --git a/Engine/Source/Programs/Horde/Plugins/Secrets/HordeServer.Secrets/Secrets/Providers/HcpVaultConfig.cs b/Engine/Source/Programs/Horde/Plugins/Secrets/HordeServer.Secrets/Secrets/Providers/HcpVaultConfig.cs
new file mode 100644
--- /dev/null
+++ b/Engine/Source/Programs/Horde/Plugins/Secrets/HordeServer.Secrets/Secrets/Providers/HcpVaultConfig.cs
@@ -0,0 +1,67 @@
+// Copyright Epic Games, Inc. All Rights Reserved.
+
+namespace HordeServer.Secrets.Providers
+{
+	/// <summary>
+	/// Credentials to use to log into HCP Vault.
+	/// </summary>
+	public enum HcpVaultCredentialsType
+	{
+		/// <summary>
+		/// Use a pre-shared token provided by Vault for authentication.
+		/// </summary>
+		PreSharedKey,
+		
+		/// <summary>
+		/// Use AWS to authenticate against the Vault server.
+		/// </summary>
+		AwsAuth
+	}
+
+	/// <summary>
+	/// Configuration for HCP Vault secrets.
+	/// </summary>
+	public class HcpVaultConfig
+	{
+		/// <summary>
+		/// The Vault authentication token to use when the type of credentials is <see cref="HcpVaultCredentialsType.PreSharedKey"/>.
+		/// <remarks>
+		/// More information on this token can be found in the API documentation for <see href="https://developer.hashicorp.com/vault/api-docs#authentication">X-Vault-Token</see>
+		/// </remarks>
+		/// </summary>
+		public string? PreSharedKey { get; set; }
+		
+		/// <summary>
+		/// The address of the Vault server.
+		/// </summary>
+		public string? EndPoint { get; set; }
+
+		/// <summary>
+		/// The AWS ARN to assume a role when logging into Vault.
+		/// To be used when the type of credentials is <see cref="HcpVaultCredentialsType.AwsAuth"/>.
+		/// <remarks>When null AWS SDK's default credential search will be used.</remarks>
+		/// </summary>
+		public string? AwsArnRole { get; set; }
+
+		/// <summary>
+		/// The value for the X-Vault-AWS-IAM-Server-ID header when the Vault server is configured with
+		/// <see href="https://developer.hashicorp.com/vault/api-docs/auth/aws#iam_server_id_header_value">iam_server_id_header_value</see>.
+		/// Most likely to be the Vault server's DNS name.
+		/// To be used when the type of credentials is <see cref="HcpVaultCredentialsType.AwsAuth"/>.
+		/// </summary>
+		public string? AwsIamServerId { get; set; }
+	
+		/// <summary>
+		/// The name of the <see href="https://developer.hashicorp.com/vault/api-docs/auth/aws#role-4">role</see> in the Vault configuration
+		/// for the login. This role name can be different to the name of the AWS ARN role name.
+		/// To be used when the type of credentials is <see cref="HcpVaultCredentialsType.AwsAuth"/>.
+		/// <remarks>When null the name of IAM principal will be used.</remarks>
+		/// </summary>
+		public string? Role { get; set; }
+	
+		/// <summary>
+		/// Type of credentials to use.
+		/// </summary>
+		public HcpVaultCredentialsType Credentials { get; set; }
+	}
+}
diff --git a/Engine/Source/Programs/Horde/Plugins/Secrets/HordeServer.Secrets/Secrets/Providers/HcpVaultSecretProvider.cs b/Engine/Source/Programs/Horde/Plugins/Secrets/HordeServer.Secrets/Secrets/Providers/HcpVaultSecretProvider.cs
new file mode 100644
--- /dev/null
+++ b/Engine/Source/Programs/Horde/Plugins/Secrets/HordeServer.Secrets/Secrets/Providers/HcpVaultSecretProvider.cs
@@ -0,0 +1,234 @@
+// Copyright Epic Games, Inc. All Rights Reserved.
+
+using System.Net;
+using System.Net.Http.Json;
+using System.Text;
+using System.Text.Json;
+using System.Text.Json.Serialization;
+using Amazon.Runtime;
+using Amazon.Runtime.Internal;
+using Amazon.Runtime.Internal.Auth;
+using Amazon.Runtime.Internal.Util;
+using Amazon.SecurityToken;
+using Amazon.SecurityToken.Model;
+using Amazon.SecurityToken.Model.Internal.MarshallTransformations;
+
+namespace HordeServer.Secrets.Providers
+{
+	/// <summary>
+	/// Fetches secrets from a HashiCorp (HCP) Vault.
+	/// </summary>
+	public class HcpVaultSecretProvider : ISecretProvider
+	{
+		/// <summary>
+		/// The login parameters to Vault when using AWS for authentication.
+		/// <remarks>
+		/// The parameters for the <see href="https://developer.hashicorp.com/vault/api-docs/auth/aws#login">v1/auth/aws/login</see> request.
+		/// </remarks>
+		/// </summary>
+		private class AwsAuthRequest
+		{
+			/// <summary>
+			/// The <see href="https://developer.hashicorp.com/vault/api-docs/auth/aws#iam_http_request_method">HTTP method</see>
+			/// used in the signed request.
+			/// </summary>
+			[JsonPropertyName("iam_http_request_method")]
+			public string Method { get; set; } = "POST";
+
+			/// <summary>
+			/// Base64-encoded <see href="https://developer.hashicorp.com/vault/api-docs/auth/aws#iam_request_body">body</see> of the signed request.
+			/// </summary>
+			[JsonPropertyName("iam_request_body")]
+			public string Body { get; set; } = Convert.ToBase64String("Action=GetCallerIdentity&Version=2011-06-15"u8.ToArray());
+
+			/// <summary>
+			/// Based64-encoded key/value pairs of <see href="https://developer.hashicorp.com/vault/api-docs/auth/aws#iam_request_headers">headers</see>
+			/// for the sts:GetCallerIdentity request.
+			/// </summary>
+			[JsonPropertyName("iam_request_headers")]
+			public required string Headers { get; set; }
+
+			/// <summary>
+			/// Base64-encoded of the <see href="https://developer.hashicorp.com/vault/api-docs/auth/aws#iam_request_url">HTTP URL</see> used to sign
+			/// the request.
+			/// </summary>
+			[JsonPropertyName("iam_request_url")]
+			public string Url { get; set; } = Convert.ToBase64String("https://sts.amazonaws.com/"u8.ToArray());
+
+			/// <summary>
+			/// Name of the <see href="https://developer.hashicorp.com/vault/api-docs/auth/aws#role-4">role</see> the login is attempted against.
+			/// This is not the same as the AWS ARN Role.
+			/// </summary>
+			[JsonPropertyName("role")]
+			public string Role { get; set; } = String.Empty;
+		}
+
+		/// <summary>
+		/// The <see href="https://developer.hashicorp.com/vault/api-docs/auth/aws#sample-response-12">response</see> from the login request.
+		/// </summary>
+		private class AuthResponse
+		{
+			[JsonPropertyName("auth")]
+			public AuthInfo? AuthInfo { get; set; }
+
+			public string[]? Errors { get; set; }
+		}
+
+		/// <summary>
+		/// The data from a successful login response.
+		/// </summary>
+		private class AuthInfo
+		{
+			/// <summary>
+			/// The token returned from a login request and to be used as the value for the X-Vault-Token header for following API calls.
+			/// </summary>
+			[JsonPropertyName("client_token")]
+			public string? Token { get; set; }
+		}
+
+		/// <summary>
+		/// The response from requesting a secret.
+		/// </summary>
+		private class SecretResponse
+		{
+			/// <summary>
+			/// The key/value paris of all items for a given path to a secret.
+			/// </summary>
+			public Dictionary<string, Dictionary<string, object>> Data { get; set; } = new();
+
+			public string[]? Errors { get; set; }
+		}
+
+		/// <inheritdoc/>
+		public string Name => "HcpVault";
+
+		private readonly IHttpClientFactory _httpClientFactory;
+		private const string HttpClientName = "HordeSecretProviderHcpVault";
+		private const string AwsAssumeRoleSessionName = "HordeSecretProviderVaultSession";
+
+		/// <summary>
+		/// Constructor
+		/// </summary>
+		public HcpVaultSecretProvider(IHttpClientFactory httpClientFactory)
+		{
+			_httpClientFactory = httpClientFactory;
+		}
+
+		private HttpClient GetHttpClient()
+		{
+			return _httpClientFactory.CreateClient(HttpClientName);
+		}
+
+		private static async Task<string> GetAwsSignedHeadersAsync(string? serverId, string? arnRole, CancellationToken cancellationToken)
+		{
+			ImmutableCredentials credentials;
+			if (arnRole != null)
+			{
+				using AmazonSecurityTokenServiceClient client = new();
+				AssumeRoleRequest assumeRoleRequest = new() { RoleArn = arnRole, RoleSessionName = AwsAssumeRoleSessionName };
+				AssumeRoleResponse? assumeRole = await client.AssumeRoleAsync(assumeRoleRequest, cancellationToken);
+				credentials = await assumeRole.Credentials.GetCredentialsAsync();
+			}
+			else
+			{
+				credentials = await FallbackCredentialsFactory.GetCredentials().GetCredentialsAsync();
+			}
+			IRequest? request = GetCallerIdentityRequestMarshaller.Instance.Marshall(new GetCallerIdentityRequest());
+			request.ResourcePath = "/";
+			request.Headers.Add("Content-Type", "application/x-www-form-urlencoded; charset=utf-8");
+			if (serverId != null)
+			{
+				request.Headers.Add("X-Vault-AWS-IAM-Server-ID", serverId);
+			}
+			if (credentials.UseToken)
+			{
+				request.Headers.Add("X-Amz-Security-Token", credentials.Token);
+			}
+			AmazonSecurityTokenServiceConfig securityTokenServiceConfig = new();
+#pragma warning disable CS0618 // Type or member is obsolete
+			request.Endpoint = new Uri(securityTokenServiceConfig.DetermineServiceURL());
+#pragma warning restore CS0618 // Type or member is obsolete
+			new AWS4Signer().Sign(request, securityTokenServiceConfig, new RequestMetrics(), credentials.AccessKey, credentials.SecretKey);
+			return JsonSerializer.Serialize(request.Headers);
+		}
+
+		private async Task<string> GetTokenFromAwsAsync(HcpVaultConfig config, CancellationToken cancellationToken)
+		{
+			string signedHeaders = await GetAwsSignedHeadersAsync(config.AwsIamServerId, config.AwsArnRole, cancellationToken);
+			AwsAuthRequest payload = new()
+			{
+				Headers = Convert.ToBase64String(Encoding.UTF8.GetBytes(signedHeaders)),
+				Role = config.Role ?? String.Empty
+			};
+			// https://developer.hashicorp.com/vault/api-docs/auth/aws#login 
+			Uri requestUri = new(new Uri(config.EndPoint!), "v1/auth/aws/login");
+			HttpClient httpClient = GetHttpClient();
+			HttpResponseMessage httpResponseMessage = await httpClient.PostAsJsonAsync(requestUri, payload, cancellationToken);
+			AuthResponse? response = await httpResponseMessage.Content.ReadFromJsonAsync<AuthResponse>(cancellationToken);
+			if (httpResponseMessage.StatusCode != HttpStatusCode.OK)
+			{
+				string errorMessage = $"({httpResponseMessage.StatusCode})";
+				if (response?.Errors?.Length > 0)
+				{
+					errorMessage += $" {String.Join(';', response.Errors).ReplaceLineEndings(String.Empty)}";
+				}
+				throw new InvalidOperationException($"Unable to log into HCP Vault using AWS {errorMessage}");
+			}
+			return response?.AuthInfo?.Token ?? throw new InvalidOperationException($"Unable to log into HCP Vault using AWS (No Token Content)");
+		}
+
+		private async Task<string> GetSecretAsync(Uri path, string token, CancellationToken cancellationToken)
+		{
+			// https://developer.hashicorp.com/vault/api-docs 
+			using HttpRequestMessage request = new(HttpMethod.Get, path);
+			HttpClient httpClient = GetHttpClient();
+			httpClient.DefaultRequestHeaders.Add("X-Vault-Token", token);
+			using HttpResponseMessage httpResponseMessage = await httpClient.SendAsync(request, cancellationToken);
+			SecretResponse? response = await httpResponseMessage.Content.ReadFromJsonAsync<SecretResponse>(cancellationToken);
+			if (httpResponseMessage.StatusCode != HttpStatusCode.OK)
+			{
+				string errorMessage = $"({httpResponseMessage.StatusCode})";
+				if (response?.Errors?.Length > 0)
+				{
+					errorMessage += $" {String.Join(';', response.Errors).ReplaceLineEndings(String.Empty)}";
+				}
+				throw new InvalidOperationException($"Unable to fetch secret '{path.AbsolutePath}' from HCP Vault {errorMessage}");
+			}
+			if (response != null)
+			{
+				if (response.Data.TryGetValue("data", out Dictionary<string, object>? value))
+				{
+					return JsonSerializer.Serialize(value);
+				}
+			}
+			throw new InvalidOperationException($"Unable to fetch secret '{path.AbsolutePath}' from HCP Vault (No Content)");
+		}
+
+		/// <inheritdoc/>
+		public async Task<string> GetSecretAsync(string path, SecretProviderConfig? config, CancellationToken cancellationToken)
+		{
+			if (config?.HcpVault == null)
+			{
+				throw new InvalidOperationException($"Unable to fetch secret {path} from HCP Vault (No Options)");
+			}
+			if (config.HcpVault.EndPoint == null)
+			{
+				throw new InvalidOperationException($"Unable to fetch secret {path} from HCP Vault (No HcpVaultEndPoint)");
+			}
+			string token;
+			if (config.HcpVault.Credentials == HcpVaultCredentialsType.PreSharedKey)
+			{
+				token = config.HcpVault.PreSharedKey ?? throw new InvalidOperationException($"Unable to fetch secret {path} form HCP Vault (No Vault Token)");
+			}
+			else if (config.HcpVault.Credentials == HcpVaultCredentialsType.AwsAuth)
+			{
+				token = await GetTokenFromAwsAsync(config.HcpVault, cancellationToken);
+			}
+			else
+			{
+				throw new NotImplementedException();
+			}
+			return await GetSecretAsync(new Uri(new Uri(config.HcpVault.EndPoint), path), token, cancellationToken);
+		}
+	}
+}
diff --git a/Engine/Source/Programs/Horde/Plugins/Secrets/HordeServer.Secrets/Secrets/SecretCollection.cs b/Engine/Source/Programs/Horde/Plugins/Secrets/HordeServer.Secrets/Secrets/SecretCollection.cs
--- a/Engine/Source/Programs/Horde/Plugins/Secrets/HordeServer.Secrets/Secrets/SecretCollection.cs
+++ b/Engine/Source/Programs/Horde/Plugins/Secrets/HordeServer.Secrets/Secrets/SecretCollection.cs
@@ -79,9 +79,22 @@ namespace HordeServer.Secrets
 			// Fetch all the secrets from external providers
 			foreach (ExternalSecretConfig externalConfig in config.Sources)
 			{
-				ISecretProvider provider = _providers[externalConfig.Provider];
+				ISecretProvider provider;
+				SecretProviderConfig? providerConfig = null;
+				if (!String.IsNullOrEmpty(externalConfig.Provider))
+				{
+					provider = _providers[externalConfig.Provider];
+				}
+				else if (config.ProviderLookup.TryGetValue(externalConfig.ProviderConfig, out providerConfig) && !String.IsNullOrEmpty(providerConfig.Provider))
+				{
+					provider = _providers[providerConfig.Provider];
+				}
+				else
+				{
+					throw new InvalidOperationException($"Missing or invalid secret provider identifier for '{externalConfig.ProviderConfig}'");
+				}
 
-				string secret = await provider.GetSecretAsync(externalConfig.Path, cancellationToken);
+				string secret = await provider.GetSecretAsync(externalConfig.Path, providerConfig, cancellationToken);
 				if (externalConfig.Format == ExternalSecretFormat.Text)
 				{
 					data[externalConfig.Key ?? "default"] = secret;
diff --git a/Engine/Source/Programs/Horde/Plugins/Secrets/HordeServer.Secrets/Secrets/SecretConfig.cs b/Engine/Source/Programs/Horde/Plugins/Secrets/HordeServer.Secrets/Secrets/SecretConfig.cs
--- a/Engine/Source/Programs/Horde/Plugins/Secrets/HordeServer.Secrets/Secrets/SecretConfig.cs
+++ b/Engine/Source/Programs/Horde/Plugins/Secrets/HordeServer.Secrets/Secrets/SecretConfig.cs
@@ -29,6 +29,13 @@ namespace HordeServer.Secrets
 		/// </summary>
 		public List<ExternalSecretConfig> Sources { get; set; } = new List<ExternalSecretConfig>();
 
+		/// <summary>
+		/// List of configurations for secret providers
+		/// </summary>
+		public List<SecretProviderConfig> Providers { get; set; } = new List<SecretProviderConfig>();
+
+		internal readonly Dictionary<string, SecretProviderConfig> ProviderLookup = new();
+
 		/// <summary>
 		/// Defines access to this particular secret
 		/// </summary>
@@ -40,6 +47,11 @@ namespace HordeServer.Secrets
 		/// <param name="parentAcl">Parent acl object</param>
 		public void PostLoad(AclConfig parentAcl)
 		{
+			ProviderLookup.Clear();
+			foreach (SecretProviderConfig provider in Providers)
+			{
+				ProviderLookup[provider.Name] = provider;
+			}
 			Acl.PostLoad(parentAcl, $"secret:{Id}", AclConfig.GetActions([typeof(SecretAclAction)]));
 		}
 
@@ -73,6 +85,12 @@ namespace HordeServer.Secrets
 	/// </summary>
 	public class ExternalSecretConfig
 	{
+		/// <summary>
+		/// The name to look up a secret provider config. A config is mutually exclusive to <see cref="Provider"/>
+		/// and to be used when a secret requires additional options than just the name of the provider
+		/// </summary>
+		public string ProviderConfig { get; set; } = String.Empty;
+
 		/// <summary>
 		/// Name of the provider to use
 		/// </summary>
@@ -84,7 +102,7 @@ namespace HordeServer.Secrets
 		public ExternalSecretFormat Format { get; set; }
 
 		/// <summary>
-		/// Optional key indicating the parameter to set in the resulting data array. Required if if <see cref="Format"/> is <see cref="ExternalSecretFormat.Text"/>.
+		/// Optional key indicating the parameter to set in the resulting data array. Required if <see cref="Format"/> is <see cref="ExternalSecretFormat.Text"/>.
 		/// </summary>
 		public string? Key { get; set; }
 
diff --git a/Engine/Source/Programs/Horde/Plugins/Secrets/HordeServer.Secrets/Secrets/SecretProviderConfig.cs b/Engine/Source/Programs/Horde/Plugins/Secrets/HordeServer.Secrets/Secrets/SecretProviderConfig.cs
new file mode 100644
--- /dev/null
+++ b/Engine/Source/Programs/Horde/Plugins/Secrets/HordeServer.Secrets/Secrets/SecretProviderConfig.cs
@@ -0,0 +1,27 @@
+// Copyright Epic Games, Inc. All Rights Reserved.
+
+using HordeServer.Secrets.Providers;
+
+namespace HordeServer.Secrets
+{
+	/// <summary>
+	/// Common settings object for different secret providers
+	/// </summary>
+	public class SecretProviderConfig
+	{
+		/// <summary>
+		/// Name of the secret provider config
+		/// </summary>
+		public string Name { get; set; } = String.Empty;
+
+		/// <summary>
+		/// Name of the provider
+		/// </summary>
+		public string Provider { get; set; } = String.Empty;
+
+		/// <summary>
+		/// Configuration for HCP Vault
+		/// </summary>
+		public HcpVaultConfig? HcpVault { get; set; }
+	}
+}
diff --git a/Engine/Source/Programs/Horde/Plugins/Secrets/HordeServer.Secrets/SecretsPlugin.cs b/Engine/Source/Programs/Horde/Plugins/Secrets/HordeServer.Secrets/SecretsPlugin.cs
--- a/Engine/Source/Programs/Horde/Plugins/Secrets/HordeServer.Secrets/SecretsPlugin.cs
+++ b/Engine/Source/Programs/Horde/Plugins/Secrets/HordeServer.Secrets/SecretsPlugin.cs
@@ -10,7 +10,7 @@ using Microsoft.Extensions.DependencyInjection;
 namespace HordeServer
 {
 	/// <summary>
-	/// Entry point for the storage plugin
+	/// Entry point for the secrets plugin
 	/// </summary>
 	[Plugin("Secrets", GlobalConfigType = typeof(SecretsConfig))]
 	public class SecretsPlugin : IPluginStartup
@@ -22,9 +22,11 @@ namespace HordeServer
 		/// <inheritdoc/>
 		public void ConfigureServices(IServiceCollection services)
 		{
+			services.AddHttpClient();
 			services.AddSingleton<SecretCollectionInternal>();
 			services.AddScoped<ISecretCollection, SecretCollection>();
 			services.AddSingleton<ISecretProvider, AwsParameterStoreSecretProvider>();
+			services.AddSingleton<ISecretProvider, HcpVaultSecretProvider>();
 		}
 	}
 
diff --git a/Engine/Source/Programs/Horde/Scripts/test.sh b/Engine/Source/Programs/Horde/Scripts/test.sh
--- a/Engine/Source/Programs/Horde/Scripts/test.sh
+++ b/Engine/Source/Programs/Horde/Scripts/test.sh
@@ -19,6 +19,7 @@ testProjects=(
 	"Source/Programs/Horde/Plugins/Build/HordeServer.Build.Tests/HordeServer.Build.Tests.csproj"
 	"Source/Programs/Horde/Plugins/Compute/HordeServer.Compute.Tests/HordeServer.Compute.Tests.csproj"
 	"Source/Programs/Horde/Plugins/Ddc/HordeServer.Ddc.Tests/HordeServer.Ddc.Tests.csproj"
+	"Source/Programs/Horde/Plugins/Secrets/HordeServer.Secrets.Tests/HordeServer.Secrets.Tests.csproj"
 	"Source/Programs/Horde/Plugins/Storage/HordeServer.Storage.Tests/HordeServer.Storage.Tests.csproj"
 	"Source/Programs/Horde/Plugins/Tools/HordeServer.Tools.Tests/HordeServer.Tools.Tests.csproj"
 )
