Horde: Update release notes
#rnx

diff --git a/Engine/Source/Programs/Horde/Docs/ReleaseNotes.md b/Engine/Source/Programs/Horde/Docs/ReleaseNotes.md
--- a/Engine/Source/Programs/Horde/Docs/ReleaseNotes.md
+++ b/Engine/Source/Programs/Horde/Docs/ReleaseNotes.md
@@ -1,5 +1,18 @@
 # Release Notes
 
+## 2025-03-13
+
+* Add retry policy for blob writes when getting an upload redirect (40683860)
+* Add issue summary badges to job label overview, capped at 20 (40652838)
+* Perforce server page visual fixes & linking from status page (40650335)
+* Dashboard side of device cleaning suipport (40650241)
+* Updated supportsPartitionedWorkspace doc (40580744)
+* Constrain jupiterdownloads to UGS (40577651)
+* Improve description of OIDC client secret config value to clarify public vs confidential client requirements (40575283)
+* Add tooltips to issue navigation buttons (40572897)
+* Lower log level for socket send buffer teardown from warning to info. This reduces log "pollution" until a better communication channel for Horde-specific diagnostics can be implemented. (40543971)
+* Make parameter boolean checks case insensitive on new build dialog (40540848)
+
 ## 2025-03-06
 
 * Update approved licenses (HTML changes) (40539666)
