Horde: Verify changes exist before attempting to get descriptions in ReleaseNotes tool.
#rnx

diff --git a/Engine/Source/Programs/Horde/Tools/ReleaseNotes/Program.cs b/Engine/Source/Programs/Horde/Tools/ReleaseNotes/Program.cs
--- a/Engine/Source/Programs/Horde/Tools/ReleaseNotes/Program.cs
+++ b/Engine/Source/Programs/Horde/Tools/ReleaseNotes/Program.cs
@@ -124,6 +124,12 @@ namespace CreateReleaseNotes
 			changes.RemoveAll(x => x.User.Equals("buildmachine", StringComparison.OrdinalIgnoreCase));
 			changes = changes.DistinctBy(x => x.Number).ToList();
 
+			if (changes.Count == 0)
+			{
+				logger.LogError("No changes found. Verify server and client CL numbers.");
+				return 1;
+			}
+
 			List<DescribeRecord> describeRecords = await perforce.DescribeAsync(changes.Select(x => x.Number).ToArray());
 
 			HashSet<int> changeNumbers = new HashSet<int>();
