Horde : Fix JSON example

Remove trailing comma.

#rnx

diff --git a/Engine/Source/Programs/Horde/Docs/Config/Secrets.md b/Engine/Source/Programs/Horde/Docs/Config/Secrets.md
--- a/Engine/Source/Programs/Horde/Docs/Config/Secrets.md
+++ b/Engine/Source/Programs/Horde/Docs/Config/Secrets.md
@@ -93,7 +93,7 @@ two secrets from HashiCorp Vault one using an AWS role and the other a pre-share
                     // The Vault API path to the secret e.g. /v1/secret/data/my-secret
                     "path": "path-of-secret-in-vault",
                     // For Vault the format is always json because a secret can contain multiple items
-                    "format": "json",
+                    "format": "json"
                 },
 
                 {
