Horde: Improve Build Health Fix CL messaging when fix CL is not in latest failed stream span

diff --git a/Engine/Source/Programs/Horde/Plugins/Build/HordeServer.Build/Notifications/Sinks/SlackNotificationSink.cs b/Engine/Source/Programs/Horde/Plugins/Build/HordeServer.Build/Notifications/Sinks/SlackNotificationSink.cs
--- a/Engine/Source/Programs/Horde/Plugins/Build/HordeServer.Build/Notifications/Sinks/SlackNotificationSink.cs
+++ b/Engine/Source/Programs/Horde/Plugins/Build/HordeServer.Build/Notifications/Sinks/SlackNotificationSink.cs
@@ -1346,12 +1346,22 @@ namespace HordeServer.Notifications.Sinks
 					else
 					{
 						string fixFailedEventId = $"issue_{issue.Id}_fixfailed_{issue.FixCommitId}";
+
+						IIssueStream? fixStream = issue.Streams.FirstOrDefault(x => x.ContainsFix ?? false);
+
 						string fixFailedMessage = $"Issue not fixed by {FormatChange(issue.FixCommitId)}; see {FormatJobStep(fixFailedSpan.LastFailure, fixFailedSpan.NodeName)} at CL {fixFailedSpan.LastFailure.CommitId} in {fixFailedSpan.StreamName}.";
+
+						if (fixStream != null && fixStream.StreamId != fixFailedSpan.StreamId)
+						{
+							string fixStreamName = spans.FirstOrDefault(x => x.StreamId == fixStream.StreamId)?.StreamName ?? fixStream.StreamId.ToString();
+							fixFailedMessage = $"Issue possibly not fixed by {FormatChange(issue.FixCommitId)}.  The fix CL may be in the process of automatically or manually merging from {fixStreamName} to {fixFailedSpan.StreamName}; see {FormatJobStep(fixFailedSpan.LastFailure, fixFailedSpan.NodeName)} at CL {fixFailedSpan.LastFailure.CommitId} in {fixFailedSpan.StreamName}.";
+						}
+						
 						if (issue.OwnerId.HasValue)
 						{
 							string mention = await FormatMentionAsync(issue.OwnerId.Value, workflow.AllowMentions, cancellationToken);
 							fixFailedMessage += $" ({mention})";
-						}
+						}						
 						await PostSingleMessageToThreadAsync(triageChannel, fixFailedEventId, threadId, fixFailedMessage, cancellationToken);
 					}
 
