Horde: Add compile-time parameter for agent to self-report tool ID

This is useful for self-contained agents that will not use the default multi-platform version. It can also allow workstation-based agents to use a custom version.

diff --git a/Engine/Source/Programs/Horde/HordeAgent/HordeAgent.csproj b/Engine/Source/Programs/Horde/HordeAgent/HordeAgent.csproj
--- a/Engine/Source/Programs/Horde/HordeAgent/HordeAgent.csproj
+++ b/Engine/Source/Programs/Horde/HordeAgent/HordeAgent.csproj
@@ -230,11 +230,13 @@
     <PropertyGroup>
       <TemplateFile>Utility/VersionInfo.cs.template</TemplateFile>
       <OutputFile>$(IntermediateOutputPath)/VersionInfo.cs</OutputFile>
+      <!-- Self-contained versions should set this to match the tool ID for that OS/arch -->
+      <ToolId Condition="'$(ToolId)' == ''">unknown</ToolId>
     </PropertyGroup>
     <ReadLinesFromFile File="$(TemplateFile)">
       <Output TaskParameter="Lines" ItemName="Lines" />
     </ReadLinesFromFile>
-    <WriteLinesToFile File="$(OutputFile)" Lines="@(Lines-&gt;Replace('$VERSION$', '$(InformationalVersion)'))" Overwrite="true" Encoding="UTF-8" />
+    <WriteLinesToFile File="$(OutputFile)" Lines="@(Lines-&gt;Replace('$VERSION$', '$(InformationalVersion)')-&gt;Replace('$TOOL_ID$', '$(ToolId)'))" Overwrite="true" Encoding="UTF-8" />
   </Target>
   
 </Project>
diff --git a/Engine/Source/Programs/Horde/HordeAgent/Utility/VersionInfo.cs.template b/Engine/Source/Programs/Horde/HordeAgent/Utility/VersionInfo.cs.template
--- a/Engine/Source/Programs/Horde/HordeAgent/Utility/VersionInfo.cs.template
+++ b/Engine/Source/Programs/Horde/HordeAgent/Utility/VersionInfo.cs.template
@@ -1,5 +1,7 @@
 // Copyright Epic Games, Inc. All Rights Reserved.
 
+using System.Diagnostics.CodeAnalysis;
+
 namespace HordeAgent.Utility;
 
 /// <summary>
@@ -13,6 +15,11 @@ public static class VersionInfo
 	/// Version (the pattern is replaced prior to compilation)
 	/// </summary>
 	private const string VersionConstant = "$VERSION$";
+	
+	/// <summary>
+	/// Tool ID (the pattern is replaced prior to compilation)
+	/// </summary>
+	private const string ToolIdConstant = "$TOOL_ID$";
 
 	/// <summary>
 	/// Version of this Horde Agent
@@ -21,4 +28,13 @@ public static class VersionInfo
 	{
 		get => String.IsNullOrEmpty(VersionConstant) ? "unknown" : VersionConstant;
 	}
+	
+	/// <summary>
+	/// Self-reported tool ID of this agent. This ID will override and define exact variant of agent software any server-initiated upgrades will use.
+	/// </summary>
+	[SuppressMessage("Maintainability", "CA1508:Avoid dead conditional code", Justification = "Generated code")]
+	public static string? ToolId
+	{
+		get => ToolIdConstant == "unknown" || String.IsNullOrEmpty(ToolIdConstant) ? null : ToolIdConstant;
+	}
 }
\ No newline at end of file
