Horde: Auto-assign Horde issues from Jira issues when possible

diff --git a/Engine/Source/Programs/Horde/Plugins/Build/HordeServer.Build/Issues/External/JiraService.cs b/Engine/Source/Programs/Horde/Plugins/Build/HordeServer.Build/Issues/External/JiraService.cs
--- a/Engine/Source/Programs/Horde/Plugins/Build/HordeServer.Build/Issues/External/JiraService.cs
+++ b/Engine/Source/Programs/Horde/Plugins/Build/HordeServer.Build/Issues/External/JiraService.cs
@@ -106,6 +106,8 @@ namespace HordeServer.Issues.External
 		/// </summary>
 		readonly IssueService _issueService;
 
+		readonly IUserCollection _userCollection;
+
 		readonly ITicker _ticker;
 
 		readonly Uri _jiraUrl;
@@ -119,14 +121,16 @@ namespace HordeServer.Issues.External
 		/// </summary>
 		/// <param name="settings"></param>
 		/// <param name="issueService"></param>
+		/// <param name="userCollection"></param>
 		/// <param name="clock"></param>
 		/// <param name="buildConfig"></param>
 		/// <param name="logger"></param>
-		public JiraService(IOptions<BuildServerConfig> settings, IssueService issueService, IClock clock, IOptionsMonitor<BuildConfig> buildConfig, ILogger<JiraService> logger)
+		public JiraService(IOptions<BuildServerConfig> settings, IssueService issueService, IUserCollection userCollection, IClock clock, IOptionsMonitor<BuildConfig> buildConfig, ILogger<JiraService> logger)
 		{
 			_settings = settings.Value;
 			_logger = logger;
 			_issueService = issueService;
+			_userCollection = userCollection;
 			_ticker = clock.AddTicker<JiraService>(TimeSpan.FromMinutes(2.0), TickAsync, logger);
 			_jiraUrl = _settings.JiraUrl!;
 			_buildConfig = buildConfig;
@@ -321,8 +325,30 @@ namespace HordeServer.Issues.External
 				}
 			}
 
-			await GetIssuesAsync(jiraKeys.ToArray(), cancellationToken);
+			// Auto assign any issues where possible
+			List<IExternalIssue> externalIssues = await GetIssuesAsync(jiraKeys.ToArray(), cancellationToken);
+
+			foreach (IIssue issue in openIssues)
+			{
+				if (String.IsNullOrEmpty(issue.ExternalIssueKey) || issue.OwnerId != null)
+				{
+					continue;
+				}
+
+				IExternalIssue? externalIssue = externalIssues.Where(x => (issue.ExternalIssueKey != null) && (x.Key == issue.ExternalIssueKey)).FirstOrDefault();
 
+				if (externalIssue != null && !String.IsNullOrEmpty(externalIssue.AssigneeEmailAddress))
+				{
+					IUser? user = await _userCollection.FindUserByEmailAsync(externalIssue.AssigneeEmailAddress, cancellationToken:cancellationToken);
+
+					if (user != null)
+					{
+						await _issueService.UpdateIssueAsync(issue.Id, ownerId: user.Id, cancellationToken: cancellationToken);
+
+						_logger.LogInformation("Assigning issue {IssueId} to user {UserId} by matching {EmailAddress}", issue.Id, user.Id, externalIssue.AssigneeEmailAddress);
+					}
+				}
+			}
 		}
 
 		public async Task<(string? key, string? url)> CreateIssueAsync(IUser user, string? externalIssueUser, int issueId, string summary, string projectId, string componentId, string issueType, string? description, string? hordeIssueLink, CancellationToken cancellationToken)
diff --git a/Engine/Source/Programs/Horde/Plugins/Build/HordeServer.Build/Issues/IssueService.cs b/Engine/Source/Programs/Horde/Plugins/Build/HordeServer.Build/Issues/IssueService.cs
--- a/Engine/Source/Programs/Horde/Plugins/Build/HordeServer.Build/Issues/IssueService.cs
+++ b/Engine/Source/Programs/Horde/Plugins/Build/HordeServer.Build/Issues/IssueService.cs
@@ -198,7 +198,6 @@ namespace HordeServer.Issues
 		readonly JobService _jobService;
 		readonly IJobStepRefCollection _jobStepRefs;
 		readonly IIssueCollection _issueCollection;
-		//readonly IExternalIssueService _externalIssueService;
 		readonly ICommitService _commitService;
 		readonly IStreamCollection _streams;
 		readonly IUserCollection _userCollection;
@@ -267,7 +266,6 @@ namespace HordeServer.Issues
 			// Get all the collections
 			_jobService = jobService;
 			_issueCollection = issueCollection;
-			//_externalIssueService = externalIssueService;
 			_commitService = commitService;
 			_jobStepRefs = jobStepRefs;
 			_streams = streams;
@@ -500,9 +498,7 @@ namespace HordeServer.Issues
 			using TelemetrySpan tracerSpan = _tracer.StartActiveSpan($"{nameof(IssueService)}.{nameof(UpdateIssueAsync)}");
 			
 			IIssue? issue;
-			//bool autoAssign = ownerId == null && _externalIssueService != null && !String.IsNullOrEmpty(externalIssueKey);
 
-			//_logger.LogInformation("Updating issue {ArgOne} {ArgTwo} {ArgThree}", ownerId == null,  _externalIssueService != null , !String.IsNullOrEmpty(externalIssueKey));
 			for (; ; )
 			{
 				issue = await _issueCollection.GetIssueAsync(id, cancellationToken);
@@ -511,56 +507,6 @@ namespace HordeServer.Issues
 					return false;
 				}
 
-				/*
-				if (issue.OwnerId != null || issue.ExternalIssueKey == externalIssueKey)
-				{
-					autoAssign = false;
-					_logger.LogInformation("Updating issue {ArgOne} {ArgTwo} external: {ArgThree}", issue.OwnerId != null, issue.ExternalIssueKey == externalIssueKey, externalIssueKey);
-				}
-
-				try
-				{
-					// Auto assign to external issue owner if possible
-					if (autoAssign)
-					{
-						// Get the external issue and auto assign					
-						List<IExternalIssue> externalIssues = await _externalIssueService!.GetIssuesAsync(new string[] { externalIssueKey! }, cancellationToken);
-						if (externalIssues.Count > 0 && !String.IsNullOrEmpty(externalIssues[0].AssigneeEmailAddress))
-						{
-							IUser? matchedUser = await _userCollection.FindUserByEmailAsync(externalIssues[0].AssigneeEmailAddress!, cancellationToken);
-							if (matchedUser != null)
-							{
-								ownerId = matchedUser.Id;
-								_logger.LogInformation("Assigning issue owner to {OwnerId} for external issue {ExternalIssueKey} by matching with {UserEmail}", ownerId, externalIssueKey, externalIssues[0].AssigneeEmailAddress);
-							}
-							else
-							{
-								_logger.LogInformation("Unable to match owner for external issue {ExternalIssueKey} by matching with {UserEmail}", externalIssueKey, externalIssues[0].AssigneeEmailAddress);
-							}
-						}
-						else
-						{
-							if (externalIssues.Count == 0)
-							{
-								_logger.LogInformation("Unable to get external issue {ExternalIssueKey} for auto assign", externalIssueKey);
-							}
-							else
-							{
-								_logger.LogInformation("Unable to get external issue {ExternalIssueKey} email address for auto assign", externalIssueKey);
-							}							
-						}
-					}
-				}
-				catch (Exception ex)
-				{
-					_logger.LogWarning(ex, "Exception while auto-assigning issue {IssueId} to {ExternalId}: {ExceptionMessage}", issue.Id, externalIssueKey, ex.Message);
-				}
-				finally
-				{
-					autoAssign = false;
-				}
-				*/
-
 				issue = await _issueCollection.TryUpdateIssueAsync(issue, initiatedById, new UpdateIssueOptions { UserSummary = summary, Description = description, Promoted = promoted, OwnerId = ownerId ?? resolvedById, NominatedById = nominatedById, DeclinedById = declinedById, Acknowledged = acknowledged, FixCommitId = fixCommitId, FixSystemic = fixSystemic, ResolvedById = resolvedById, ExcludeSpanIds = removeSpanIds, ExternalIssueKey = externalIssueKey, QuarantinedById = quarantinedById, ForceClosedById = forceClosedById, WorkflowThreadUrl = workflowThreadUrl }, cancellationToken: cancellationToken);
 				if (issue != null)
 				{
