Horde: Updating schema documentation

#rnx

diff --git a/Engine/Source/Programs/Horde/Docs/Config/Schema/AclActions.md b/Engine/Source/Programs/Horde/Docs/Config/Schema/AclActions.md
--- a/Engine/Source/Programs/Horde/Docs/Config/Schema/AclActions.md
+++ b/Engine/Source/Programs/Horde/Docs/Config/Schema/AclActions.md
@@ -158,6 +158,16 @@
 | `DeleteAccount` |  |
 | `ViewAccount` |  |
 
+## Agents
+
+| Name | Description |
+| ---- | ----------- |
+| `CreateAgent` | Ability to create an agent. This may be done explicitly, or granted to agents to allow them to self-register. |
+| `UpdateAgent` | Update an agent's name, pools, etc... |
+| `DeleteAgent` | Soft-delete an agent |
+| `ViewAgent` | View an agent |
+| `ListAgents` | List the available agents |
+
 ## Artifacts
 
 | Name | Description |
@@ -176,18 +186,12 @@
 | `UpdateBisectTask` |  |
 | `ViewBisectTask` |  |
 
-## Ddc
+## Compute
 
 | Name | Description |
 | ---- | ----------- |
-| `DdcReadObject` | General read access to refs / blobs and so on |
-| `DdcWriteObject` | General write access to upload refs / blobs etc |
-| `DdcDeleteObject` | Access to delete blobs / refs etc |
-| `DdcDeleteBucket` | Access to delete a particular bucket |
-| `DdcDeleteNamespace` | Access to delete a whole namespace |
-| `DdcReadTransactionLog` | Access to read the transaction log |
-| `DdcWriteTransactionLog` | Access to write the transaction log |
-| `DdcAdminAction` | Access to perform administrative task |
+| `AddComputeTasks` | User can add tasks to the compute cluster |
+| `GetComputeTasks` | User can get and list tasks from the compute cluster |
 
 ## Devices
 
@@ -207,6 +211,24 @@
 | `RetryJobStep` |  |
 | `ViewJob` |  |
 
+## Leases
+
+| Name | Description |
+| ---- | ----------- |
+| `ViewLeases` | View all the leases that an agent has worked on |
+| `ViewLeaseTasks` | View the task data for a lease |
+
+## Logs
+
+| Name | Description |
+| ---- | ----------- |
+| `CreateLog` | Ability to create a log. Implicitly granted to agents. |
+| `UpdateLog` | Ability to update log metadata |
+| `ViewLog` | Ability to view a log contents |
+| `WriteLogData` | Ability to write log data |
+| `CreateEvent` | Ability to create events |
+| `ViewEvent` | Ability to view events |
+
 ## Notices
 
 | Name | Description |
@@ -221,6 +243,16 @@
 | ---- | ----------- |
 | `CreateSubscription` |  |
 
+## Pools
+
+| Name | Description |
+| ---- | ----------- |
+| `CreatePool` | Create a global pool of agents |
+| `UpdatePool` | Modify an agent pool |
+| `DeletePool` | Delete an agent pool |
+| `ViewPool` | Ability to view a pool |
+| `ListPools` | View all the available agent pools |
+
 ## Projects
 
 | Name | Description |
@@ -246,6 +278,21 @@
 | `DeleteAccount` |  |
 | `ViewAccount` |  |
 
+## Sessions
+
+| Name | Description |
+| ---- | ----------- |
+| `CreateSession` | Granted to agents to call CreateSession, which returns a bearer token identifying themselves valid to call UpdateSesssion via gRPC. |
+| `ViewSession` | Allows viewing information about an agent session |
+
+## Software
+
+| Name | Description |
+| ---- | ----------- |
+| `UploadSoftware` | Ability to upload new versions of the agent software |
+| `DownloadSoftware` | Ability to download the agent software |
+| `DeleteSoftware` | Ability to delete agent software |
+
 ## Streams
 
 | Name | Description |
@@ -270,11 +317,11 @@
 
 | Name | Description |
 | ---- | ----------- |
-| `CreateAgent` | Ability to create an agent. This may be done explicitly, or granted to agents to allow them to self-register. |
-| `UpdateAgent` | Update an agent's name, pools, etc... |
-| `DeleteAgent` | Soft-delete an agent |
-| `ViewAgent` | View an agent |
-| `ListAgents` | List the available agents |
+| `CreateAgent` |  |
+| `UpdateAgent` |  |
+| `DeleteAgent` |  |
+| `ViewAgent` |  |
+| `ListAgents` |  |
 
 ## Artifacts
 
@@ -298,21 +345,21 @@
 
 | Name | Description |
 | ---- | ----------- |
-| `AddComputeTasks` | User can add tasks to the compute cluster |
-| `GetComputeTasks` | User can get and list tasks from the compute cluster |
+| `AddComputeTasks` |  |
+| `GetComputeTasks` |  |
 
 ## Ddc
 
 | Name | Description |
 | ---- | ----------- |
-| `DdcReadObject` |  |
-| `DdcWriteObject` |  |
-| `DdcDeleteObject` |  |
-| `DdcDeleteBucket` |  |
-| `DdcDeleteNamespace` |  |
-| `DdcReadTransactionLog` |  |
-| `DdcWriteTransactionLog` |  |
-| `DdcAdminAction` |  |
+| `DdcReadObject` | General read access to refs / blobs and so on |
+| `DdcWriteObject` | General write access to upload refs / blobs etc |
+| `DdcDeleteObject` | Access to delete blobs / refs etc |
+| `DdcDeleteBucket` | Access to delete a particular bucket |
+| `DdcDeleteNamespace` | Access to delete a whole namespace |
+| `DdcReadTransactionLog` | Access to read the transaction log |
+| `DdcWriteTransactionLog` | Access to write the transaction log |
+| `DdcAdminAction` | Access to perform administrative task |
 
 ## Devices
 
@@ -336,19 +383,19 @@
 
 | Name | Description |
 | ---- | ----------- |
-| `ViewLeases` | View all the leases that an agent has worked on |
-| `ViewLeaseTasks` | View the task data for a lease |
+| `ViewLeases` |  |
+| `ViewLeaseTasks` |  |
 
 ## Logs
 
 | Name | Description |
 | ---- | ----------- |
-| `CreateLog` | Ability to create a log. Implicitly granted to agents. |
-| `UpdateLog` | Ability to update log metadata |
-| `ViewLog` | Ability to view a log contents |
-| `WriteLogData` | Ability to write log data |
-| `CreateEvent` | Ability to create events |
-| `ViewEvent` | Ability to view events |
+| `CreateLog` |  |
+| `UpdateLog` |  |
+| `ViewLog` |  |
+| `WriteLogData` |  |
+| `CreateEvent` |  |
+| `ViewEvent` |  |
 
 ## Notices
 
@@ -368,11 +415,11 @@
 
 | Name | Description |
 | ---- | ----------- |
-| `CreatePool` | Create a global pool of agents |
-| `UpdatePool` | Modify an agent pool |
-| `DeletePool` | Delete an agent pool |
-| `ViewPool` | Ability to view a pool |
-| `ListPools` | View all the available agent pools |
+| `CreatePool` |  |
+| `UpdatePool` |  |
+| `DeletePool` |  |
+| `ViewPool` |  |
+| `ListPools` |  |
 
 ## Projects
 
@@ -403,16 +450,16 @@
 
 | Name | Description |
 | ---- | ----------- |
-| `CreateSession` | Granted to agents to call CreateSession, which returns a bearer token identifying themselves valid to call UpdateSesssion via gRPC. |
-| `ViewSession` | Allows viewing information about an agent session |
+| `CreateSession` |  |
+| `ViewSession` |  |
 
 ## Software
 
 | Name | Description |
 | ---- | ----------- |
-| `UploadSoftware` | Ability to upload new versions of the agent software |
-| `DownloadSoftware` | Ability to download the agent software |
-| `DeleteSoftware` | Ability to delete agent software |
+| `UploadSoftware` |  |
+| `DownloadSoftware` |  |
+| `DeleteSoftware` |  |
 
 ## Streams
 
@@ -558,6 +605,12 @@
 | `UpdateReplicator` |  |
 | `ViewReplicator` |  |
 
+## Secrets
+
+| Name | Description |
+| ---- | ----------- |
+| `ViewSecret` | View a credential |
+
 ## ServiceAccounts
 
 | Name | Description |
@@ -730,7 +783,7 @@
 
 | Name | Description |
 | ---- | ----------- |
-| `ViewSecret` | View a credential |
+| `ViewSecret` |  |
 
 ## ServiceAccounts
 
diff --git a/Engine/Source/Programs/Horde/Docs/Config/Schema/Globals.md b/Engine/Source/Programs/Horde/Docs/Config/Schema/Globals.md
--- a/Engine/Source/Programs/Horde/Docs/Config/Schema/Globals.md
+++ b/Engine/Source/Programs/Horde/Docs/Config/Schema/Globals.md
@@ -69,10 +69,10 @@ Name | Description
 ---- | -----------
 `analytics` | [AnalyticsConfig](#analyticsconfig)<br>Configuration for the analytics plugin
 `build` | [BuildConfig](#buildconfig)<br>Configuration for the build plugin
-`ddc` | [EmptyPluginConfig](#emptypluginconfig)<br>Configuration for the ddc plugin
 `compute` | [ComputeConfig](#computeconfig)<br>Configuration for the compute plugin
-`experimental` | [ExperimentalConfig](#experimentalconfig)<br>Configuration for the experimental plugin
+`ddc` | [EmptyPluginConfig](#emptypluginconfig)<br>Configuration for the ddc plugin
 `secrets` | [SecretsConfig](#secretsconfig)<br>Configuration for the secrets plugin
+`experimental` | [ExperimentalConfig](#experimentalconfig)<br>Configuration for the experimental plugin
 `storage` | [StorageConfig](#storageconfig)<br>Configuration for the storage plugin
 `symbols` | [SymbolsConfig](#symbolsconfig)<br>Configuration for the symbols plugin
 `tools` | [ToolsConfig](#toolsconfig)<br>Configuration for the tools plugin
@@ -243,11 +243,6 @@ Name | Description
 `excludeActions` | `string[]`<br>Actions to exclude from the inherited actions
 `extends` | `string[]`<br>Other profiles to extend from
 
-## EmptyPluginConfig
-
-Empty implementation of
-
-
 ## ComputeConfig
 
 Configuration for the compute system
@@ -432,13 +427,10 @@ Name | Description
 `description` | `string`<br>Human-readable description
 `computeId` | `string`<br>Compute ID for this network (used when allocating compute resources)
 
-## ExperimentalConfig
+## EmptyPluginConfig
 
-Configuration for the experimental plugin
+Empty implementation of
 
-Name | Description
----- | -----------
-`notifications` | [NotificationConfig](Notification.md)`[]`<br>Notifications to be sent on the Horde server
 
 ## SecretsConfig
 
@@ -457,6 +449,7 @@ Name | Description
 `id` | `string`<br>Identifier for this secret
 `data` | `string` `->` `string`<br>Key/value pairs associated with this secret
 `sources` | [ExternalSecretConfig](#externalsecretconfig)`[]`<br>Providers to source key/value pairs from
+`providers` | [SecretProviderConfig](#secretproviderconfig)`[]`<br>List of configurations for secret providers
 `acl` | [AclConfig](#aclconfig)<br>Defines access to this particular secret
 
 ## ExternalSecretConfig
@@ -465,9 +458,10 @@ Configuration for an external secret provider
 
 Name | Description
 ---- | -----------
+`providerConfig` | `string`<br>The name to look up a secret provider config. A config is mutually exclusive to and to be used when a secret requires additional options than just the name of the provider
 `provider` | `string`<br>Name of the provider to use
 `format` | [ExternalSecretFormat](#externalsecretformat-enum)<br>Format of the secret
-`key` | `string`<br>Optional key indicating the parameter to set in the resulting data array. Required if if  is .
+`key` | `string`<br>Optional key indicating the parameter to set in the resulting data array. Required if  is .
 `path` | `string`<br>Optional value indicating what to fetch from the provider
 
 ## ExternalSecretFormat (Enum)
@@ -479,6 +473,46 @@ Name | Description
 `Text` | Secret is a plain text value which will be stored using the external secret key
 `Json` | Secret is a JSON formatted string containing key/value pairs
 
+## SecretProviderConfig
+
+Common settings object for different secret providers
+
+Name | Description
+---- | -----------
+`name` | `string`<br>Name of the secret provider config
+`provider` | `string`<br>Name of the provider
+`hcpVault` | [HcpVaultConfig](#hcpvaultconfig)<br>Configuration for HCP Vault
+
+## HcpVaultConfig
+
+Configuration for HCP Vault secrets.
+
+Name | Description
+---- | -----------
+`preSharedKey` | `string`<br>The Vault authentication token to use when the type of credentials is .<br>More information on this token can be found in the API documentation for X-Vault-Token
+`endPoint` | `string`<br>The address of the Vault server.
+`awsArnRole` | `string`<br>The AWS ARN to assume a role when logging into Vault. To be used when the type of credentials is . When null AWS SDK's default credential search will be used.
+`awsIamServerId` | `string`<br>The value for the X-Vault-AWS-IAM-Server-ID header when the Vault server is configured with iam_server_id_header_value. Most likely to be the Vault server's DNS name. To be used when the type of credentials is .
+`role` | `string`<br>The name of the role in the Vault configuration for the login. This role name can be different to the name of the AWS ARN role name. To be used when the type of credentials is . When null the name of IAM principal will be used.
+`credentials` | [HcpVaultCredentialsType](#hcpvaultcredentialstype-enum)<br>Type of credentials to use.
+
+## HcpVaultCredentialsType (Enum)
+
+Credentials to use to log into HCP Vault.
+
+Name | Description
+---- | -----------
+`PreSharedKey` | Use a pre-shared token provided by Vault for authentication.
+`AwsAuth` | Use AWS to authenticate against the Vault server.
+
+## ExperimentalConfig
+
+Configuration for the experimental plugin
+
+Name | Description
+---- | -----------
+`notifications` | [NotificationConfig](Notification.md)`[]`<br>Notifications to be sent on the Horde server
+
 ## StorageConfig
 
 Configuration for storage
