Horde : Fix missing using statement

#rnx

diff --git a/Engine/Source/Programs/Horde/Plugins/Secrets/HordeServer.Secrets/Secrets/Providers/HcpVaultSecretProvider.cs b/Engine/Source/Programs/Horde/Plugins/Secrets/HordeServer.Secrets/Secrets/Providers/HcpVaultSecretProvider.cs
--- a/Engine/Source/Programs/Horde/Plugins/Secrets/HordeServer.Secrets/Secrets/Providers/HcpVaultSecretProvider.cs
+++ b/Engine/Source/Programs/Horde/Plugins/Secrets/HordeServer.Secrets/Secrets/Providers/HcpVaultSecretProvider.cs
@@ -163,7 +163,7 @@ namespace HordeServer.Secrets.Providers
 			// https://developer.hashicorp.com/vault/api-docs/auth/aws#login 
 			Uri requestUri = new(new Uri(config.EndPoint!), "v1/auth/aws/login");
 			HttpClient httpClient = GetHttpClient();
-			HttpResponseMessage httpResponseMessage = await httpClient.PostAsJsonAsync(requestUri, payload, cancellationToken);
+			using HttpResponseMessage httpResponseMessage = await httpClient.PostAsJsonAsync(requestUri, payload, cancellationToken);
 			AuthResponse? response = await httpResponseMessage.Content.ReadFromJsonAsync<AuthResponse>(cancellationToken);
 			if (httpResponseMessage.StatusCode != HttpStatusCode.OK)
 			{
