Horde : HCP Vault Secret Provider
Fix when Vault API returns no JSON content on a non-200 status code.
Add logging for AWS credentials found and assuming to.

#rnx

diff --git a/Engine/Source/Programs/Horde/Plugins/Secrets/HordeServer.Secrets.Tests/Secrets/SecretProviderTests.cs b/Engine/Source/Programs/Horde/Plugins/Secrets/HordeServer.Secrets.Tests/Secrets/SecretProviderTests.cs
--- a/Engine/Source/Programs/Horde/Plugins/Secrets/HordeServer.Secrets.Tests/Secrets/SecretProviderTests.cs
+++ b/Engine/Source/Programs/Horde/Plugins/Secrets/HordeServer.Secrets.Tests/Secrets/SecretProviderTests.cs
@@ -6,6 +6,7 @@ using HordeServer.Secrets;
 using HordeServer.Secrets.Providers;
 using HordeServer.Server;
 using HordeServer.Utilities;
+using Microsoft.Extensions.Logging.Abstractions;
 using Microsoft.Extensions.Options;
 
 namespace HordeServer.Tests.Secrets
@@ -48,7 +49,7 @@ namespace HordeServer.Tests.Secrets
 		{
 			SecretCollectionInternal secretCollectionInternal = new([
 				new InMemoryProvider(),
-				new HcpVaultSecretProvider(new StubHttpClientFactory())
+				new HcpVaultSecretProvider(new StubHttpClientFactory(), NullLogger<HcpVaultSecretProvider>.Instance)
 			]);
 			SecretsConfig secretsConfig = new();
 			secretsConfig.Secrets.Add(new SecretConfig()
diff --git a/Engine/Source/Programs/Horde/Plugins/Secrets/HordeServer.Secrets/Secrets/Providers/HcpVaultSecretProvider.cs b/Engine/Source/Programs/Horde/Plugins/Secrets/HordeServer.Secrets/Secrets/Providers/HcpVaultSecretProvider.cs
--- a/Engine/Source/Programs/Horde/Plugins/Secrets/HordeServer.Secrets/Secrets/Providers/HcpVaultSecretProvider.cs
+++ b/Engine/Source/Programs/Horde/Plugins/Secrets/HordeServer.Secrets/Secrets/Providers/HcpVaultSecretProvider.cs
@@ -12,6 +12,8 @@ using Amazon.Runtime.Internal.Util;
 using Amazon.SecurityToken;
 using Amazon.SecurityToken.Model;
 using Amazon.SecurityToken.Model.Internal.MarshallTransformations;
+using Microsoft.Extensions.Logging;
+using ILogger = Microsoft.Extensions.Logging.ILogger;
 
 namespace HordeServer.Secrets.Providers
 {
@@ -70,8 +72,6 @@ namespace HordeServer.Secrets.Providers
 		{
 			[JsonPropertyName("auth")]
 			public AuthInfo? AuthInfo { get; set; }
-
-			public string[]? Errors { get; set; }
 		}
 
 		/// <summary>
@@ -95,23 +95,23 @@ namespace HordeServer.Secrets.Providers
 			/// The key/value paris of all items for a given path to a secret.
 			/// </summary>
 			public Dictionary<string, Dictionary<string, object>> Data { get; set; } = new();
-
-			public string[]? Errors { get; set; }
 		}
 
 		/// <inheritdoc/>
 		public string Name => "HcpVault";
 
 		private readonly IHttpClientFactory _httpClientFactory;
+		private readonly ILogger _logger;
 		private const string HttpClientName = "HordeSecretProviderHcpVault";
 		private const string AwsAssumeRoleSessionName = "HordeSecretProviderVaultSession";
 
 		/// <summary>
 		/// Constructor
 		/// </summary>
-		public HcpVaultSecretProvider(IHttpClientFactory httpClientFactory)
+		public HcpVaultSecretProvider(IHttpClientFactory httpClientFactory, ILogger<HcpVaultSecretProvider> logger)
 		{
 			_httpClientFactory = httpClientFactory;
+			_logger = logger;
 		}
 
 		private HttpClient GetHttpClient()
@@ -119,20 +119,40 @@ namespace HordeServer.Secrets.Providers
 			return _httpClientFactory.CreateClient(HttpClientName);
 		}
 
-		private static async Task<string> GetAwsSignedHeadersAsync(string? serverId, string? arnRole, CancellationToken cancellationToken)
+		private static async Task<string> GetAwsCallerArnAsync(AWSCredentials credentials, CancellationToken cancellationToken)
+		{
+			using AmazonSecurityTokenServiceClient client = new(credentials);
+			GetCallerIdentityRequest request = new();
+			GetCallerIdentityResponse? response = await client.GetCallerIdentityAsync(request, cancellationToken);
+			return response.Arn;
+		}
+
+		private async Task<ImmutableCredentials> GetAwsCredentialsAsync(string? role, CancellationToken cancellationToken)
 		{
 			ImmutableCredentials credentials;
-			if (arnRole != null)
+			if (!String.IsNullOrEmpty(role))
 			{
+				string originalCaller = await GetAwsCallerArnAsync(FallbackCredentialsFactory.GetCredentials(), cancellationToken);
 				using AmazonSecurityTokenServiceClient client = new();
-				AssumeRoleRequest assumeRoleRequest = new() { RoleArn = arnRole, RoleSessionName = AwsAssumeRoleSessionName };
+				AssumeRoleRequest assumeRoleRequest = new() { RoleArn = role, RoleSessionName = AwsAssumeRoleSessionName };
 				AssumeRoleResponse? assumeRole = await client.AssumeRoleAsync(assumeRoleRequest, cancellationToken);
+				string newCaller = await GetAwsCallerArnAsync(assumeRole.Credentials, cancellationToken);
+				_logger.LogDebug("Caller: {OriginalCaller} to AssumeRole: {NewCaller}", originalCaller, newCaller);
 				credentials = await assumeRole.Credentials.GetCredentialsAsync();
 			}
 			else
 			{
-				credentials = await FallbackCredentialsFactory.GetCredentials().GetCredentialsAsync();
+				AWSCredentials awsCredentials = FallbackCredentialsFactory.GetCredentials();
+				string originalCaller = await GetAwsCallerArnAsync(awsCredentials, cancellationToken);
+				credentials = await awsCredentials.GetCredentialsAsync();
+				_logger.LogDebug("Caller: {OriginalCaller} to AccessKey: {AccessKey}", originalCaller, credentials.AccessKey);
 			}
+			return credentials;
+		}
+	
+		private async Task<string> GetAwsSignedHeadersAsync(string? serverId, string? arnRole, CancellationToken cancellationToken)
+		{
+			ImmutableCredentials credentials = await GetAwsCredentialsAsync(arnRole, cancellationToken);
 			IRequest? request = GetCallerIdentityRequestMarshaller.Instance.Marshall(new GetCallerIdentityRequest());
 			request.ResourcePath = "/";
 			request.Headers.Add("Content-Type", "application/x-www-form-urlencoded; charset=utf-8");
@@ -164,16 +184,12 @@ namespace HordeServer.Secrets.Providers
 			Uri requestUri = new(new Uri(config.EndPoint!), "v1/auth/aws/login");
 			HttpClient httpClient = GetHttpClient();
 			using HttpResponseMessage httpResponseMessage = await httpClient.PostAsJsonAsync(requestUri, payload, cancellationToken);
-			AuthResponse? response = await httpResponseMessage.Content.ReadFromJsonAsync<AuthResponse>(cancellationToken);
 			if (httpResponseMessage.StatusCode != HttpStatusCode.OK)
 			{
-				string errorMessage = $"({httpResponseMessage.StatusCode})";
-				if (response?.Errors?.Length > 0)
-				{
-					errorMessage += $" {String.Join(';', response.Errors).ReplaceLineEndings(String.Empty)}";
-				}
+				string errorMessage = $"Status = {httpResponseMessage.StatusCode} Body={await httpResponseMessage.Content.ReadAsStringAsync(cancellationToken)}";
 				throw new InvalidOperationException($"Unable to log into HCP Vault using AWS {errorMessage}");
 			}
+			AuthResponse? response = await httpResponseMessage.Content.ReadFromJsonAsync<AuthResponse>(cancellationToken);
 			return response?.AuthInfo?.Token ?? throw new InvalidOperationException($"Unable to log into HCP Vault using AWS (No Token Content)");
 		}
 
@@ -184,16 +200,12 @@ namespace HordeServer.Secrets.Providers
 			HttpClient httpClient = GetHttpClient();
 			httpClient.DefaultRequestHeaders.Add("X-Vault-Token", token);
 			using HttpResponseMessage httpResponseMessage = await httpClient.SendAsync(request, cancellationToken);
-			SecretResponse? response = await httpResponseMessage.Content.ReadFromJsonAsync<SecretResponse>(cancellationToken);
 			if (httpResponseMessage.StatusCode != HttpStatusCode.OK)
 			{
-				string errorMessage = $"({httpResponseMessage.StatusCode})";
-				if (response?.Errors?.Length > 0)
-				{
-					errorMessage += $" {String.Join(';', response.Errors).ReplaceLineEndings(String.Empty)}";
-				}
+				string errorMessage = $"Status = {httpResponseMessage.StatusCode} Body={await httpResponseMessage.Content.ReadAsStringAsync(cancellationToken)}";
 				throw new InvalidOperationException($"Unable to fetch secret '{path.AbsolutePath}' from HCP Vault {errorMessage}");
 			}
+			SecretResponse? response = await httpResponseMessage.Content.ReadFromJsonAsync<SecretResponse>(cancellationToken);
 			if (response != null)
 			{
 				if (response.Data.TryGetValue("data", out Dictionary<string, object>? value))
