Horde: Implement retry mechanism for agent capabilities lookup

Intermittent failures in WMI/CIM has been observed which this should help address.

diff --git a/Engine/Source/Programs/Horde/HordeAgent/Services/CapabilitiesService.cs b/Engine/Source/Programs/Horde/HordeAgent/Services/CapabilitiesService.cs
--- a/Engine/Source/Programs/Horde/HordeAgent/Services/CapabilitiesService.cs
+++ b/Engine/Source/Programs/Horde/HordeAgent/Services/CapabilitiesService.cs
@@ -53,13 +53,41 @@ namespace HordeAgent.Services
 			_tracer = tracer;
 			_logger = logger;
 		}
-
+		
+		/// <inheritdoc/>
+		public async Task<RpcAgentCapabilities> GetCapabilitiesAsync(DirectoryReference? workingDir)
+		{
+			const int MaxAttempts = 5;
+			const int RetryDelayMs = 5000;
+			
+			for (int attempt = 0; attempt <= MaxAttempts; attempt++)
+			{
+				try
+				{
+					return await GetCapabilitiesWithTimeoutMonitoringAsync(workingDir);
+				}
+				catch (Exception ex)
+				{
+					if (attempt >= MaxAttempts)
+					{
+						_logger.LogError("GetCapabilities failed {Attempts} of {MaxAttempts} attempts. Giving up.", attempt, MaxAttempts);
+						throw;
+					}
+					
+					_logger.LogWarning("GetCapabilities failed: {Message}. Waiting {DelayMs} ms before trying again...", ex.Message, RetryDelayMs);
+					await Task.Delay(RetryDelayMs);
+				}
+			}
+			
+			throw new InvalidOperationException("Unexpected code path");
+		}
+		
 		/// <summary>
 		/// Gets the hardware capabilities of this worker
 		/// </summary>
 		/// <param name="workingDir">Working directory for the agent</param>
 		/// <returns>Worker object for advertising to the server</returns>
-		public async Task<RpcAgentCapabilities> GetCapabilitiesAsync(DirectoryReference? workingDir)
+		public async Task<RpcAgentCapabilities> GetCapabilitiesWithTimeoutMonitoringAsync(DirectoryReference? workingDir)
 		{
 			try
 			{
