Horde: Ignore failed WMI lookup of Windows-based capabilities

Still see issues with CIM/WMI lookups failing for agents. This is a heavy-handed but temporary fix until we can move off WMI completely.

diff --git a/Engine/Source/Programs/Horde/HordeAgent/Services/CapabilitiesService.cs b/Engine/Source/Programs/Horde/HordeAgent/Services/CapabilitiesService.cs
--- a/Engine/Source/Programs/Horde/HordeAgent/Services/CapabilitiesService.cs
+++ b/Engine/Source/Programs/Horde/HordeAgent/Services/CapabilitiesService.cs
@@ -134,117 +134,16 @@ namespace HordeAgent.Services
 
 				capabilities.Properties.Add($"{KnownPropertyNames.OsFamily}=Windows");
 				capabilities.Properties.Add($"{KnownPropertyNames.OsFamilyCompatibility}=Windows");
-
-				// WMI doesn't work currently on ARM64, due to reliance on a NET Framework assembly.
-				using (CimSession session = CimSession.Create(null))
+				
+				try
 				{
-					const string QueryNamespace = @"root\cimv2";
-					const string QueryDialect = "WQL";
-					
-					{
-						// Add OS info
-						using TelemetrySpan _ = _tracer.StartActiveSpan("GetOsInfo");
-						foreach (CimInstance instance in session.QueryInstances(QueryNamespace, QueryDialect, "select * from Win32_OperatingSystem"))
-						{
-							foreach (CimProperty property in instance.CimInstanceProperties)
-							{
-								string name = property.Name;
-								if (name.Equals("Caption", StringComparison.OrdinalIgnoreCase))
-								{
-									capabilities.Properties.Add($"OSDistribution={property.Value}");
-								}
-								else if (name.Equals("Version", StringComparison.OrdinalIgnoreCase))
-								{
-									capabilities.Properties.Add($"OSKernelVersion={property.Value}");
-								}
-							}
-						}
-					}
-					
-					{
-						// Add CPU info
-						using TelemetrySpan _ = _tracer.StartActiveSpan("GetCpuInfo");
-						Dictionary<string, int> cpuNameToCount = new ();
-						int totalPhysicalCores = 0;
-						int totalEnabledPhysicalCores = 0;
-						int totalLogicalCores = 0;
-						
-						foreach (CimInstance instance in session.QueryInstances(QueryNamespace, QueryDialect, "select * from Win32_Processor"))
-						{
-							foreach (CimProperty property in instance.CimInstanceProperties)
-							{
-								switch (property.Name.ToUpperInvariant())
-								{
-									case "NAME":
-										{
-											string cpuName = property.Value.ToString() ?? String.Empty;
-											cpuNameToCount.TryGetValue(cpuName, out int count);
-											cpuNameToCount[cpuName] = count + 1;
-											break;
-										}
-									case "NUMBEROFCORES" when property.Value is uint c: totalPhysicalCores += (int)c; break;
-									case "NUMBEROFENABLEDCORE" when property.Value is uint c: totalEnabledPhysicalCores += (int)c; break;
-									case "NUMBEROFLOGICALPROCESSORS" when property.Value is uint c: totalLogicalCores += (int)c; break;
-								}
-							}
-						}
-
-						AddCpuInfo(capabilities, cpuNameToCount, totalLogicalCores, Math.Min(totalEnabledPhysicalCores, totalPhysicalCores));
-					}
-					
-					{
-						// Add RAM info
-						using TelemetrySpan _ = _tracer.StartActiveSpan("GetRamInfo");
-						ulong totalCapacity = 0;
-						foreach (CimInstance instance in session.QueryInstances(QueryNamespace, QueryDialect, "select Capacity from Win32_PhysicalMemory"))
-						{
-							foreach (CimProperty property in instance.CimInstanceProperties)
-							{
-								if (property.Name.Equals("Capacity", StringComparison.OrdinalIgnoreCase) && property.Value is ulong capacity)
-								{
-									totalCapacity += capacity;
-								}
-							}
-						}
-
-						capabilities.Resources.Add(KnownResourceNames.Ram, (int)(totalCapacity / (1024 * 1024 * 1024)));
-					}
-					
-					{
-						// Add GPU info
-						using TelemetrySpan _ = _tracer.StartActiveSpan("GetGpuInfo");
-						int index = 0;
-						foreach (CimInstance instance in session.QueryInstances(QueryNamespace, QueryDialect, "select Name, DriverVersion, AdapterRAM from Win32_VideoController"))
-						{
-							string? name = null;
-							string? driverVersion = null;
-							
-							foreach (CimProperty property in instance.CimInstanceProperties)
-							{
-								if (property.Name.Equals("Name", StringComparison.OrdinalIgnoreCase))
-								{
-									name = property.Value.ToString();
-								}
-								else if (property.Name.Equals("DriverVersion", StringComparison.OrdinalIgnoreCase))
-								{
-									driverVersion = property.Value.ToString();
-								}
-							}
-							
-							if (name != null)
-							{
-								string prefix = $"GPU-{++index}";
-								capabilities.Properties.Add($"{prefix}-Name={name}");
-								
-								if (driverVersion != null)
-								{
-									capabilities.Properties.Add($"{prefix}-DriverVersion={driverVersion}");
-								}
-							}
-						}
-					}
+					GetWindowsCapabilitiesFromWmi(capabilities);
 				}
-
+				catch (Exception e)
+				{
+					_logger.LogWarning("Unable to update WIM-based capabilities for Windows (CPU/GPU/memory): {Message}", e.Message);
+				}
+				
 				// Add EC2 properties if needed
 				if (_settings.CurrentValue.EnableAwsEc2Support)
 				{
@@ -523,7 +422,118 @@ namespace HordeAgent.Services
 			capabilities.Properties.AddRange(_settings.CurrentValue.Properties.Select(kvp => $"{kvp.Key}={kvp.Value}"));
 			return capabilities;
 		}
-
+		
+		private void GetWindowsCapabilitiesFromWmi(RpcAgentCapabilities capabilities)
+		{
+			// WMI doesn't work currently on ARM64, due to reliance on a NET Framework assembly.
+			using CimSession session = CimSession.Create(null);
+			const string QueryNamespace = @"root\cimv2";
+			const string QueryDialect = "WQL";
+			
+			{
+				// Add OS info
+				using TelemetrySpan _ = _tracer.StartActiveSpan("GetOsInfo");
+				foreach (CimInstance instance in session.QueryInstances(QueryNamespace, QueryDialect, "select * from Win32_OperatingSystem"))
+				{
+					foreach (CimProperty property in instance.CimInstanceProperties)
+					{
+						string name = property.Name;
+						if (name.Equals("Caption", StringComparison.OrdinalIgnoreCase))
+						{
+							capabilities.Properties.Add($"OSDistribution={property.Value}");
+						}
+						else if (name.Equals("Version", StringComparison.OrdinalIgnoreCase))
+						{
+							capabilities.Properties.Add($"OSKernelVersion={property.Value}");
+						}
+					}
+				}
+			}
+			
+			{
+				// Add CPU info
+				using TelemetrySpan _ = _tracer.StartActiveSpan("GetCpuInfo");
+				Dictionary<string, int> cpuNameToCount = new ();
+				int totalPhysicalCores = 0;
+				int totalEnabledPhysicalCores = 0;
+				int totalLogicalCores = 0;
+				
+				foreach (CimInstance instance in session.QueryInstances(QueryNamespace, QueryDialect, "select * from Win32_Processor"))
+				{
+					foreach (CimProperty property in instance.CimInstanceProperties)
+					{
+						switch (property.Name.ToUpperInvariant())
+						{
+							case "NAME":
+								{
+									string cpuName = property.Value.ToString() ?? String.Empty;
+									cpuNameToCount.TryGetValue(cpuName, out int count);
+									cpuNameToCount[cpuName] = count + 1;
+									break;
+								}
+							case "NUMBEROFCORES" when property.Value is uint c: totalPhysicalCores += (int)c; break;
+							case "NUMBEROFENABLEDCORE" when property.Value is uint c: totalEnabledPhysicalCores += (int)c; break;
+							case "NUMBEROFLOGICALPROCESSORS" when property.Value is uint c: totalLogicalCores += (int)c; break;
+						}
+					}
+				}
+				
+				AddCpuInfo(capabilities, cpuNameToCount, totalLogicalCores, Math.Min(totalEnabledPhysicalCores, totalPhysicalCores));
+			}
+			
+			{
+				// Add RAM info
+				using TelemetrySpan _ = _tracer.StartActiveSpan("GetRamInfo");
+				ulong totalCapacity = 0;
+				foreach (CimInstance instance in session.QueryInstances(QueryNamespace, QueryDialect, "select Capacity from Win32_PhysicalMemory"))
+				{
+					foreach (CimProperty property in instance.CimInstanceProperties)
+					{
+						if (property.Name.Equals("Capacity", StringComparison.OrdinalIgnoreCase) && property.Value is ulong capacity)
+						{
+							totalCapacity += capacity;
+						}
+					}
+				}
+				
+				capabilities.Resources.Add(KnownResourceNames.Ram, (int)(totalCapacity / (1024 * 1024 * 1024)));
+			}
+			
+			{
+				// Add GPU info
+				using TelemetrySpan _ = _tracer.StartActiveSpan("GetGpuInfo");
+				int index = 0;
+				foreach (CimInstance instance in session.QueryInstances(QueryNamespace, QueryDialect, "select Name, DriverVersion, AdapterRAM from Win32_VideoController"))
+				{
+					string? name = null;
+					string? driverVersion = null;
+					
+					foreach (CimProperty property in instance.CimInstanceProperties)
+					{
+						if (property.Name.Equals("Name", StringComparison.OrdinalIgnoreCase))
+						{
+							name = property.Value.ToString();
+						}
+						else if (property.Name.Equals("DriverVersion", StringComparison.OrdinalIgnoreCase))
+						{
+							driverVersion = property.Value.ToString();
+						}
+					}
+					
+					if (name != null)
+					{
+						string prefix = $"GPU-{++index}";
+						capabilities.Properties.Add($"{prefix}-Name={name}");
+						
+						if (driverVersion != null)
+						{
+							capabilities.Properties.Add($"{prefix}-DriverVersion={driverVersion}");
+						}
+					}
+				}
+			}
+		}
+		
 		/// <summary>
 		/// Determine if the current user is a running with elevated permissions on Windows (UAC)
 		/// </summary>
