Horde : Move SecretProviderConfig to SecretsConfig
Move the SecretProviderConfig from being a member of SecretConfig to SecretsConfig, to reduce duplication inside the globals configs

#rnx

diff --git a/Engine/Source/Programs/Horde/Docs/Config/Secrets.md b/Engine/Source/Programs/Horde/Docs/Config/Secrets.md
--- a/Engine/Source/Programs/Horde/Docs/Config/Secrets.md
+++ b/Engine/Source/Programs/Horde/Docs/Config/Secrets.md
@@ -76,9 +76,13 @@ An example secret may be configured as follows:
 
 Secrets are queried from the external provider when requested by a user, and are not cached by Horde.
 
-Secrets from external providers that require additional options, for example HashiCorp Vault to set the endpoint and
-credentials, will use a `providerconfig` property instead of `name` property to provide the data. For example to fetch
-two secrets from HashiCorp Vault one using an AWS role and the other a pre-shared token would be configured as follows:
+Secrets from external providers that require additional options, for example HashiCorp Vault, for example to set the endpoint and
+credentials, will use a `providerConfig` property instead of `name` property to provide the data. For example to fetch
+three secrets from HashiCorp Vault using different authentication methods:
+
+   * Use an AWS ARN to pass to AWS Security Token Service AssumeRole
+   * Use AWS default credential search
+   * A pre-shared token provided by the Vault server
 
    ```json
     "secrets": [
@@ -87,46 +91,26 @@ two secrets from HashiCorp Vault one using an AWS role and the other a pre-share
 
             // Some values read from HashiCorp Vault
             "sources": [
-
                 {
-                    "providerconfig": "hcp-vault-aws-role",
-                    // The Vault API path to the secret e.g. /v1/secret/data/my-secret
-                    "path": "path-of-secret-in-vault",
-                    // For Vault the format is always json because a secret can contain multiple items
+                    "providerConfig": "hcp-vault-aws-assume-role",
+                    // The Vault API path to the secret
+                    "path": "/v1/secret/data/first-secret",
+                    // Format must be json for Vault because a path in Vault returns a JSON document
                     "format": "json"
                 },
 
                 {
-                    "providerconfig": "hcp-vault-pre-shared-key",
+                    "providerConfig": "hcp-vault-aws",
                     "format": "json",
-                    "path": "path-of-secret-in-vault"
-                }
-            ],
-
-            "providers": [
-
-                {
-                    "name": "hcp-vault-aws-role",
-                    "provider": "HcpVault",
-                    "hcpvault": {
-                        "credentials": "awsauth",
-                        "endpoint": "https://vault.example.com",
-                        "awsiamserverid": "vault.example.com",
-                        "awsarnrole": "arn:aws:iam::1234567890:role/vault-example-role",
-                        "role": "vault-role-name"
-                    }
+                    "path": "/v1/secret/data/second-secret",
                 },
 
                 {
-                    "name": "hcp-vault-pre-shared-key",
-                    "provider": "HcpVault",
-                    "hcpvault": {
-                        "credentials": "presharedkey",
-                        "endpoint": "https://vault.example.com",
-                        "presharedkey": "the-key-provided-by-vault"
-                    }
+                    "providerConfig": "hcp-vault-pre-shared-key",
+                    "path": "/v1/secret/data/third-secret",
+                    "format": "json"
                 }
-            ]
+            ],
 
             // Only allow Horde agents to access this
             "acl": {
@@ -143,6 +127,41 @@ two secrets from HashiCorp Vault one using an AWS role and the other a pre-share
                 ]
             }
         }
+    ],
+
+    "secretProviderConfigs": [
+        {
+            "name": "hcp-vault-aws-assume-role",
+            "provider": "HcpVault",
+            "hcpCault": {
+                "credentials": "awsauth",
+                "endpoint": "https://vault.example.com",
+                "awsIamServerId": "vault.example.com",
+                "awsArnRole": "arn:aws:iam::1234567890:role/vault-example-role",
+                "role": "vault-role-name"
+            }
+        },
+
+        {
+            "name": "hcp-vault-aws",
+            "provider": "HcpVault",
+            "hcpVault": {
+                "credentials": "awsauth",
+                "endpoint": "https://vault.example.com",
+                "awsIamServerId": "vault.example.com",
+                "role": "vault-role-name"
+            }
+        },
+
+        {
+            "name": "hcp-vault-pre-shared-key",
+            "provider": "HcpVault",
+            "hcpVault": {
+                "credentials": "presharedkey",
+                "endpoint": "https://vault.example.com",
+                "preSharedKey": "the-key-provided-by-vault"
+            }
+        }
     ]
  ```
 
diff --git a/Engine/Source/Programs/Horde/Plugins/Secrets/HordeServer.Secrets.Tests/Secrets/SecretProviderTests.cs b/Engine/Source/Programs/Horde/Plugins/Secrets/HordeServer.Secrets.Tests/Secrets/SecretProviderTests.cs
--- a/Engine/Source/Programs/Horde/Plugins/Secrets/HordeServer.Secrets.Tests/Secrets/SecretProviderTests.cs
+++ b/Engine/Source/Programs/Horde/Plugins/Secrets/HordeServer.Secrets.Tests/Secrets/SecretProviderTests.cs
@@ -51,106 +51,138 @@ namespace HordeServer.Tests.Secrets
 				new InMemoryProvider(),
 				new HcpVaultSecretProvider(new StubHttpClientFactory(), NullLogger<HcpVaultSecretProvider>.Instance)
 			]);
-			SecretsConfig secretsConfig = new();
-			secretsConfig.Secrets.Add(new SecretConfig()
-			{
-				Id = new SecretId("Secret"),
-				Sources = [
-					new ExternalSecretConfig()
-					{
-						Provider = "InMemory",
-						Path = "secret/username",
-						Key = "username"
-					},
-					new ExternalSecretConfig()
-					{
-						Provider = "InMemory",
-						Path = "secret/password",
-						Key = "password"
-					}
-				]
-			});
-			secretsConfig.Secrets.Add(new SecretConfig()
-			{
-				Id = new SecretId("SecretUsingProviderConfig"),
-				Sources = [
-					new ExternalSecretConfig()
-					{
-						ProviderConfig = "config-id",
-						Path = "secret/username",
-						Key = "username",
-					},
-					new ExternalSecretConfig()
-					{
-						ProviderConfig = "config-id",
-						Path = "secret/password",
-						Key = "password",
-					}
-				],
-				Providers = [
-					new SecretProviderConfig()
-					{
-						Name = "config-id",
-						Provider = "InMemory"
-					}
-				]
-			});
-			secretsConfig.Secrets.Add(new SecretConfig()
-			{
-				Id = new SecretId("SecretUsingHcpVault"),
-				Sources = [
-					new ExternalSecretConfig
-					{
-						ProviderConfig = "pre-shared-key",
-						Path = "/v1/secret/data/example",
-						Format = ExternalSecretFormat.Json
-					}
-				],
-				Providers = [
-					new SecretProviderConfig()
-					{
-						Name = "pre-shared-key",
-						Provider = "HcpVault",
-						HcpVault = new HcpVaultConfig()
+			string json = """
+				{
+					"secrets": [
 						{
-							Credentials = HcpVaultCredentialsType.PreSharedKey,
-							EndPoint = "http://127.0.0.1:8200",
-							PreSharedKey = "replace"
+							"id": "secret",
+							"data": {
+								"username": "j.doe",
+								"password": "password123"
+							}
+						},
+						{
+							"id": "secretUsingProviderConfig",
+							"sources": [
+								{
+									"providerConfig": "in-memory",
+									"path": "secret/username",
+									"key": "username"
+								},
+								{
+									"providerConfig": "in-memory",
+									"path": "secret/password",
+									"key": "password"
+								}
+							]
+						},
+						{
+							"id": "secretUsingHcpVault",
+							"sources": [
+								{
+									"providerConfig": "vault-pre-shared-key",
+									"path": "/v1/secret/data/example",
+									"format": "json"
+								}
+							]
+						},
+						{
+							"id": "missingSecretUsingHcpVault",
+							"sources": [
+								{
+									"providerConfig": "vault-pre-shared-key",
+									"path": "/v1/secret/data/notfound",
+									"format": "json"
+								}
+							]
+						},
+						{
+							"id": "multipleSecretUsingHcpVault",
+							"sources": [
+								{
+									"providerConfig": "vault-pre-shared-key",
+									"path": "/v1/secret/data/first",
+									"format": "json"
+								},
+								{
+									"providerConfig": "vault-pre-shared-key",
+									"path": "/v1/secret/data/second",
+									"format": "json"
+								},
+								{
+									"providerConfig": "vault-pre-shared-key",
+									"path": "/v1/secret/data/third",
+									"format": "json"
+								}
+							]
+						},
+						{
+							"id": "secretUsingHcpVaultAwsAuth",
+							"sources": [
+								{
+									"providerConfig": "vault-aws-assume-role",
+									"path": "/v1/secret/data/example",
+									"format": "json"
+								}
+							]
+						},
+						{
+							"id": "emptySecret",
+							"sources": [
+							]
 						}
-					}
-				]
-			});
-			secretsConfig.Secrets.Add(new SecretConfig()
-			{
-				Id = new SecretId("SecretUsingHcpVaultAwsAuth"),
-				Sources = [
-					new ExternalSecretConfig
-					{
-						ProviderConfig = "vault-aws",
-						Path = "/v1/secret/data/example",
-						Format = ExternalSecretFormat.Json
-					}
-				],
-				Providers = [
-					new SecretProviderConfig()
-					{
-						Name = "vault-aws",
-						Provider = "HcpVault",
-						HcpVault = new HcpVaultConfig(){
-							Credentials = HcpVaultCredentialsType.AwsAuth,
-							EndPoint = "https://vault.example.com/",
-							AwsIamServerId = "vault.example.com",
-							AwsArnRole = "arn:aws:iam::1234567890:role/vault-example-role",
-							Role = "vault-example"
+					],
+					"secretProviderConfigs": [
+						{
+							"name": "in-memory",
+							"provider": "InMemory"
+						},
+						{
+							"name": "vault-pre-shared-key",
+							"provider": "HcpVault",
+							// To set up a local HCP Vault dev server with test data:
+							// 1. install the vault executable https://developer.hashicorp.com/vault/install
+							// 2. in a shell run: vault server -dev
+							//    make a note of the token value called 'Root Token' printed to stdout 
+							// 3. in a second shell set the environment variable VAULT_ADDR
+							//    the addr will likely be http://127.0.0.1:8200
+							// 4. in the second shell login in with the token: vault login
+							// 5. in the second shell add the data:
+							//    vault kv put secret/example username="User01" password="Password123"
+							//    vault kv put secret/first one="1"
+							//    vault kv put secret/second two="2"
+							//    vault kv put secret/third three="3"
+							// 6. update the PreSharedKey property below 
+							"hcpVault": {
+								"credentials": "PreSharedKey",
+								"endpoint": "http://127.0.0.1:8200",
+								"preSharedKey": "replace-with-token"
+							}
+						},
+						{
+							"name": "vault-pre-shared-key-provider",
+							"provider": "HcpVault",
+							"hcpVault": {
+								"credentials": "PreSharedKey",
+								"endpoint": "https://vault.example.com",
+								"preSharedKey": "Abc123Def"
+							}
+						},
+						{
+							"name": "vault-aws-assume-role",
+							"provider": "HcpVault",
+							"hcpVault": {
+								"Credentials": "AwsAuth",
+								"endpoint": "https://vault.example.com/",
+								"awsIamServerId": "vault.example.com",
+								"awsArnRole": "arn:aws:iam::1234567890:role/vault-example-role",
+								"role": "vault-example"
+							}
 						}
-					}
-				]
-			});
-			secretsConfig.Secrets.Add(new SecretConfig()
-			{
-				Id = new SecretId("EmptySecret"),
-				Sources = [new ExternalSecretConfig() ]
-			});
+					]
+				}
+			""";
+			SecretsConfig secretsConfig = JsonSerializer.Deserialize<SecretsConfig>(json, JsonUtils.DefaultSerializerOptions)!;
 			SecretCollection secretCollection = new(secretCollectionInternal, new OptionsSnapshotStub<SecretsConfig>(secretsConfig));
 			GlobalConfig config = new();
 			config.Plugins.AddSecretsConfig(secretsConfig);
@@ -163,8 +195,8 @@ namespace HordeServer.Tests.Secrets
 		{
 			SecretCollection secretCollection = CreateSecretCollection(); 
 			ISecret? secret = await secretCollection.GetAsync(new SecretId("Secret"), CancellationToken.None);
-			Assert.AreEqual(secret?.Data["username"], "user01");
-			Assert.AreEqual(secret?.Data["password"], "password01");
+			Assert.AreEqual(secret?.Data["username"], "j.doe");
+			Assert.AreEqual(secret?.Data["password"], "password123");
 		}
 
 		[TestMethod]
@@ -177,31 +209,45 @@ namespace HordeServer.Tests.Secrets
 		}
 
 		[TestMethod]
-		public async Task InvalidGetSecretAsync()
+		public async Task GetEmptySecretAsync()
 		{
 			SecretCollection secretCollection = CreateSecretCollection();
-			await Assert.ThrowsExceptionAsync<InvalidOperationException>(() => secretCollection.GetAsync(new SecretId("EmptySecret"), CancellationToken.None));
+			ISecret? secret = await secretCollection.GetAsync(new SecretId("EmptySecret"), CancellationToken.None);
+			Assert.AreEqual(secret?.Data.Count, 0);
 		}
 
 		[Ignore("Requires access to a running HCP Vault server and a pre-shared key")]
 		[TestMethod]
 		public async Task GetSecretFromVaultAsync()
 		{
-			// To set up a local dev server with test data:
-			// 1. install the vault executable https://developer.hashicorp.com/vault/install
-			// 2. in a shell run: vault server -dev
-			//    make a note of the token value called 'Root Token' printed to stdout 
-			// 3. in a second shell set the environment variable VAULT_ADDR
-			//    the addr will likely be http://127.0.0.1:8200
-			// 4. in the second shell login in with the token: vault login
-			// 5. in the second shell add the data: vault kv put secret/example username="User01" password="Password123"
-			// 6. update the SecretProviderConfig options for the SecretId used in this test 
 			SecretCollection secretCollection = CreateSecretCollection(); 
 			ISecret? secret = await secretCollection.GetAsync(new SecretId("SecretUsingHcpVault"), CancellationToken.None);
 			Assert.AreEqual(secret?.Data["username"], "User01");
 			Assert.AreEqual(secret?.Data["password"], "Password123");
 		}
 
+		[Ignore("Requires access to a running HCP Vault server and a pre-shared key")]
+		[TestMethod]
+		public async Task GetMultipleSecretFromVaultAsync()
+		{
+			SecretCollection secretCollection = CreateSecretCollection(); 
+			ISecret? secret = await secretCollection.GetAsync(new SecretId("MultipleSecretUsingHcpVault"), CancellationToken.None);
+			Assert.AreEqual(secret?.Data["one"], "1");
+			Assert.AreEqual(secret?.Data["two"], "2");
+			Assert.AreEqual(secret?.Data["three"], "3");
+		}
+
+		[Ignore("Requires access to a running HCP Vault server and a pre-shared key")]
+		[TestMethod]
+		public async Task SecretNotFoundFromVaultAsync()
+		{
+			SecretCollection secretCollection = CreateSecretCollection(); 
+			InvalidOperationException ex = await Assert.ThrowsExceptionAsync<InvalidOperationException>(() =>
+				secretCollection.GetAsync(new SecretId("MissingSecretUsingHcpVault"), CancellationToken.None)
+			);
+			Assert.AreEqual(ex.Message, "Unable to fetch secret '/v1/secret/data/notfound' from HCP Vault Status = NotFound Body={\"errors\":[]}\n");
+		}
+
 		[Ignore("Requires access to a running HCP Vault server and supported AWS IAM")]
 		[TestMethod]
 		public async Task GetSecretFromVaultWithAwsAuthAsync()
@@ -217,80 +263,30 @@ namespace HordeServer.Tests.Secrets
 		}
 
 		[TestMethod]
-		public void GetSecretConfigFromJson()
+		public async Task GetSecretFromMinimalJsonAsync()
 		{
 			string json = """
 				{
 					"secrets": [
 						{
-							"id": "my-secrets",
-							"sources":[
-								{
-									"ProviderConfig": "vault-aws-provider",
-									"path": "path-to-secret",
-									"format": "json"
-								},
-								{
-									"ProviderConfig": "vault-pre-shared-key-provider",
-									"path": "path-to-secret",
-									"format": "json"
-								}
-							],
-							"providers": [
-								{
-									"Name": "vault-aws-provider",
-									"Provider": "HcpVault",
-									"HcpVault": {
-										"Credentials": "awsauth",
-										"EndPoint": "https://vault.example.com",
-										"AwsIamServerId": "vault.example.com",
-										"AwsArnRole": "arn:aws:iam::1234567890:role/vault-example-role",
-										"Role": "vault-role-name"
-									}
-								},
-								{
-									"Name": "vault-pre-shared-key-provider",
-									"Provider": "HcpVault",
-									"HcpVault": {
-										"credentials": "presharedkey",
-										"endpoint": "https://vault.example.com",
-										"presharedkey": "Abc123Def"
-									}
-								}
-							]
+							"id": "basic-secret",
+							"data": {
+								"username": "j.doe",
+								"password": "password123"
+							}
 						}
 					]
 				}
 			""";
 			SecretsConfig secretsConfig = JsonSerializer.Deserialize<SecretsConfig>(json, JsonUtils.DefaultSerializerOptions)!;
-			SecretConfig secretConfig = secretsConfig.Secrets[0];
-			Assert.AreEqual(secretConfig.Id.ToString(), "my-secrets");
-
-			ExternalSecretConfig externalSecretConfig0 = secretConfig.Sources[0];
-			SecretProviderConfig secretProviderConfig0 = secretConfig.Providers[0];
-			HcpVaultConfig hcpVaultConfig0 = secretProviderConfig0.HcpVault!;
-			Assert.AreEqual(externalSecretConfig0.Format, ExternalSecretFormat.Json);
-			Assert.AreEqual(externalSecretConfig0.Path, "path-to-secret");
-			Assert.AreEqual(externalSecretConfig0.ProviderConfig, "vault-aws-provider");
-			Assert.AreEqual(secretProviderConfig0.Name, "vault-aws-provider");
-			Assert.AreEqual(secretProviderConfig0.Provider, "HcpVault");
-			Assert.AreEqual(hcpVaultConfig0.EndPoint, "https://vault.example.com");
-			Assert.AreEqual(hcpVaultConfig0.Credentials, HcpVaultCredentialsType.AwsAuth);
-			Assert.AreEqual(hcpVaultConfig0.AwsArnRole, "arn:aws:iam::1234567890:role/vault-example-role");
-			Assert.AreEqual(hcpVaultConfig0.AwsIamServerId, "vault.example.com");
-			Assert.AreEqual(hcpVaultConfig0.Role, "vault-role-name");
-
-			ExternalSecretConfig externalSecretConfig1 = secretConfig.Sources[1];
-			SecretProviderConfig secretProviderConfig1 = secretConfig.Providers[1];
-			HcpVaultConfig hcpVaultConfig1 = secretProviderConfig1.HcpVault!;
-			Assert.AreEqual(externalSecretConfig1.Format, ExternalSecretFormat.Json);
-			Assert.AreEqual(externalSecretConfig1.Path, "path-to-secret");
-			Assert.AreEqual(externalSecretConfig1.ProviderConfig, "vault-pre-shared-key-provider");
-			Assert.AreEqual(secretProviderConfig1.Name, "vault-pre-shared-key-provider");
-			Assert.AreEqual(secretProviderConfig1.Provider, "HcpVault");
-			Assert.AreEqual(hcpVaultConfig1.EndPoint, "https://vault.example.com");
-			Assert.AreEqual(hcpVaultConfig1.Credentials, HcpVaultCredentialsType.PreSharedKey);
-			Assert.AreEqual(hcpVaultConfig1.PreSharedKey, "Abc123Def");
+			SecretCollectionInternal secretCollectionInternal = new([]);
+			SecretCollection secretCollection = new(secretCollectionInternal, new OptionsSnapshotStub<SecretsConfig>(secretsConfig));
+			GlobalConfig config = new();
+			config.Plugins.AddSecretsConfig(secretsConfig);
+			config.PostLoad(new ServerSettings(), [], []);
+			ISecret? secret = await secretCollection.GetAsync(new SecretId("basic-secret"), CancellationToken.None);
+			Assert.AreEqual(secret?.Data["username"], "j.doe");
+			Assert.AreEqual(secret?.Data["password"], "password123");
 		}
 	}
 }
diff --git a/Engine/Source/Programs/Horde/Plugins/Secrets/HordeServer.Secrets/Secrets/SecretCollection.cs b/Engine/Source/Programs/Horde/Plugins/Secrets/HordeServer.Secrets/Secrets/SecretCollection.cs
--- a/Engine/Source/Programs/Horde/Plugins/Secrets/HordeServer.Secrets/Secrets/SecretCollection.cs
+++ b/Engine/Source/Programs/Horde/Plugins/Secrets/HordeServer.Secrets/Secrets/SecretCollection.cs
@@ -31,7 +31,12 @@ namespace HordeServer.Secrets
 			{
 				return null;
 			}
-			return await _secretCollectionInternal.GetAsync(secretConfig, cancellationToken);
+
+			// secret provider configs are optional
+			Dictionary<string, SecretProviderConfig>? secretProviderConfigs;
+			_secretsConfig.Value.TryGetSecretProviderConfigs(secretId, out secretProviderConfigs);
+
+			return await _secretCollectionInternal.GetAsync(secretConfig, secretProviderConfigs, cancellationToken);
 		}
 	}
 
@@ -66,8 +71,9 @@ namespace HordeServer.Secrets
 		/// Resolve a secret to concrete values
 		/// </summary>
 		/// <param name="config">Configuration for the secret</param>
+		/// <param name="secretProviderConfigs">Secret provider configs for the secret</param>
 		/// <param name="cancellationToken">Cancellation token for the operation</param>
-		public async Task<ISecret> GetAsync(SecretConfig config, CancellationToken cancellationToken)
+		public async Task<ISecret> GetAsync(SecretConfig config, Dictionary<string, SecretProviderConfig>? secretProviderConfigs, CancellationToken cancellationToken)
 		{
 			// Add the hard-coded secrets
 			Dictionary<string, string> data = new Dictionary<string, string>();
@@ -85,7 +91,7 @@ namespace HordeServer.Secrets
 				{
 					provider = _providers[externalConfig.Provider];
 				}
-				else if (config.ProviderLookup.TryGetValue(externalConfig.ProviderConfig, out providerConfig) && !String.IsNullOrEmpty(providerConfig.Provider))
+				else if (secretProviderConfigs != null && secretProviderConfigs.TryGetValue(externalConfig.ProviderConfig, out providerConfig) && !String.IsNullOrEmpty(providerConfig.Provider))
 				{
 					provider = _providers[providerConfig.Provider];
 				}
diff --git a/Engine/Source/Programs/Horde/Plugins/Secrets/HordeServer.Secrets/Secrets/SecretConfig.cs b/Engine/Source/Programs/Horde/Plugins/Secrets/HordeServer.Secrets/Secrets/SecretConfig.cs
--- a/Engine/Source/Programs/Horde/Plugins/Secrets/HordeServer.Secrets/Secrets/SecretConfig.cs
+++ b/Engine/Source/Programs/Horde/Plugins/Secrets/HordeServer.Secrets/Secrets/SecretConfig.cs
@@ -29,13 +29,6 @@ namespace HordeServer.Secrets
 		/// </summary>
 		public List<ExternalSecretConfig> Sources { get; set; } = new List<ExternalSecretConfig>();
 
-		/// <summary>
-		/// List of configurations for secret providers
-		/// </summary>
-		public List<SecretProviderConfig> Providers { get; set; } = new List<SecretProviderConfig>();
-
-		internal readonly Dictionary<string, SecretProviderConfig> ProviderLookup = new();
-
 		/// <summary>
 		/// Defines access to this particular secret
 		/// </summary>
@@ -47,11 +40,6 @@ namespace HordeServer.Secrets
 		/// <param name="parentAcl">Parent acl object</param>
 		public void PostLoad(AclConfig parentAcl)
 		{
-			ProviderLookup.Clear();
-			foreach (SecretProviderConfig provider in Providers)
-			{
-				ProviderLookup[provider.Name] = provider;
-			}
 			Acl.PostLoad(parentAcl, $"secret:{Id}", AclConfig.GetActions([typeof(SecretAclAction)]));
 		}
 
diff --git a/Engine/Source/Programs/Horde/Plugins/Secrets/HordeServer.Secrets/SecretsConfig.cs b/Engine/Source/Programs/Horde/Plugins/Secrets/HordeServer.Secrets/SecretsConfig.cs
--- a/Engine/Source/Programs/Horde/Plugins/Secrets/HordeServer.Secrets/SecretsConfig.cs
+++ b/Engine/Source/Programs/Horde/Plugins/Secrets/HordeServer.Secrets/SecretsConfig.cs
@@ -23,6 +23,13 @@ namespace HordeServer
 
 		private readonly Dictionary<SecretId, SecretConfig> _secretLookup = new Dictionary<SecretId, SecretConfig>();
 
+		/// <summary>
+		/// List of configurations for secret providers
+		/// </summary>
+		public List<SecretProviderConfig> SecretProviderConfigs { get; set; } = new();
+
+		private readonly Dictionary<SecretId, Dictionary<string, SecretProviderConfig>> _secretProviderConfigLookup = new();
+
 		/// <inheritdoc/>
 		public void PostLoad(PluginConfigOptions configOptions)
 		{
@@ -32,6 +39,30 @@ namespace HordeServer
 				_secretLookup.Add(secret.Id, secret);
 				secret.PostLoad(configOptions.ParentAcl);
 			}
+
+			Dictionary<string, SecretProviderConfig> secretProviderConfigNameLookup = new();
+			foreach (SecretProviderConfig secretProviderConfig in SecretProviderConfigs)
+			{
+				secretProviderConfigNameLookup[secretProviderConfig.Name] = secretProviderConfig;
+			}
+			_secretProviderConfigLookup.Clear();
+			foreach (KeyValuePair<SecretId, SecretConfig> secret in _secretLookup)
+			{
+				foreach (ExternalSecretConfig externalConfig in secret.Value.Sources)
+				{
+					if (!String.IsNullOrEmpty(externalConfig.ProviderConfig))
+					{
+						if (!_secretProviderConfigLookup.ContainsKey(secret.Key))
+						{
+							_secretProviderConfigLookup[secret.Key] = new Dictionary<string, SecretProviderConfig>();
+						}
+						if (secretProviderConfigNameLookup.TryGetValue(externalConfig.ProviderConfig, out SecretProviderConfig? secretProviderConfig))
+						{
+							_secretProviderConfigLookup[secret.Key].TryAdd(externalConfig.ProviderConfig, secretProviderConfig);
+						}
+					}
+				}
+			}
 		}
 
 		/// <summary>
@@ -43,6 +74,15 @@ namespace HordeServer
 		public bool TryGetSecret(SecretId secretId, [NotNullWhen(true)] out SecretConfig? config)
 			=> _secretLookup.TryGetValue(secretId, out config);
 
+		/// <summary>
+		/// Attempts to get a secret provider configuration from this object
+		/// </summary>
+		/// <param name="secretId">Secret id</param>
+		/// <param name="config">Receives the secret provider configurations on success</param>
+		/// <returns>True on success</returns>
+		public bool TryGetSecretProviderConfigs(SecretId secretId, [NotNullWhen(true)] out Dictionary<string, SecretProviderConfig>? config)
+			=> _secretProviderConfigLookup.TryGetValue(secretId, out config);
+
 		/// <summary>
 		/// Authorize access to a secret
 		/// </summary>
