Horde: Replace use of WMI/CIM with Win32 API calls for resolving Windows-specific capabilities

CIM has proven problematic on Epic's buildfarm and raw Win32 API calls are also much faster.

diff --git a/Engine/Source/Programs/Horde/HordeAgent/Services/CapabilitiesService.cs b/Engine/Source/Programs/Horde/HordeAgent/Services/CapabilitiesService.cs
--- a/Engine/Source/Programs/Horde/HordeAgent/Services/CapabilitiesService.cs
+++ b/Engine/Source/Programs/Horde/HordeAgent/Services/CapabilitiesService.cs
@@ -16,7 +16,6 @@ using EpicGames.Horde.Compute;
 using HordeCommon.Rpc.Messages;
 using Microsoft.Extensions.Logging;
 using Microsoft.Extensions.Options;
-using Microsoft.Management.Infrastructure;
 using OpenTelemetry.Trace;
 using AddressFamily = System.Net.Sockets.AddressFamily;
 
@@ -32,6 +31,11 @@ namespace HordeAgent.Services
 		Task<RpcAgentCapabilities> GetCapabilitiesAsync(DirectoryReference? workingDir);
 	}
 	
+	/// <summary>
+	/// Marker interface to separate system-specific capability services during dependency-injection
+	/// </summary>
+	internal interface ISystemCapabilities : ICapabilitiesService;
+	
 	/// <summary>
 	/// Implements the message handling loop for an agent. Runs asynchronously until disposed.
 	/// </summary>
@@ -126,35 +130,15 @@ namespace HordeAgent.Services
 
 			if (RuntimeInformation.IsOSPlatform(OSPlatform.Windows))
 			{
-				using TelemetrySpan winSpan = _tracer.StartActiveSpan("GetWindowsCapabilities");
-				capabilities.Properties.Add($"{KnownPropertyNames.Platform}=Win64");
-				capabilities.Properties.Add($"{KnownPropertyNames.PlatformGroup}=Windows");
-				capabilities.Properties.Add($"{KnownPropertyNames.PlatformGroup}=Microsoft");
-				capabilities.Properties.Add($"{KnownPropertyNames.PlatformGroup}=Desktop");
-
-				capabilities.Properties.Add($"{KnownPropertyNames.OsFamily}=Windows");
-				capabilities.Properties.Add($"{KnownPropertyNames.OsFamilyCompatibility}=Windows");
-				
-				try
-				{
-					GetWindowsCapabilitiesFromWmi(capabilities);
-				}
-				catch (Exception e)
+				// This should be injected as ISystemCapabilities but as failing WMI/CIM lookups must be addressed,
+				// WindowsCapabilities is manually initialized for now.
+				WindowsCapabilities winSystemCaps = new (_settings, _logger, _tracer);
+				RpcAgentCapabilities winCaps = await winSystemCaps.GetCapabilitiesAsync(workingDir);
+				capabilities.Properties.AddRange(winCaps.Properties);
+				foreach (KeyValuePair<string, int> entry in winCaps.Resources)
 				{
-					_logger.LogWarning("Unable to update WIM-based capabilities for Windows (CPU/GPU/memory): {Message}", e.Message);
+					capabilities.Resources[entry.Key] = entry.Value;
 				}
-				
-				// Add EC2 properties if needed
-				if (_settings.CurrentValue.EnableAwsEc2Support)
-				{
-					await AddAwsPropertiesAsync(capabilities.Properties, logger);
-				}
-
-				// Add session information
-				capabilities.Properties.Add($"User={Environment.UserName}");
-				capabilities.Properties.Add($"Domain={Environment.UserDomainName}");
-				capabilities.Properties.Add($"Interactive={Environment.UserInteractive}");
-				capabilities.Properties.Add($"Elevated={IsUserAdministrator()}");
 			}
 			else if (RuntimeInformation.IsOSPlatform(OSPlatform.Linux))
 			{
@@ -166,12 +150,6 @@ namespace HordeAgent.Services
 
 				capabilities.Properties.Add($"{KnownPropertyNames.OsFamily}=Linux");
 				capabilities.Properties.Add($"{KnownPropertyNames.OsFamilyCompatibility}=Linux");
-
-				// Add EC2 properties if needed
-				if (_settings.CurrentValue.EnableAwsEc2Support)
-				{
-					await AddAwsPropertiesAsync(capabilities.Properties, logger);
-				}
 				
 				if (_settings.CurrentValue.WineExecutablePath != null)
 				{
@@ -193,10 +171,10 @@ namespace HordeAgent.Services
 					}
 
 					Dictionary<string, int> nameToCount = new Dictionary<string, int>(StringComparer.Ordinal);
-					foreach (string cpuName in cpuNames.Values)
+					foreach (string cpuNameTemp in cpuNames.Values)
 					{
-						nameToCount.TryGetValue(cpuName, out int count);
-						nameToCount[cpuName] = count + 1;
+						nameToCount.TryGetValue(cpuNameTemp, out int count);
+						nameToCount[cpuNameTemp] = count + 1;
 					}
 
 					HashSet<string> logicalCores = new HashSet<string>();
@@ -212,8 +190,9 @@ namespace HordeAgent.Services
 							physicalCores.Add(physicalCoreId);
 						}
 					}
-
-					AddCpuInfo(capabilities, nameToCount, logicalCores.Count, physicalCores.Count);
+					
+					string cpuName = cpuNames.Values.FirstOrDefault() ?? "Unknown CPU";
+					AddCpuInfo(_settings.CurrentValue, capabilities, cpuName, logicalCores.Count, physicalCores.Count);
 				}
 
 				// Parse the RAM info
@@ -319,6 +298,12 @@ namespace HordeAgent.Services
 					}
 				}
 			}
+			
+			// Add EC2 properties if needed
+			if (_settings.CurrentValue.EnableAwsEc2Support)
+			{
+				await AddAwsPropertiesAsync(capabilities.Properties, logger);
+			}
 
 			// Get the IP addresses
 			try
@@ -423,117 +408,6 @@ namespace HordeAgent.Services
 			return capabilities;
 		}
 		
-		private void GetWindowsCapabilitiesFromWmi(RpcAgentCapabilities capabilities)
-		{
-			// WMI doesn't work currently on ARM64, due to reliance on a NET Framework assembly.
-			using CimSession session = CimSession.Create(null);
-			const string QueryNamespace = @"root\cimv2";
-			const string QueryDialect = "WQL";
-			
-			{
-				// Add OS info
-				using TelemetrySpan _ = _tracer.StartActiveSpan("GetOsInfo");
-				foreach (CimInstance instance in session.QueryInstances(QueryNamespace, QueryDialect, "select * from Win32_OperatingSystem"))
-				{
-					foreach (CimProperty property in instance.CimInstanceProperties)
-					{
-						string name = property.Name;
-						if (name.Equals("Caption", StringComparison.OrdinalIgnoreCase))
-						{
-							capabilities.Properties.Add($"OSDistribution={property.Value}");
-						}
-						else if (name.Equals("Version", StringComparison.OrdinalIgnoreCase))
-						{
-							capabilities.Properties.Add($"OSKernelVersion={property.Value}");
-						}
-					}
-				}
-			}
-			
-			{
-				// Add CPU info
-				using TelemetrySpan _ = _tracer.StartActiveSpan("GetCpuInfo");
-				Dictionary<string, int> cpuNameToCount = new ();
-				int totalPhysicalCores = 0;
-				int totalEnabledPhysicalCores = 0;
-				int totalLogicalCores = 0;
-				
-				foreach (CimInstance instance in session.QueryInstances(QueryNamespace, QueryDialect, "select * from Win32_Processor"))
-				{
-					foreach (CimProperty property in instance.CimInstanceProperties)
-					{
-						switch (property.Name.ToUpperInvariant())
-						{
-							case "NAME":
-								{
-									string cpuName = property.Value.ToString() ?? String.Empty;
-									cpuNameToCount.TryGetValue(cpuName, out int count);
-									cpuNameToCount[cpuName] = count + 1;
-									break;
-								}
-							case "NUMBEROFCORES" when property.Value is uint c: totalPhysicalCores += (int)c; break;
-							case "NUMBEROFENABLEDCORE" when property.Value is uint c: totalEnabledPhysicalCores += (int)c; break;
-							case "NUMBEROFLOGICALPROCESSORS" when property.Value is uint c: totalLogicalCores += (int)c; break;
-						}
-					}
-				}
-				
-				AddCpuInfo(capabilities, cpuNameToCount, totalLogicalCores, Math.Min(totalEnabledPhysicalCores, totalPhysicalCores));
-			}
-			
-			{
-				// Add RAM info
-				using TelemetrySpan _ = _tracer.StartActiveSpan("GetRamInfo");
-				ulong totalCapacity = 0;
-				foreach (CimInstance instance in session.QueryInstances(QueryNamespace, QueryDialect, "select Capacity from Win32_PhysicalMemory"))
-				{
-					foreach (CimProperty property in instance.CimInstanceProperties)
-					{
-						if (property.Name.Equals("Capacity", StringComparison.OrdinalIgnoreCase) && property.Value is ulong capacity)
-						{
-							totalCapacity += capacity;
-						}
-					}
-				}
-				
-				capabilities.Resources.Add(KnownResourceNames.Ram, (int)(totalCapacity / (1024 * 1024 * 1024)));
-			}
-			
-			{
-				// Add GPU info
-				using TelemetrySpan _ = _tracer.StartActiveSpan("GetGpuInfo");
-				int index = 0;
-				foreach (CimInstance instance in session.QueryInstances(QueryNamespace, QueryDialect, "select Name, DriverVersion, AdapterRAM from Win32_VideoController"))
-				{
-					string? name = null;
-					string? driverVersion = null;
-					
-					foreach (CimProperty property in instance.CimInstanceProperties)
-					{
-						if (property.Name.Equals("Name", StringComparison.OrdinalIgnoreCase))
-						{
-							name = property.Value.ToString();
-						}
-						else if (property.Name.Equals("DriverVersion", StringComparison.OrdinalIgnoreCase))
-						{
-							driverVersion = property.Value.ToString();
-						}
-					}
-					
-					if (name != null)
-					{
-						string prefix = $"GPU-{++index}";
-						capabilities.Properties.Add($"{prefix}-Name={name}");
-						
-						if (driverVersion != null)
-						{
-							capabilities.Properties.Add($"{prefix}-DriverVersion={driverVersion}");
-						}
-					}
-				}
-			}
-		}
-		
 		/// <summary>
 		/// Determine if the current user is a running with elevated permissions on Windows (UAC)
 		/// </summary>
@@ -546,11 +420,9 @@ namespace HordeAgent.Services
 
 			try
 			{
-				using (WindowsIdentity identity = WindowsIdentity.GetCurrent())
-				{
-					WindowsPrincipal principal = new WindowsPrincipal(identity);
-					return principal.IsInRole(WindowsBuiltInRole.Administrator);
-				}
+				using WindowsIdentity identity = WindowsIdentity.GetCurrent();
+				WindowsPrincipal principal = new (identity);
+				return principal.IsInRole(WindowsBuiltInRole.Administrator);
 			}
 			catch
 			{
@@ -588,18 +460,16 @@ namespace HordeAgent.Services
 			}
 		}
 
-		void AddCpuInfo(RpcAgentCapabilities capabilities, Dictionary<string, int> nameToCount, int numLogicalCores, int numPhysicalCores)
+		public static void AddCpuInfo(AgentSettings settings, RpcAgentCapabilities capabilities, string cpuName, int numLogicalCores, int numPhysicalCores)
 		{
-			if (nameToCount.Count > 0)
-			{
-				capabilities.Properties.Add("CPU=" + String.Join(", ", nameToCount.Select(x => (x.Value > 1) ? $"{x.Key} x {x.Value}" : x.Key)));
-			}
+			capabilities.Properties.Add($"CPU={cpuName}");
 
-			if (_settings.CurrentValue.CpuCount is > 0)
+			// Agent settings can artificially limit number of available cores
+			if (settings.CpuCount is > 0)
 			{
-				numLogicalCores = Math.Min(numLogicalCores, _settings.CurrentValue.CpuCount.Value);
+				numLogicalCores = Math.Min(numLogicalCores, settings.CpuCount.Value);
 			}
-			numLogicalCores = (int)(numLogicalCores * _settings.CurrentValue.CpuMultiplier);
+			numLogicalCores = (int)(numLogicalCores * settings.CpuMultiplier);
 
 			if (numLogicalCores > 0)
 			{
diff --git a/Engine/Source/Programs/Horde/HordeAgent/Services/WindowsCapabilities.cs b/Engine/Source/Programs/Horde/HordeAgent/Services/WindowsCapabilities.cs
new file mode 100644
--- /dev/null
+++ b/Engine/Source/Programs/Horde/HordeAgent/Services/WindowsCapabilities.cs
@@ -0,0 +1,172 @@
+﻿// Copyright Epic Games, Inc. All Rights Reserved.
+
+using System.Runtime.InteropServices;
+using System.Runtime.Versioning;
+using EpicGames.Core;
+using EpicGames.Horde.Agents;
+using HordeCommon.Rpc.Messages;
+using Microsoft.Extensions.Logging;
+using Microsoft.Extensions.Options;
+using Microsoft.Win32;
+using OpenTelemetry.Trace;
+
+namespace HordeAgent.Services;
+
+/// <summary>
+/// Get Windows-specific capabilities 
+/// </summary>
+[SupportedOSPlatform("windows")]
+public class WindowsCapabilities(IOptionsMonitor<AgentSettings> settings, ILogger logger, Tracer tracer) : ISystemCapabilities
+{
+	/// <inheritdoc/>
+	public Task<RpcAgentCapabilities> GetCapabilitiesAsync(DirectoryReference? workingDir)
+	{
+		using TelemetrySpan span = tracer.StartActiveSpan($"{nameof(WindowsCapabilities)}.{nameof(GetCapabilitiesAsync)}");
+		RpcAgentCapabilities caps = new();
+		
+		// OS and platform
+		SetProp(caps, KnownPropertyNames.Platform, "Win64");
+		SetProp(caps, KnownPropertyNames.PlatformGroup, "Windows");
+		SetProp(caps, KnownPropertyNames.PlatformGroup, "Microsoft");
+		SetProp(caps, KnownPropertyNames.PlatformGroup, "Desktop");
+		SetProp(caps, KnownPropertyNames.OsFamily, "Windows");
+		SetProp(caps, KnownPropertyNames.OsFamilyCompatibility, "Windows");
+		SetProp(caps, "OSDistribution", RuntimeInformation.OSDescription);
+		SetProp(caps, "OSKernelVersion", Environment.OSVersion.Version.ToString());
+		
+		// User
+		SetProp(caps, "User", Environment.UserName);
+		SetProp(caps, "Domain", Environment.UserDomainName);
+		SetProp(caps, "Interactive", Environment.UserInteractive);
+		SetProp(caps, "Elevated", CapabilitiesService.IsUserAdministrator());
+		
+		// CPU
+		int totalPhysicalCores = Environment.ProcessorCount;
+		int totalLogicalCores = GetNumLogicalCores();
+		CapabilitiesService.AddCpuInfo(settings.CurrentValue, caps, GetCpuName(), totalLogicalCores, totalPhysicalCores);
+		
+		// Memory
+		MEMORYSTATUSEX memStatus = new();
+		if (GlobalMemoryStatusEx(memStatus))
+		{
+			SetResource(caps, KnownResourceNames.Ram, (int)(memStatus.ullTotalPhys / (1024 * 1024 * 1024)));
+		}
+		else
+		{
+			int error = Marshal.GetLastWin32Error();
+			logger.LogWarning("Unable to get size of physical memory via GlobalMemoryStatusEx. Win32 error: {Error}", error);
+		}
+		
+		// GPU
+		SetGpuInfo(caps);
+		
+		return Task.FromResult(caps);
+	}
+	
+	private static void SetProp(RpcAgentCapabilities caps, string key, string value) => caps.Properties.Add($"{key}={value}");
+	private static void SetProp(RpcAgentCapabilities caps, string key, int value) => caps.Properties.Add($"{key}={value}");
+	private static void SetProp(RpcAgentCapabilities caps, string key, bool value) => caps.Properties.Add($"{key}={value}");
+	private static void SetResource(RpcAgentCapabilities caps, string key, int value) => caps.Resources.Add(key, value);
+	
+	private static int GetNumLogicalCores()
+	{
+		int total = 0;
+		ushort groupCount = GetActiveProcessorGroupCount();
+		for (ushort i = 0; i < groupCount; i++)
+		{
+			total += (int)GetActiveProcessorCount(i);
+		}
+		
+		return total;
+	}
+	
+	private static string GetCpuName()
+	{
+		try
+		{
+			using RegistryKey? key = Registry.LocalMachine.OpenSubKey(@"HARDWARE\DESCRIPTION\System\CentralProcessor\0");
+			if (key != null)
+			{
+				return key.GetValue("ProcessorNameString")?.ToString() ?? "Unknown CPU";
+			}
+		}
+		catch (Exception)
+		{
+			// Ignore exception
+		}
+		
+		return "Unknown CPU";
+	}
+	
+	private static void SetGpuInfo(RpcAgentCapabilities capabilities)
+	{
+		// GUID for display adapters in Windows
+		const string GpuRegistryKey = @"SYSTEM\CurrentControlSet\Control\Class\{4d36e968-e325-11ce-bfc1-08002be10318}";
+		string[] skipDevices = ["Remote Display", "Basic Display", "Standard VGA", "RDP", "Generic"];
+		int index = 0;
+		
+		using RegistryKey? baseKey = Registry.LocalMachine.OpenSubKey(GpuRegistryKey);
+		if (baseKey == null)
+		{
+			return;
+		}
+		
+		foreach (string subKeyName in baseKey.GetSubKeyNames())
+		{
+			if (subKeyName.Equals("Properties", StringComparison.OrdinalIgnoreCase))
+			{
+				continue;
+			}
+			
+			using RegistryKey? deviceKey = baseKey.OpenSubKey(subKeyName);
+			if (deviceKey == null)
+			{
+				continue;
+			}
+			
+			string? deviceDesc = deviceKey.GetValue("DriverDesc") as string;
+			if (String.IsNullOrEmpty(deviceDesc) || skipDevices.Any(term => deviceDesc.Contains(term, StringComparison.Ordinal)))
+			{
+				continue;
+			}
+			
+			string prefix = $"GPU-{++index}";
+			SetProp(capabilities, $"{prefix}-Name", deviceDesc);
+			
+			string? driverVersion = deviceKey.GetValue("DriverVersion") as string;
+			if (!String.IsNullOrEmpty(driverVersion))
+			{
+				SetProp(capabilities, $"{prefix}-DriverVersion", driverVersion);
+			}
+		}
+	}
+	
+	#region P/Invoke
+	#pragma warning disable
+	[DllImport("kernel32.dll")]
+	static extern ushort GetActiveProcessorGroupCount();
+	
+	[DllImport("kernel32.dll")]
+	static extern uint GetActiveProcessorCount(ushort GroupNumber);
+	
+	[DllImport("kernel32.dll", SetLastError = true)]
+	[return: MarshalAs(UnmanagedType.Bool)]
+	static extern bool GlobalMemoryStatusEx([In, Out] MEMORYSTATUSEX lpBuffer);
+	
+	[StructLayout(LayoutKind.Sequential)]
+	private class MEMORYSTATUSEX
+	{
+		public uint dwLength = (uint)Marshal.SizeOf(typeof(MEMORYSTATUSEX));
+		public uint dwMemoryLoad;
+		public ulong ullTotalPhys;
+		public ulong ullAvailPhys;
+		public ulong ullTotalPageFile;
+		public ulong ullAvailPageFile;
+		public ulong ullTotalVirtual;
+		public ulong ullAvailVirtual;
+		public ulong ullAvailExtendedVirtual;
+	}
+	#pragma warning enable
+	
+	#endregion P/Invoke
+}
\ No newline at end of file
