Adjusting Horde issue generation to discard any Gauntlet events if a fatal error event or a build drop event occurs. This should have no impact on Test Reports or log highlighting and is only relevant to Horde issue generation

#rnx

diff --git a/Engine/Source/Programs/AutomationTool/Gauntlet/Unreal/Utils/Gauntlet.UnrealTestResults.cs b/Engine/Source/Programs/AutomationTool/Gauntlet/Unreal/Utils/Gauntlet.UnrealTestResults.cs
--- a/Engine/Source/Programs/AutomationTool/Gauntlet/Unreal/Utils/Gauntlet.UnrealTestResults.cs
+++ b/Engine/Source/Programs/AutomationTool/Gauntlet/Unreal/Utils/Gauntlet.UnrealTestResults.cs
@@ -196,7 +196,7 @@ namespace Gauntlet
 										: (Event.IsWarning ?
 											Logging.LogLevel.Warning : Logging.LogLevel.Information);
 
-			EventId EventIdType = Event.IsCriticalFailure ? KnownLogEvents.Gauntlet_FatalEvent : KnownLogEvents.Gauntlet_UnrealEngineTestEvent;
+			EventId EventIdType = KnownLogEvents.Gauntlet_UnrealEngineTestEvent;
 
 			string Message = Event.FormatToString();
 			string Format = null;
diff --git a/Engine/Source/Programs/Horde/Plugins/Build/HordeServer.Build.Tests/Issues/IssueServiceTests.cs b/Engine/Source/Programs/Horde/Plugins/Build/HordeServer.Build.Tests/Issues/IssueServiceTests.cs
--- a/Engine/Source/Programs/Horde/Plugins/Build/HordeServer.Build.Tests/Issues/IssueServiceTests.cs
+++ b/Engine/Source/Programs/Horde/Plugins/Build/HordeServer.Build.Tests/Issues/IssueServiceTests.cs
@@ -2097,8 +2097,8 @@ namespace HordeServer.Tests.Issues
 				Assert.AreEqual("Automation fatal errors in Update Version Files, Compile UnrealHeaderTool Win64 and Compile ShooterGameEditor Win64", issues[0].Summary);
 			}
 			// #3
-			// Scenario: Gauntlet Test + Fatal event
-			// Expected: Gauntlet Test fingerprint using hash prefix and job step salt, Gauntlet Fatal is a separated issue
+			// Scenario: Gauntlet Test
+			// Expected: Gauntlet Test fingerprint using hash prefix and job step salt
 			{
 				string[] logErrors =
 				{
@@ -2117,7 +2117,6 @@ namespace HordeServer.Tests.Issues
 				IJob job = CreateJob(_mainStreamId, 130, "Test Build", _graph);
 				await using (TestJsonLogger logger = await CreateLoggerAsync(job, 0, 0))
 				{
-					logger.LogError(KnownLogEvents.Gauntlet_FatalEvent, "Some critical error");
 					foreach (string error in logErrors)
 					{
 						logger.LogError(KnownLogEvents.Gauntlet_TestEvent, "{Error}", error);
@@ -2128,14 +2127,57 @@ namespace HordeServer.Tests.Issues
 				await UpdateCompleteStepAsync(job, 0, 2, JobStepOutcome.Success);
 
 				IReadOnlyList<IIssue> issues = await IssueCollection.FindIssuesAsync();
-				Assert.AreEqual(2, issues.Count);
+				Assert.AreEqual(1, issues.Count);
 				Assert.AreEqual("Gauntlet:test", issues[0].Fingerprints[0].Type);
-				Assert.AreEqual("Gauntlet:fatal", issues[1].Fingerprints[0].Type);
 				Assert.AreEqual(10, issues[0].Fingerprints[0].Keys.Count);
 				Assert.AreEqual("hash:", issues[0].Fingerprints[0].Keys.First().Name.Substring(0, 5));
 				Assert.AreEqual($":{job.StreamId}:Update Version Files", issues[0].Fingerprints[0].Keys.First().Name.Substring(37));
 				Assert.AreEqual("Automation test errors in Update Version Files", issues[0].Summary);
 			}
+			// #4
+			// Scenario: Gauntlet Test + Fatal event + Gauntlet Test
+			// Expected: Fatal event will prune any issues that are not either a Fatal event or a Drop Build event
+			{
+				IJob job = CreateJob(_mainStreamId, 140, "Test Build", _graph);
+				await using (TestJsonLogger logger = await CreateLoggerAsync(job, 0, 0))
+				{
+					logger.LogError(KnownLogEvents.Gauntlet_TestEvent, "Expected 'SimpleValue::Get meta equality' to be true.");
+					logger.LogCritical(KnownLogEvents.Gauntlet_FatalEvent, "Some critical error");
+					logger.LogError(KnownLogEvents.Gauntlet_TestEvent, "Process encountered fatal error");
+				}
+				await UpdateCompleteStepAsync(job, 0, 0, JobStepOutcome.Failure);
+				await UpdateCompleteStepAsync(job, 0, 1, JobStepOutcome.Success);
+				await UpdateCompleteStepAsync(job, 0, 2, JobStepOutcome.Success);
+
+				IReadOnlyList<IIssue> issues = await IssueCollection.FindIssuesAsync();
+				Assert.AreEqual(1, issues.Count);
+				Assert.AreEqual("Gauntlet:fatal", issues[0].Fingerprints[0].Type);
+				Assert.AreEqual(1, issues[0].Fingerprints[0].Keys.Count);
+				Assert.AreEqual($"{job.StreamId}:Update Version Files", issues[0].Fingerprints[0].Keys.First().Name);
+				Assert.AreEqual("Automation fatal errors in Update Version Files", issues[0].Summary);
+			}
+			// #5
+			// Scenario: Fatal event + Gauntlet Test + Build Drop
+			// Expected: Both Fatal and Build Drop events are treated as severe errors and are the only 2 issues that are kept
+			{
+				IJob job = CreateJob(_mainStreamId, 150, "Test Build", _graph);
+				await using (TestJsonLogger logger = await CreateLoggerAsync(job, 0, 0))
+				{
+					logger.LogCritical(KnownLogEvents.Gauntlet_FatalEvent, "Some critical error");
+					logger.LogError(KnownLogEvents.Gauntlet_TestEvent, "Process encountered fatal error");
+					logger.LogError(KnownLogEvents.Gauntlet_BuildDropEvent, "Failed to find build");
+				}
+				await UpdateCompleteStepAsync(job, 0, 0, JobStepOutcome.Failure);
+				await UpdateCompleteStepAsync(job, 0, 1, JobStepOutcome.Success);
+				await UpdateCompleteStepAsync(job, 0, 2, JobStepOutcome.Success);
+
+				IReadOnlyList<IIssue> issues = await IssueCollection.FindIssuesAsync();
+				Assert.AreEqual(2, issues.Count);
+				Assert.AreEqual("Gauntlet:build drop", issues[0].Fingerprints[0].Type);
+				Assert.AreEqual("Gauntlet:fatal", issues[1].Fingerprints[0].Type);
+				Assert.AreEqual("Automation build drop errors in Update Version Files", issues[0].Summary);
+				Assert.AreEqual("Automation fatal errors in Update Version Files", issues[1].Summary);
+			}
 		}
 
 		[TestMethod]
diff --git a/Engine/Source/Programs/Shared/EpicGames.Horde/Issues/Handlers/GauntletIssueHandler.cs b/Engine/Source/Programs/Shared/EpicGames.Horde/Issues/Handlers/GauntletIssueHandler.cs
--- a/Engine/Source/Programs/Shared/EpicGames.Horde/Issues/Handlers/GauntletIssueHandler.cs
+++ b/Engine/Source/Programs/Shared/EpicGames.Horde/Issues/Handlers/GauntletIssueHandler.cs
@@ -56,6 +56,11 @@ namespace EpicGames.Horde.Issues.Handlers
 		/// </summary>
 		const int MaxMessageLength = 2000;
 
+		/// <summary>
+		/// Whether or not a severe event was reported
+		/// </summary>
+		bool _wasSevereEventReported = false;
+
 		readonly IssueHandlerContext _context;
 		readonly List<IssueEventGroup> _issues = new List<IssueEventGroup>();
 
@@ -72,6 +77,15 @@ namespace EpicGames.Horde.Issues.Handlers
 			{ KnownLogEvents.Gauntlet_FatalEvent, FatalPrefix}
 		};
 
+		/// <summary>
+		///  Known Gauntlet events associated with the highest severity
+		/// </summary>
+		static readonly HashSet<EventId> s_knownSevereGauntletEvents = new HashSet<EventId>
+		{
+			KnownLogEvents.Gauntlet_BuildDropEvent,
+			KnownLogEvents.Gauntlet_FatalEvent
+		};
+
 		/// <inheritdoc/>
 		public override int Priority => 10;
 
@@ -217,6 +231,12 @@ namespace EpicGames.Horde.Issues.Handlers
 		{
 			if (issueEvent.EventId != null && IsMatchingEventId(issueEvent.EventId.Value))
 			{
+				bool isSevereEvent = s_knownSevereGauntletEvents.Contains(issueEvent.EventId.Value);
+				if (_wasSevereEventReported && !isSevereEvent)
+				{
+					return true;
+				}
+
 				string gauntletType = GetEventPrefix(issueEvent.EventId!.Value);
 				IssueEventGroup issue = new IssueEventGroup($"Gauntlet:{gauntletType}", "Automation {Meta:GauntletType} {Severity} in {Meta:Node}", IssueChangeFilter.All);
 				issue.Events.Add(issueEvent);
@@ -230,6 +250,14 @@ namespace EpicGames.Horde.Issues.Handlers
 				}
 				_issues.Add(issue);
 
+				if (isSevereEvent)
+				{
+					// We've encountered a severe event where either the engine has crashed or a build was not found.
+					// We can ignore other issues to prevent superfluous issues from being created
+					_issues.RemoveAll((issue) => !issue.Events.Any(x => s_knownSevereGauntletEvents.Contains(x.EventId!.Value)));
+					_wasSevereEventReported = true;
+				}
+
 				return true;
 			}
 			return false;
