Horde: Updating schema documentation

#rnx

diff --git a/Engine/Source/Programs/Horde/Docs/Config/Schema/AclActions.md b/Engine/Source/Programs/Horde/Docs/Config/Schema/AclActions.md
--- a/Engine/Source/Programs/Horde/Docs/Config/Schema/AclActions.md
+++ b/Engine/Source/Programs/Horde/Docs/Config/Schema/AclActions.md
@@ -605,6 +605,12 @@
 | `UpdateReplicator` |  |
 | `ViewReplicator` |  |
 
+## Secrets
+
+| Name | Description |
+| ---- | ----------- |
+| `ViewSecret` | View a credential |
+
 ## ServiceAccounts
 
 | Name | Description |
@@ -777,7 +783,7 @@
 
 | Name | Description |
 | ---- | ----------- |
-| `ViewSecret` | View a credential |
+| `ViewSecret` |  |
 
 ## ServiceAccounts
 
diff --git a/Engine/Source/Programs/Horde/Docs/Config/Schema/Globals.md b/Engine/Source/Programs/Horde/Docs/Config/Schema/Globals.md
--- a/Engine/Source/Programs/Horde/Docs/Config/Schema/Globals.md
+++ b/Engine/Source/Programs/Horde/Docs/Config/Schema/Globals.md
@@ -71,8 +71,8 @@ Name | Description
 `build` | [BuildConfig](#buildconfig)<br>Configuration for the build plugin
 `compute` | [ComputeConfig](#computeconfig)<br>Configuration for the compute plugin
 `ddc` | [EmptyPluginConfig](#emptypluginconfig)<br>Configuration for the ddc plugin
-`experimental` | [ExperimentalConfig](#experimentalconfig)<br>Configuration for the experimental plugin
 `secrets` | [SecretsConfig](#secretsconfig)<br>Configuration for the secrets plugin
+`experimental` | [ExperimentalConfig](#experimentalconfig)<br>Configuration for the experimental plugin
 `storage` | [StorageConfig](#storageconfig)<br>Configuration for the storage plugin
 `symbols` | [SymbolsConfig](#symbolsconfig)<br>Configuration for the symbols plugin
 `tools` | [ToolsConfig](#toolsconfig)<br>Configuration for the tools plugin
@@ -432,14 +432,6 @@ Name | Description
 Empty implementation of
 
 
-## ExperimentalConfig
-
-Configuration for the experimental plugin
-
-Name | Description
----- | -----------
-`notifications` | [NotificationConfig](Notification.md)`[]`<br>Notifications to be sent on the Horde server
-
 ## SecretsConfig
 
 Configuration for the secrets system
@@ -447,6 +439,7 @@ Configuration for the secrets system
 Name | Description
 ---- | -----------
 `secrets` | [SecretConfig](#secretconfig)`[]`<br>List of secrets
+`secretProviderConfigs` | [SecretProviderConfig](#secretproviderconfig)`[]`<br>List of configurations for secret providers
 
 ## SecretConfig
 
@@ -457,7 +450,6 @@ Name | Description
 `id` | `string`<br>Identifier for this secret
 `data` | `string` `->` `string`<br>Key/value pairs associated with this secret
 `sources` | [ExternalSecretConfig](#externalsecretconfig)`[]`<br>Providers to source key/value pairs from
-`providers` | [SecretProviderConfig](#secretproviderconfig)`[]`<br>List of configurations for secret providers
 `acl` | [AclConfig](#aclconfig)<br>Defines access to this particular secret
 
 ## ExternalSecretConfig
@@ -513,6 +505,14 @@ Name | Description
 `PreSharedKey` | Use a pre-shared token provided by Vault for authentication.
 `AwsAuth` | Use AWS to authenticate against the Vault server.
 
+## ExperimentalConfig
+
+Configuration for the experimental plugin
+
+Name | Description
+---- | -----------
+`notifications` | [NotificationConfig](Notification.md)`[]`<br>Notifications to be sent on the Horde server
+
 ## StorageConfig
 
 Configuration for storage
