Horde : Update Secret docs to be explicit that AWS Auth use IAM with Vault Provider

#rnx

diff --git a/Engine/Source/Programs/Horde/Plugins/Secrets/HordeServer.Secrets/Secrets/Providers/HcpVaultConfig.cs b/Engine/Source/Programs/Horde/Plugins/Secrets/HordeServer.Secrets/Secrets/Providers/HcpVaultConfig.cs
--- a/Engine/Source/Programs/Horde/Plugins/Secrets/HordeServer.Secrets/Secrets/Providers/HcpVaultConfig.cs
+++ b/Engine/Source/Programs/Horde/Plugins/Secrets/HordeServer.Secrets/Secrets/Providers/HcpVaultConfig.cs
@@ -11,9 +11,9 @@ namespace HordeServer.Secrets.Providers
 		/// Use a pre-shared token provided by Vault for authentication.
 		/// </summary>
 		PreSharedKey,
-		
+
 		/// <summary>
-		/// Use AWS to authenticate against the Vault server.
+		/// Use AWS IAM credentials to authenticate against the Vault server.
 		/// </summary>
 		AwsAuth
 	}
@@ -30,16 +30,16 @@ namespace HordeServer.Secrets.Providers
 		/// </remarks>
 		/// </summary>
 		public string? PreSharedKey { get; set; }
-		
+
 		/// <summary>
 		/// The address of the Vault server.
 		/// </summary>
 		public string? EndPoint { get; set; }
 
 		/// <summary>
-		/// The AWS ARN to assume a role when logging into Vault.
+		/// The optional AWS ARN to assume a role when logging into Vault.
 		/// To be used when the type of credentials is <see cref="HcpVaultCredentialsType.AwsAuth"/>.
-		/// <remarks>When null AWS SDK's default credential search will be used.</remarks>
+		/// <remarks>When not set the AWS SDK default credential search will be used.</remarks>
 		/// </summary>
 		public string? AwsArnRole { get; set; }
 
@@ -50,15 +50,15 @@ namespace HordeServer.Secrets.Providers
 		/// To be used when the type of credentials is <see cref="HcpVaultCredentialsType.AwsAuth"/>.
 		/// </summary>
 		public string? AwsIamServerId { get; set; }
-	
+
 		/// <summary>
 		/// The name of the <see href="https://developer.hashicorp.com/vault/api-docs/auth/aws#role-4">role</see> in the Vault configuration
-		/// for the login. This role name can be different to the name of the AWS ARN role name.
+		/// for the login. This role name can be different to the name of the AWS IAM principal.
 		/// To be used when the type of credentials is <see cref="HcpVaultCredentialsType.AwsAuth"/>.
-		/// <remarks>When null the name of IAM principal will be used.</remarks>
+		/// <remarks>When not set the name of IAM principal will be used.</remarks>
 		/// </summary>
 		public string? Role { get; set; }
-	
+
 		/// <summary>
 		/// Type of credentials to use.
 		/// </summary>
