Horde : Rename SecretProviderConfigs to ProviderConfigs in Globals config

#rnx

diff --git a/Engine/Source/Programs/Horde/Docs/Config/Secrets.md b/Engine/Source/Programs/Horde/Docs/Config/Secrets.md
--- a/Engine/Source/Programs/Horde/Docs/Config/Secrets.md
+++ b/Engine/Source/Programs/Horde/Docs/Config/Secrets.md
@@ -128,41 +128,42 @@ three secrets from HashiCorp Vault using different authentication methods:
             }
         }
     ],
-
-    "secretProviderConfigs": [
-        {
-            "name": "hcp-vault-aws-assume-role",
-            "provider": "HcpVault",
-            "hcpCault": {
-                "credentials": "awsauth",
-                "endpoint": "https://vault.example.com",
-                "awsIamServerId": "vault.example.com",
-                "awsArnRole": "arn:aws:iam::1234567890:role/vault-example-role",
-                "role": "vault-role-name"
-            }
-        },
-
-        {
-            "name": "hcp-vault-aws",
-            "provider": "HcpVault",
-            "hcpVault": {
-                "credentials": "awsauth",
-                "endpoint": "https://vault.example.com",
-                "awsIamServerId": "vault.example.com",
-                "role": "vault-role-name"
-            }
-        },
-
-        {
-            "name": "hcp-vault-pre-shared-key",
-            "provider": "HcpVault",
-            "hcpVault": {
-                "credentials": "presharedkey",
-                "endpoint": "https://vault.example.com",
-                "preSharedKey": "the-key-provided-by-vault"
-            }
+    "plugins": {
+        "secrets": {
+            "providerConfigs": [
+                {
+                    "name": "hcp-vault-aws-assume-role",
+                    "provider": "HcpVault",
+                    "hcpCault": {
+                        "credentials": "awsauth",
+                        "endpoint": "https://vault.example.com",
+                        "awsIamServerId": "vault.example.com",
+                        "awsArnRole": "arn:aws:iam::1234567890:role/vault-example-role",
+                        "role": "vault-role-name"
+                    }
+                },
+                {
+                    "name": "hcp-vault-aws",
+                    "provider": "HcpVault",
+                    "hcpVault": {
+                        "credentials": "awsauth",
+                        "endpoint": "https://vault.example.com",
+                        "awsIamServerId": "vault.example.com",
+                        "role": "vault-role-name"
+                    }
+                },
+                {
+                    "name": "hcp-vault-pre-shared-key",
+                    "provider": "HcpVault",
+                    "hcpVault": {
+                        "credentials": "presharedkey",
+                        "endpoint": "https://vault.example.com",
+                        "preSharedKey": "the-key-provided-by-vault"
+                    }
+                }
+            ]
         }
-    ]
+    }
  ```
 
 ## Using secrets
diff --git a/Engine/Source/Programs/Horde/Plugins/Secrets/HordeServer.Secrets.Tests/Secrets/SecretProviderTests.cs b/Engine/Source/Programs/Horde/Plugins/Secrets/HordeServer.Secrets.Tests/Secrets/SecretProviderTests.cs
--- a/Engine/Source/Programs/Horde/Plugins/Secrets/HordeServer.Secrets.Tests/Secrets/SecretProviderTests.cs
+++ b/Engine/Source/Programs/Horde/Plugins/Secrets/HordeServer.Secrets.Tests/Secrets/SecretProviderTests.cs
@@ -132,7 +132,7 @@ namespace HordeServer.Tests.Secrets
 							]
 						}
 					],
-					"secretProviderConfigs": [
+					"providerConfigs": [
 						{
 							"name": "in-memory",
 							"provider": "InMemory"
diff --git a/Engine/Source/Programs/Horde/Plugins/Secrets/HordeServer.Secrets/SecretsConfig.cs b/Engine/Source/Programs/Horde/Plugins/Secrets/HordeServer.Secrets/SecretsConfig.cs
--- a/Engine/Source/Programs/Horde/Plugins/Secrets/HordeServer.Secrets/SecretsConfig.cs
+++ b/Engine/Source/Programs/Horde/Plugins/Secrets/HordeServer.Secrets/SecretsConfig.cs
@@ -26,9 +26,9 @@ namespace HordeServer
 		/// <summary>
 		/// List of configurations for secret providers
 		/// </summary>
-		public List<SecretProviderConfig> SecretProviderConfigs { get; set; } = new();
+		public List<SecretProviderConfig> ProviderConfigs { get; set; } = new();
 
-		private readonly Dictionary<SecretId, Dictionary<string, SecretProviderConfig>> _secretProviderConfigLookup = new();
+		private readonly Dictionary<SecretId, Dictionary<string, SecretProviderConfig>> _providerConfigLookup = new();
 
 		/// <inheritdoc/>
 		public void PostLoad(PluginConfigOptions configOptions)
@@ -40,25 +40,25 @@ namespace HordeServer
 				secret.PostLoad(configOptions.ParentAcl);
 			}
 
-			Dictionary<string, SecretProviderConfig> secretProviderConfigNameLookup = new();
-			foreach (SecretProviderConfig secretProviderConfig in SecretProviderConfigs)
+			Dictionary<string, SecretProviderConfig> providerConfigNameLookup = new();
+			foreach (SecretProviderConfig secretProviderConfig in ProviderConfigs)
 			{
-				secretProviderConfigNameLookup[secretProviderConfig.Name] = secretProviderConfig;
+				providerConfigNameLookup[secretProviderConfig.Name] = secretProviderConfig;
 			}
-			_secretProviderConfigLookup.Clear();
+			_providerConfigLookup.Clear();
 			foreach (KeyValuePair<SecretId, SecretConfig> secret in _secretLookup)
 			{
 				foreach (ExternalSecretConfig externalConfig in secret.Value.Sources)
 				{
 					if (!String.IsNullOrEmpty(externalConfig.ProviderConfig))
 					{
-						if (!_secretProviderConfigLookup.ContainsKey(secret.Key))
+						if (!_providerConfigLookup.ContainsKey(secret.Key))
 						{
-							_secretProviderConfigLookup[secret.Key] = new Dictionary<string, SecretProviderConfig>();
+							_providerConfigLookup[secret.Key] = new Dictionary<string, SecretProviderConfig>();
 						}
-						if (secretProviderConfigNameLookup.TryGetValue(externalConfig.ProviderConfig, out SecretProviderConfig? secretProviderConfig))
+						if (providerConfigNameLookup.TryGetValue(externalConfig.ProviderConfig, out SecretProviderConfig? secretProviderConfig))
 						{
-							_secretProviderConfigLookup[secret.Key].TryAdd(externalConfig.ProviderConfig, secretProviderConfig);
+							_providerConfigLookup[secret.Key].TryAdd(externalConfig.ProviderConfig, secretProviderConfig);
 						}
 					}
 				}
@@ -81,7 +81,7 @@ namespace HordeServer
 		/// <param name="config">Receives the secret provider configurations on success</param>
 		/// <returns>True on success</returns>
 		public bool TryGetSecretProviderConfigs(SecretId secretId, [NotNullWhen(true)] out Dictionary<string, SecretProviderConfig>? config)
-			=> _secretProviderConfigLookup.TryGetValue(secretId, out config);
+			=> _providerConfigLookup.TryGetValue(secretId, out config);
 
 		/// <summary>
 		/// Authorize access to a secret
