Horde: Add shortened alias for templates query param on api/v1/jobs endpoint

diff --git a/Engine/Source/Programs/Horde/Plugins/Build/HordeServer.Build/Jobs/JobsController.cs b/Engine/Source/Programs/Horde/Plugins/Build/HordeServer.Build/Jobs/JobsController.cs
--- a/Engine/Source/Programs/Horde/Plugins/Build/HordeServer.Build/Jobs/JobsController.cs
+++ b/Engine/Source/Programs/Horde/Plugins/Build/HordeServer.Build/Jobs/JobsController.cs
@@ -32,7 +32,7 @@ using Microsoft.Extensions.Options;
 namespace HordeServer.Jobs
 {
 	/// <summary>
-	/// Controller for the /api/v1/jobs endpoing
+	/// Controller for the /api/v1/jobs endpoint
 	/// </summary>
 	[ApiController]
 	[Authorize]
@@ -1082,6 +1082,7 @@ namespace HordeServer.Jobs
 		/// <param name="ids">The job ids to return</param>
 		/// <param name="name">Name of the job to find</param>
 		/// <param name="templates">List of templates to find</param>
+		/// <param name="t">Alternative to templates query param, used to shorten URLs with many templates/param>
 		/// <param name="streamId">The stream to search for</param>
 		/// <param name="minChange">The minimum changelist number</param>
 		/// <param name="maxChange">The maximum changelist number</param>
@@ -1109,6 +1110,7 @@ namespace HordeServer.Jobs
 			[FromQuery] string? streamId = null,
 			[FromQuery] string? name = null,
 			[FromQuery(Name = "template")] string[]? templates = null,
+			[FromQuery(Name = "t")] string[]? templateIds = null,
 			[FromQuery] CommitId? minChange = null,
 			[FromQuery] CommitId? maxChange = null,
 			[FromQuery] bool includePreflight = true,
@@ -1130,7 +1132,21 @@ namespace HordeServer.Jobs
 			JobId[]? jobIdValues = (ids == null) ? (JobId[]?)null : Array.ConvertAll(ids, x => JobId.Parse(x));
 			StreamId? streamIdValue = (streamId == null) ? (StreamId?)null : new StreamId(streamId);
 
-			TemplateId[]? templateRefIds = (templates != null && templates.Length > 0) ? templates.Select(x => new TemplateId(x)).ToArray() : null;
+			IEnumerable<TemplateId>? allTemplates = null;
+			if (templates != null && templates.Length > 0 && templateIds != null && templateIds.Length > 0)
+			{
+				allTemplates = templates.Select(x => new TemplateId(x)).Union(templateIds.Select(x => new TemplateId(x)));
+			}
+			else if (templates != null && templates.Length > 0)
+			{
+				allTemplates = templates.Select(x => new TemplateId(x));
+			}
+			else if (templateIds != null && templateIds.Length > 0)
+			{
+				allTemplates = templateIds.Select(x => new TemplateId(x));
+			}
+
+			TemplateId[]? templateRefIds = (allTemplates != null) ? allTemplates.ToArray(): null;
 
 			UserId? preflightStartedByUserIdValue = null;
 
