Horde: Fixed badly formed XML warning

diff --git a/Engine/Source/Programs/Horde/Plugins/Build/HordeServer.Build/Jobs/JobsController.cs b/Engine/Source/Programs/Horde/Plugins/Build/HordeServer.Build/Jobs/JobsController.cs
--- a/Engine/Source/Programs/Horde/Plugins/Build/HordeServer.Build/Jobs/JobsController.cs
+++ b/Engine/Source/Programs/Horde/Plugins/Build/HordeServer.Build/Jobs/JobsController.cs
@@ -1082,7 +1082,7 @@ namespace HordeServer.Jobs
 		/// <param name="ids">The job ids to return</param>
 		/// <param name="name">Name of the job to find</param>
 		/// <param name="templates">List of templates to find</param>
-		/// <param name="t">Alternative to templates query param, used to shorten URLs with many templates/param>
+		/// <param name="t">Alternative to templates query param, used to shorten URLs with many templates</param>
 		/// <param name="streamId">The stream to search for</param>
 		/// <param name="minChange">The minimum changelist number</param>
 		/// <param name="maxChange">The maximum changelist number</param>
