Horde: Improve performance of updateFilterItems

diff --git a/Engine/Source/Programs/Horde/HordeDashboard/src/components/jobDetailsV2/JobDetailsViewCommon.tsx b/Engine/Source/Programs/Horde/HordeDashboard/src/components/jobDetailsV2/JobDetailsViewCommon.tsx
--- a/Engine/Source/Programs/Horde/HordeDashboard/src/components/jobDetailsV2/JobDetailsViewCommon.tsx
+++ b/Engine/Source/Programs/Horde/HordeDashboard/src/components/jobDetailsV2/JobDetailsViewCommon.tsx
@@ -178,6 +178,7 @@ export class JobDetailsV2 extends PollBase {
       this.views = [];
       this.timing = undefined;
       this.overview = undefined;
+      this.stepsById = new Map();
    }
 
    viewReady(order: number): boolean {
@@ -243,9 +244,22 @@ export class JobDetailsV2 extends PollBase {
    }
 
    getSteps(): StepData[] {
-      return this.batches.map(b => b.steps).flat();
+      return Array.from(this.stepsById.values());
    }
 
+   stepsById: Map<string, StepData> = new Map();
+
+   createStepsById(): void {
+      if (this.stepsById.size) {
+         return;
+      }
+
+      this.batches.forEach(b => {
+         b.steps.forEach( s => {
+            this.stepsById.set(s.id, s);
+         })
+      })
+   }
 
    stepById(id?: string): StepData | undefined {
 
@@ -253,17 +267,23 @@ export class JobDetailsV2 extends PollBase {
          return undefined;
       }
 
-      let step: StepData | undefined;
+      if (this.stepsById.has(id)) {
+         return this.stepsById.get(id)!;
+      }
 
-      this.batches.forEach(b => {
+      let step: StepData | undefined;
 
-         if (step) {
-            return;
+      for (const b of this.batches) {
+         for (const s of b.steps) {
+            if (s.id === id) {
+               step = s;
+               break;
+            }
          }
+         if (step) break;
+      }
 
-         step = b.steps.find(s => s.id === id);
-
-      });
+      this.stepsById.set(id, step!);
 
       return step;
    }
@@ -451,15 +471,15 @@ export class JobDetailsV2 extends PollBase {
 
       let step: StepData | undefined;
 
-      this.batches.forEach(b => {
-
-         if (step) {
-            return;
+      for (const b of this.batches) {
+         for (const s of b.steps) {
+            if (s.name === name) {
+               step = s;
+               break;
+            }
          }
-
-         step = b.steps.find(s => s.name === name);
-
-      });
+         if (step) break;
+      }
 
       return step;
    }
@@ -607,6 +627,7 @@ export class JobDetailsV2 extends PollBase {
          }
 
          if (initialRequest) {
+            this.createStepsById();
             this.filter.updateFilterItems();
             this.filterUpdated();
 
