Horde: Rate limit jira to horde issue sync for slack and other considerations

diff --git a/Engine/Source/Programs/Horde/Plugins/Build/HordeServer.Build/Issues/External/JiraService.cs b/Engine/Source/Programs/Horde/Plugins/Build/HordeServer.Build/Issues/External/JiraService.cs
--- a/Engine/Source/Programs/Horde/Plugins/Build/HordeServer.Build/Issues/External/JiraService.cs
+++ b/Engine/Source/Programs/Horde/Plugins/Build/HordeServer.Build/Issues/External/JiraService.cs
@@ -327,6 +327,8 @@ namespace HordeServer.Issues.External
 
 			// Auto assign any issues where possible
 			List<IExternalIssue> externalIssues = await GetIssuesAsync(jiraKeys.ToArray(), cancellationToken);
+			const int RateLimit = 2;
+			int assigned = 0;
 
 			foreach (IIssue issue in openIssues)
 			{
@@ -346,6 +348,13 @@ namespace HordeServer.Issues.External
 						await _issueService.UpdateIssueAsync(issue.Id, ownerId: user.Id, cancellationToken: cancellationToken);
 
 						_logger.LogInformation("Assigning issue {IssueId} to user {UserId} by matching {EmailAddress}", issue.Id, user.Id, externalIssue.AssigneeEmailAddress);
+
+						assigned++;
+
+						if (assigned == RateLimit)
+						{
+							break;
+						}
 					}
 				}
 			}
