Horde: Fix issue with device health report bombing

diff --git a/Engine/Source/Programs/Horde/HordeDashboard/src/backend/DeviceTelemetryHandler.ts b/Engine/Source/Programs/Horde/HordeDashboard/src/backend/DeviceTelemetryHandler.ts
--- a/Engine/Source/Programs/Horde/HordeDashboard/src/backend/DeviceTelemetryHandler.ts
+++ b/Engine/Source/Programs/Horde/HordeDashboard/src/backend/DeviceTelemetryHandler.ts
@@ -234,7 +234,7 @@ export class PoolTelemetryHandler {
             platforms.set(platform.id, []);
          }
 
-         const poolName = this.pools.get(device.poolId)!.name;
+         const poolName = this.pools.get(device.poolId)?.name ?? "No Pool";
          const ratePercent = Math.floor(count / telemetry.length * 100);
          const reservable = device.enabled && !device.maintenanceTime;
 
