Fix typo in backend output property

#rnx

diff --git a/Engine/Source/Programs/Horde/Plugins/Build/HordeServer.Build/Jobs/TestData/TestDataController.cs b/Engine/Source/Programs/Horde/Plugins/Build/HordeServer.Build/Jobs/TestData/TestDataController.cs
--- a/Engine/Source/Programs/Horde/Plugins/Build/HordeServer.Build/Jobs/TestData/TestDataController.cs
+++ b/Engine/Source/Programs/Horde/Plugins/Build/HordeServer.Build/Jobs/TestData/TestDataController.cs
@@ -816,7 +816,7 @@ namespace HordeServer.Jobs.TestData
 					BuildCommitId = testData.BuildCommitId,
 					MetadataId = testData.Metadata.ToString(),
 					PhaseRef = testData.PhaseRef.ToString(),
-					SessionsId = testData.SessionId.ToString(),
+					SessionId = testData.SessionId.ToString(),
 					Outcome = testData.Outcome,
 					HasWarning = testData.HasWarning,
 					ErrorFingerprint = testData.ErrorFingerprint,
diff --git a/Engine/Source/Programs/Shared/EpicGames.Horde/Jobs/TestData/TestDataMessages.cs b/Engine/Source/Programs/Shared/EpicGames.Horde/Jobs/TestData/TestDataMessages.cs
--- a/Engine/Source/Programs/Shared/EpicGames.Horde/Jobs/TestData/TestDataMessages.cs
+++ b/Engine/Source/Programs/Shared/EpicGames.Horde/Jobs/TestData/TestDataMessages.cs
@@ -782,7 +782,7 @@ namespace EpicGames.Horde.Jobs.TestData
 		/// <summary>
 		/// The test session id associated with the phase session
 		/// </summary>
-		public string SessionsId { get; set; } = String.Empty;
+		public string SessionId { get; set; } = String.Empty;
 
 		/// <summary>
 		/// The commit id associated with the phase session
