Horde: Add documentation for Scheduled Downtime

diff --git a/Engine/Source/Programs/Horde/Docs/Config/ScheduledDowntime.md b/Engine/Source/Programs/Horde/Docs/Config/ScheduledDowntime.md
new file mode 100644
--- /dev/null
+++ b/Engine/Source/Programs/Horde/Docs/Config/ScheduledDowntime.md
@@ -0,0 +1,75 @@
+[Horde](../../README.md) > [Configuration](../Config.md) > Scheduled Downtime
+
+# Scheduled Downtime
+
+Horde supports a set of scheduled downtimes where no new jobs are started. This window can be used pause new jobs from
+starting when external services, such as Perforce, are offline.
+
+## Configuring Scheduled Downtime
+
+A scheduled downtime requires `startTime` and `finishTime` properties, which are a [point in time relative to UTC](https://learn.microsoft.com/en-us/dotnet/api/system.datetimeoffset?view=net-8.0), to define the duration of the downtime.
+
+The format of the time properties can be either:
+
+* [ISO 8601](https://learn.microsoft.com/en-us/dotnet/standard/base-types/standard-date-and-time-format-strings#Roundtrip), for example `2025-03-01T02:00:00-05:00` to represent the 1st March 2025 at 2AM and 5 hours behind UTC.
+
+* An abbreviated form, for example `Sat Mar 1, 2025, 2AM -5`. The supported abbreviated forms are the strings supported by the [DateTimeOffset.Parse](https://learn.microsoft.com/en-us/dotnet/api/system.datetimeoffset.parse?view=net-8.0) method.
+
+The `finishTime` property can be omitted and replaced with a `duration` property that is the number of hours and minutes of the downtime. The supported values for `duration` are the strings supported by the [TimeSpan.Parse](https://learn.microsoft.com/en-us/dotnet/api/System.TimeSpan.Parse?view=net-8.0) method.
+
+The UTC offset can be omitted from the `startTime` and `finishTime` properties and replaced with the `timeZone` property that will calculate the offset. For example, for North Carolina the time zone can be either `Eastern Standard Time` or `America/New_York`. The supported time zone values are the values returned by the [TimeZoneInfo.FindSystemTimeZoneById](https://learn.microsoft.com/en-us/dotnet/api/system.timezoneinfo.findsystemtimezonebyid?view=net-8.0) method, which includes the canonical names from the operating system and the aliases provided by .NET.
+
+The canonical names from the operating system can be found using [TimeZoneInfo.GetSystemTimeZones](https://learn.microsoft.com/en-us/dotnet/api/system.timezoneinfo.getsystemtimezones?view=net-8.0). The names of aliases can be found using [TimeZoneInfo.TryConvertWindowsIdToIanaId](https://learn.microsoft.com/en-us/dotnet/api/system.timezoneinfo.tryconvertwindowsidtoianaid?view=net-8.0) and
+[TimeZoneInfo.TryConvertIanaIdToWindowsId](https://learn.microsoft.com/en-us/dotnet/api/system.timezoneinfo.tryconvertianaidtowindowsid?view=net-8.0).
+
+### 1. Using ISO 8601 format and absolute date and times
+The following example creates a weekly maintenance window using the ISO 8601 format on the 6th September at 3AM and 4 hours behind UTC, for 30 minutes and repeats every 7 days.
+
+  ```json
+    "downtime": [
+        {
+            "startTime": "2020-09-06T03:00:00-04:00",
+            "finishTime": "2020-09-06T03:30:00-04:00",
+            "frequency": "Weekly"
+        }
+    ]
+  ```
+
+### 2. Using abbreviated format and duration property
+Changing the above example to use the abbreviated format and `duration` property instead of `finishTime`. The day of the week name is optional.
+
+   ```json
+    "downtime": [
+        {
+            "startTime": "Sun Sep 6, 2020, 3AM -4",
+            "duration": "0:30",
+            "frequency": "Weekly"
+        }
+    ]
+   ```
+
+### 3. Using a timezone
+The following example creates a daily maintenance window using the abbreviated format and `timeZone` property. Using a time zone the explicit offset component can be omitted.
+
+   ```json
+    "downtime": [
+        {
+            "startTime": "Sun Sep 6, 2020, 3AM",
+            "timeZone": "Eastern Standard Time",
+            "duration": "0:30",
+            "frequency": "Daily"
+        }
+    ]
+   ```
+
+### 4. One-off downtime
+The following example creates a one-off maintenance window of 1 hour and 15 minutes. The `frequency` property defaults to `Once` and can be omitted. Because no `timeZone` property and no offset component are set the start time is in UTC.
+
+   ```json
+    "downtime": [
+        {
+            "startTime": "Dec 10, 2024, 1AM",
+            "duration": "1:15"
+        }
+    ]
+   ```
diff --git a/Engine/Source/Programs/Horde/HordeServer/Server/GlobalConfig.cs b/Engine/Source/Programs/Horde/HordeServer/Server/GlobalConfig.cs
--- a/Engine/Source/Programs/Horde/HordeServer/Server/GlobalConfig.cs
+++ b/Engine/Source/Programs/Horde/HordeServer/Server/GlobalConfig.cs
@@ -239,7 +239,7 @@ namespace HordeServer.Server
 	}
 
 	/// <summary>
-	/// How frequently the maintence window repeats
+	/// How frequently the maintenance window repeats
 	/// </summary>
 	public enum ScheduledDowntimeFrequency
 	{
@@ -261,69 +261,6 @@ namespace HordeServer.Server
 
 	/// <summary>
 	/// Settings for the maintenance window
-	/// <para>
-	/// In the JSON configuration file:
-	/// <list type="bullet">
-	/// <item>
-	/// <description>
-	/// The values for <see cref="StartTime"/> and <see cref="FinishTime"/> is a date time string in the 
-	/// <see href="https://learn.microsoft.com/en-us/dotnet/standard/base-types/standard-date-and-time-format-strings#Roundtrip">ISO 8601</see>
-	/// format or abbreviated form. The supported abbreviated strings are the strings supported by the
-	/// <see cref="DateTimeOffset.Parse(String)"/> method.
-	/// </description>
-	/// </item>
-	/// <item>
-	/// <description>
-	/// <see cref="FinishTime"/> can be omitted for <see cref="Duration"/> to set when the maintenance window will be completed.
-	/// </description>
-	/// </item>
-	/// <item>
-	/// <description>
-	/// The UTC offset component in <see cref="StartTime"/> and <see cref="FinishTime"/> strings can be omitted and <see cref="TimeZone"/> used instead.
-	/// </description>
-	/// </item>
-	/// </list>
-	/// </para>
-	/// <example>
-	/// The following example creates a weekly maintenance window using the ISO 8601 format on the 6th September at 3AM and 4 hours behind UTC,
-	/// for 30 minutes and repeats every 7 days.
-	/// <code>
-	/// "downtime": [
-	///     {
-	///         "startTime": "2020-09-06T03:00:00-04:00",
-	///         "finishTime": "2020-09-06T03:30:00-04:00",
-	///         "frequency": "Weekly"
-	///     }
-	/// ]
-	/// </code>
-	/// </example>
-	/// <example>
-	/// The following example creates a weekly maintenance window using the abbreviated format 4 hours behind UTC, using
-	/// <see cref="Duration"/> instead of <see cref="FinishTime"/>. The day of the week name is optional.
-	/// <code>
-	/// "downtime": [
-	///     {
-	///         "startTime": "Sun Sep 6, 2020, 3AM -4",
-	///         "duration": "0:30",
-	///         "frequency": "Weekly"
-	///     }
-	/// ]
-	/// </code>
-	/// </example>
-	/// <example>
-	/// The following example creates a weekly maintenance window using the abbreviated format and <see cref="TimeZone"/>. By using a timezone the
-	/// explicit offset component can be omitted.
-	/// <code>
-	/// "downtime": [
-	///     {
-	///         "startTime": "Sun Sep 6, 2020, 3AM",
-	///         "timeZone": "Eastern Standard Time",
-	///         "duration": "0:30",
-	///         "frequency": "Weekly"
-	///     }
-	/// ]
-	/// </code>
-	/// </example>
 	/// </summary>
 	public class ScheduledDowntime
 	{
@@ -338,51 +275,15 @@ namespace HordeServer.Server
 		public DateTimeOffset FinishTime { get; set; }
 
 		/// <summary>
-		/// Duration of the maintenance window. An alternative to <see cref="FinishTime"/>.
-		/// <example>
-		/// The following example creates non-recurring maintenance windows of 1 hour and 15 minutes
-		/// The supported strings for the duration are the strings supported by the <see cref="TimeSpan.Parse(String)"/> method.
-		/// <code>
-		/// "downtime": [
-		///     {
-		///         "startTime": "Dec 10, 2024, 1AM",
-		///         "duration": "1:15"
-		///     }
-		/// ]
-		/// </code>
-		/// </example>
-		/// <remarks>
-		/// If <see cref="FinishTime"/> is also set in the JSON configuration then the value of Duration will take precedence.
-		/// </remarks>
+		/// Duration of the maintenance window. An alternative to FinishTime.
 		/// </summary>
 		public TimeSpan? Duration { get; set; }
 
 		/// <summary>
-		/// The name of time zone to set the Coordinated Universal Time (UTC) offset for <see cref="StartTime"/> and <see cref="FinishTime"/>.
-		/// When TimeZone is not null will replace the existing offsets in StartTime and FinishTime.
-		/// <example>
-		/// The following example creates a maintenance window on Sunday 6th September 2020 at 3AM, US and Canada Eastern Time.
-		/// <code>
-		/// "downtime": [
-		///     {
-		///         "startTime": "Sun Sep 6, 2020, 3AM",
-		///         "timeZone": "Eastern Standard Time",
-		///         "duration": "0:30"
-		///     }
-		/// ]
-		/// </code>
-		/// </example>
-		/// <remarks>
-		/// The supported TimeZone values are the supported values for <see cref="TimeZoneInfo.FindSystemTimeZoneById(String)"/>, which
-		/// includes the canonical values of the operating system, which are returned by <see cref="TimeZoneInfo.GetSystemTimeZones()"/>
-		/// and the aliases provided by .NET.
-		/// The names of aliases can be checked using <see cref="TimeZoneInfo.TryConvertWindowsIdToIanaId(String, out String?)"/> and
-		/// <see cref="TimeZoneInfo.TryConvertIanaIdToWindowsId(String, out String?)"/>.
-		/// For example, for North Carolina the time zone can be either "Eastern Standard Time" or "America/New_York".
-		/// </remarks>
+		/// The name of time zone to set the Coordinated Universal Time (UTC) offset for StartTime and FinishTime.
 		/// </summary>
 		public string? TimeZone { get; set; }
-		
+
 		/// <summary>
 		/// Frequency that the window repeats
 		/// </summary>
