Horde: Reference scheduled downtime from config homepage

diff --git a/Engine/Source/Programs/Horde/Docs/Config.md b/Engine/Source/Programs/Horde/Docs/Config.md
--- a/Engine/Source/Programs/Horde/Docs/Config.md
+++ b/Engine/Source/Programs/Horde/Docs/Config.md
@@ -65,6 +65,13 @@ While functional, this plugin is subject to change without notice.
 * [Experimental](Config/Experimental.md): Configuring the current
   functionality of this plugin.
 
+### Scheduled Downtime
+
+Manages the maintenance windows.
+
+* [Scheduled Downtime](Config/ScheduledDowntime.md): Configuring scheduled
+  downtimes.
+
 ### Secrets
 
 Managed access to sensitive data integrated with Horde's
