Horde: Use a shared ticker update for assigning Jira build health issues

diff --git a/Engine/Source/Programs/Horde/Plugins/Build/HordeServer.Build/Issues/External/JiraService.cs b/Engine/Source/Programs/Horde/Plugins/Build/HordeServer.Build/Issues/External/JiraService.cs
--- a/Engine/Source/Programs/Horde/Plugins/Build/HordeServer.Build/Issues/External/JiraService.cs
+++ b/Engine/Source/Programs/Horde/Plugins/Build/HordeServer.Build/Issues/External/JiraService.cs
@@ -109,6 +109,7 @@ namespace HordeServer.Issues.External
 		readonly IUserCollection _userCollection;
 
 		readonly ITicker _ticker;
+		readonly ITicker _updateTicker;
 
 		readonly Uri _jiraUrl;
 
@@ -132,6 +133,7 @@ namespace HordeServer.Issues.External
 			_issueService = issueService;
 			_userCollection = userCollection;
 			_ticker = clock.AddTicker<JiraService>(TimeSpan.FromMinutes(2.0), TickAsync, logger);
+			_updateTicker = clock.AddSharedTicker("JiraService.Update", TimeSpan.FromMinutes(1.0), TickUpdateAsync, logger);
 			_jiraUrl = _settings.JiraUrl!;
 			_buildConfig = buildConfig;
 
@@ -153,6 +155,7 @@ namespace HordeServer.Issues.External
 		public async ValueTask DisposeAsync()
 		{
 			await _ticker.DisposeAsync();
+			await _updateTicker.DisposeAsync();
 			_client.Dispose();
 		}
 
@@ -289,6 +292,58 @@ namespace HordeServer.Issues.External
 			}
 		}
 
+		/// <summary>
+		/// Ticks service
+		/// </summary>
+		async ValueTask TickUpdateAsync(CancellationToken stoppingToken)
+		{
+			if (!stoppingToken.IsCancellationRequested)
+			{
+				// Refresh issues
+				IReadOnlyList<IIssue> openIssues = await _issueService.Collection.FindIssuesAsync(resolved: false, cancellationToken: stoppingToken);
+			
+				// Auto assign any issues where possible
+				List<JiraCacheValue> externalIssues = _issueCache.Values.ToList(); 
+				const int RateLimit = 4;
+				int assigned = 0;
+
+				foreach (IIssue issue in openIssues)
+				{
+					if (String.IsNullOrEmpty(issue.ExternalIssueKey) || issue.OwnerId != null)
+					{
+						continue;
+					}
+
+					JiraCacheValue? externalIssue = externalIssues.Where(x => (issue.ExternalIssueKey != null) && (x._issue.Key == issue.ExternalIssueKey)).FirstOrDefault();
+
+					if (externalIssue.HasValue)
+					{
+						JiraIssue jiraIssue = externalIssue.Value._issue;
+
+						if (!String.IsNullOrEmpty(jiraIssue.AssigneeEmailAddress))
+						{
+
+							IUser? user = await _userCollection.FindUserByEmailAsync(jiraIssue.AssigneeEmailAddress, cancellationToken: stoppingToken);
+
+							if (user != null)
+							{
+								await _issueService.UpdateIssueAsync(issue.Id, ownerId: user.Id, cancellationToken: stoppingToken);
+
+								_logger.LogInformation("Assigning issue {IssueId} to user {UserId} by matching {EmailAddress}", issue.Id, user.Id, jiraIssue.AssigneeEmailAddress);
+
+								assigned++;
+
+								if (assigned == RateLimit)
+								{
+									break;
+								}
+							}
+						}
+					}
+				}
+			}
+		}
+
 		async ValueTask TickAsync(CancellationToken cancellationToken)
 		{
 
@@ -324,40 +379,8 @@ namespace HordeServer.Issues.External
 					jiraKeys.Add(openIssue.ExternalIssueKey);
 				}
 			}
-
-			// Auto assign any issues where possible
-			List<IExternalIssue> externalIssues = await GetIssuesAsync(jiraKeys.ToArray(), cancellationToken);
-			const int RateLimit = 2;
-			int assigned = 0;
-
-			foreach (IIssue issue in openIssues)
-			{
-				if (String.IsNullOrEmpty(issue.ExternalIssueKey) || issue.OwnerId != null)
-				{
-					continue;
-				}
-
-				IExternalIssue? externalIssue = externalIssues.Where(x => (issue.ExternalIssueKey != null) && (x.Key == issue.ExternalIssueKey)).FirstOrDefault();
-
-				if (externalIssue != null && !String.IsNullOrEmpty(externalIssue.AssigneeEmailAddress))
-				{
-					IUser? user = await _userCollection.FindUserByEmailAsync(externalIssue.AssigneeEmailAddress, cancellationToken:cancellationToken);
-
-					if (user != null)
-					{
-						await _issueService.UpdateIssueAsync(issue.Id, ownerId: user.Id, cancellationToken: cancellationToken);
-
-						_logger.LogInformation("Assigning issue {IssueId} to user {UserId} by matching {EmailAddress}", issue.Id, user.Id, externalIssue.AssigneeEmailAddress);
-
-						assigned++;
-
-						if (assigned == RateLimit)
-						{
-							break;
-						}
-					}
-				}
-			}
+			
+			await GetIssuesAsync(jiraKeys.ToArray(), cancellationToken);
 		}
 
 		public async Task<(string? key, string? url)> CreateIssueAsync(IUser user, string? externalIssueUser, int issueId, string summary, string projectId, string componentId, string issueType, string? description, string? hordeIssueLink, CancellationToken cancellationToken)
