Horde: Actually start the shared Jira ticker

diff --git a/Engine/Source/Programs/Horde/Plugins/Build/HordeServer.Build/Issues/External/JiraService.cs b/Engine/Source/Programs/Horde/Plugins/Build/HordeServer.Build/Issues/External/JiraService.cs
--- a/Engine/Source/Programs/Horde/Plugins/Build/HordeServer.Build/Issues/External/JiraService.cs
+++ b/Engine/Source/Programs/Horde/Plugins/Build/HordeServer.Build/Issues/External/JiraService.cs
@@ -146,10 +146,18 @@ namespace HordeServer.Issues.External
 		}
 
 		/// <inheritdoc/>
-		public Task StartAsync(CancellationToken cancellationToken) => _ticker.StartAsync();
+		public async Task StartAsync(CancellationToken cancellationToken)
+		{
+			await _ticker.StartAsync();
+			await _updateTicker.StartAsync();
+		}
 
 		/// <inheritdoc/>
-		public Task StopAsync(CancellationToken cancellationToken) => _ticker.StopAsync();
+		public async Task StopAsync(CancellationToken cancellationToken)
+		{
+			await _ticker.StopAsync();
+			await _updateTicker.StopAsync();
+		}
 
 		/// <inheritdoc/>
 		public async ValueTask DisposeAsync()
