Horde: Run agent's management server on dedicated thread

- Uses a dedicated thread with above normal priority for each server (admin and health check)
- Removes use of async/await which relied on the shared thread pool
- HTTP responses were timing out in production, likely due to CPU starvation during heavy workloads

diff --git a/Engine/Source/Programs/Horde/HordeAgent.Tests/Services/ManagementServiceTests.cs b/Engine/Source/Programs/Horde/HordeAgent.Tests/Services/ManagementServiceTests.cs
--- a/Engine/Source/Programs/Horde/HordeAgent.Tests/Services/ManagementServiceTests.cs
+++ b/Engine/Source/Programs/Horde/HordeAgent.Tests/Services/ManagementServiceTests.cs
@@ -4,8 +4,6 @@ using System;
 using System.IO;
 using System.Net;
 using System.Text.Json;
-using System.Threading;
-using System.Threading.Tasks;
 using HordeAgent.Services;
 using Microsoft.Extensions.Logging.Abstractions;
 using Microsoft.VisualStudio.TestTools.UnitTesting;
@@ -36,9 +34,9 @@ public sealed class ManagementServiceTests
 	private readonly WorkerServiceStub _workerServiceStub = new();
 	
 	[TestMethod]
-	public async Task HttpRequest_NotFound_Async()
+	public void HttpRequest_NotFound_Async()
 	{
-		(ErrorResponse? res, IHttpListenerResponse? httpRes) = await RequestAdminAsync<ErrorResponse>("http://localhost/not-found");
+		(ErrorResponse? res, IHttpListenerResponse? httpRes) = RequestAdmin<ErrorResponse>("http://localhost/not-found");
 		Assert.AreEqual(404, httpRes.StatusCode);
 		Assert.IsFalse(res!.Success);
 		Assert.AreEqual("application/json", httpRes.ContentType);
@@ -46,43 +44,43 @@ public sealed class ManagementServiceTests
 	}
 	
 	[TestMethod]
-	public async Task HttpRequest_HealthCheck_Async()
+	public void HttpRequest_HealthCheck_Async()
 	{
 		{
 			_workerServiceStub.IsConnected = true;
-			(HealthCheckResponse? res, IHttpListenerResponse? httpRes) = await RequestHealthAsync<HealthCheckResponse>("http://localhost/");
+			(HealthCheckResponse? res, IHttpListenerResponse? httpRes) = RequestHealth<HealthCheckResponse>("http://localhost/");
 			Assert.AreEqual(200, httpRes.StatusCode);
 			Assert.IsTrue(res!.IsHealthy);
 			Assert.IsTrue(res.IsConnected);
 		}
 		{
 			_workerServiceStub.IsConnected = false;
-			(HealthCheckResponse? res, IHttpListenerResponse? httpRes) = await RequestHealthAsync<HealthCheckResponse>("http://localhost/");
+			(HealthCheckResponse? res, IHttpListenerResponse? httpRes) = RequestHealth<HealthCheckResponse>("http://localhost/");
 			Assert.AreEqual(503, httpRes.StatusCode);
 			Assert.IsFalse(res!.IsHealthy);
 			Assert.IsFalse(res.IsConnected);
 		}
 	}
 	
-	private Task<(T?, IHttpListenerResponse)> RequestAdminAsync<T>(string path)
+	private (T?, IHttpListenerResponse) RequestAdmin<T>(string path)
 	{
-		return RequestAsync<T>(path, (ms, req, res) => ms.TestAdminRouteAsync(req, res));
+		return Request<T>(path, (ms, req, res) => ms.TestAdminRoute(req, res));
 	}
 	
-	private Task<(T?, IHttpListenerResponse)> RequestHealthAsync<T>(string path)
+	private (T?, IHttpListenerResponse) RequestHealth<T>(string path)
 	{
-		return RequestAsync<T>(path, (ms, req, res) => ms.TestHealthRouteAsync(req, res));
+		return Request<T>(path, (ms, req, res) => ms.TestHealthRoute(req, res));
 	}
 	
-	private async Task<(T?, IHttpListenerResponse)> RequestAsync<T>(string path, Func<ManagementService, IHttpListenerRequest, IHttpListenerResponse, Task> cb)
+	private (T?, IHttpListenerResponse) Request<T>(string path, Action<ManagementService, IHttpListenerRequest, IHttpListenerResponse> cb)
 	{
-		await using ManagementService ms = GetMs();
+		using ManagementService ms = GetMs();
 		HttpListenerRequestStub req = new(new Uri(path));
 		HttpListenerResponseStub res = new();
-		await cb(ms, req, res);
+		cb(ms, req, res);
 		res.OutputStream.Position = 0;
 		
-		T? deserialized = await JsonSerializer.DeserializeAsync<T>(res.OutputStream, ManagementService.JsonSerializerOptions, CancellationToken.None);
+		T? deserialized = JsonSerializer.Deserialize<T>(res.OutputStream, ManagementService.JsonSerializerOptions);
 		return (deserialized, res);
 	}
 	
diff --git a/Engine/Source/Programs/Horde/HordeAgent/Services/ManagementService.cs b/Engine/Source/Programs/Horde/HordeAgent/Services/ManagementService.cs
--- a/Engine/Source/Programs/Horde/HordeAgent/Services/ManagementService.cs
+++ b/Engine/Source/Programs/Horde/HordeAgent/Services/ManagementService.cs
@@ -2,7 +2,6 @@
 
 using System.Net;
 using System.Text.Json;
-using EpicGames.Core;
 using Microsoft.Extensions.Hosting;
 using Microsoft.Extensions.Logging;
 using Microsoft.Extensions.Options;
@@ -61,13 +60,14 @@ internal class HealthCheckResponse(bool isHealthy, bool isConnected)
 /// <summary>
 /// Handles agent's built-in HTTP server used for admin tasks and health check
 /// </summary>
-internal class ManagementService : IHostedService, IAsyncDisposable
+internal class ManagementService : IHostedService, IDisposable
 {
 	private readonly HttpListener? _adminListener;
 	private readonly HttpListener? _healthListener;
-	private readonly BackgroundTask? _adminTask;
-	private readonly BackgroundTask? _healthTask;
+	private readonly Thread? _adminThread;
+	private readonly Thread? _healthThread;
 	private readonly ILogger<ManagementService> _logger;
+	private readonly CancellationTokenSource _cts = new ();
 	
 	public static readonly JsonSerializerOptions JsonSerializerOptions = new()
 	{
@@ -75,7 +75,6 @@ internal class ManagementService : IHostedService, IAsyncDisposable
 	};
 	
 	private readonly RoutesDict _adminRoutes = new() { ["/"] = HandleJson((_, _) => "Admin server"), };
-	
 	private readonly RoutesDict _healthRoutes;
 	
 	private static Action<IHttpListenerRequest, IHttpListenerResponse> HandleJson<T>(Func<IHttpListenerRequest, IHttpListenerResponse, T> callback)
@@ -117,7 +116,10 @@ internal class ManagementService : IHostedService, IAsyncDisposable
 			}
 			
 			_logger.LogInformation("Admin HTTP server listening on {Endpoints}", String.Join(",", settings.CurrentValue.AdminEndpoints));
-			_adminTask = new BackgroundTask(ct => RunServerAsync(_adminListener, _adminRoutes, ct));
+			_adminThread = new Thread(() => RunServerSync(_adminListener, _adminRoutes, "Admin"))
+			{
+				Name = "AdminHttpServer", IsBackground = true, Priority = ThreadPriority.AboveNormal,
+			};
 		}
 		
 		if (settings.CurrentValue.HealthCheckEndpoints.Length > 0)
@@ -129,25 +131,19 @@ internal class ManagementService : IHostedService, IAsyncDisposable
 			}
 			
 			_logger.LogInformation("Health check HTTP server listening on {Endpoints}", String.Join(",", settings.CurrentValue.HealthCheckEndpoints));
-			_healthTask = new BackgroundTask(ct => RunServerAsync(_healthListener, _healthRoutes, ct));
+			_healthThread = new Thread(() => RunServerSync(_healthListener, _healthRoutes, "Health"))
+			{
+				Name = "HealthHttpServer", IsBackground = true, Priority = ThreadPriority.AboveNormal,
+			};
 		}
 	}
 	
-	/// <inheritdoc/>
-	public async ValueTask DisposeAsync()
+	public void Dispose()
 	{
+		_cts.Cancel();
+		_cts.Dispose();
 		_adminListener?.Close();
 		_healthListener?.Close();
-		
-		if (_adminTask != null)
-		{
-			await _adminTask.DisposeAsync();
-		}
-		
-		if (_healthTask != null)
-		{
-			await _healthTask.DisposeAsync();
-		}
 	}
 	
 	/// <inheritdoc/>
@@ -155,10 +151,17 @@ internal class ManagementService : IHostedService, IAsyncDisposable
 	{
 		try
 		{
-			_adminListener?.Start();
-			_healthListener?.Start();
-			_adminTask?.Start();
-			_healthTask?.Start();
+			if (_adminListener != null)
+			{
+				_adminListener.Start();
+				_adminThread?.Start();
+			}
+			
+			if (_healthListener != null)
+			{
+				_healthListener.Start();
+				_healthThread?.Start();
+			}
 		}
 		catch (Exception e)
 		{
@@ -178,31 +181,33 @@ internal class ManagementService : IHostedService, IAsyncDisposable
 	/// <inheritdoc/>
 	public async Task StopAsync(CancellationToken cancellationToken)
 	{
-		if (_adminTask != null)
+		await _cts.CancelAsync();
+		
+		if (_adminListener?.IsListening == true)
 		{
-			await _adminTask.StopAsync(cancellationToken);
+			_adminListener.Stop();
 		}
 		
-		if (_healthTask != null)
+		if (_healthListener?.IsListening == true)
 		{
-			await _healthTask.StopAsync(cancellationToken);
+			_healthListener.Stop();
 		}
-		
-		_adminListener?.Stop();
-		_healthListener?.Stop();
+
+		_adminThread?.Join(TimeSpan.FromSeconds(3));
+		_healthThread?.Join(TimeSpan.FromSeconds(3));
 	}
 	
-	internal Task TestAdminRouteAsync(IHttpListenerRequest req, IHttpListenerResponse res)
+	internal void TestAdminRoute(IHttpListenerRequest req, IHttpListenerResponse res)
 	{
-		return HandleRequestAsync(req, res, _adminRoutes);
+		HandleRequest(req, res, _adminRoutes);
 	}
 	
-	internal Task TestHealthRouteAsync(IHttpListenerRequest req, IHttpListenerResponse res)
+	internal void TestHealthRoute(IHttpListenerRequest req, IHttpListenerResponse res)
 	{
-		return HandleRequestAsync(req, res, _healthRoutes);
+		HandleRequest(req, res, _healthRoutes);
 	}
 	
-	private Task HandleRequestAsync(IHttpListenerRequest req, IHttpListenerResponse res, RoutesDict routes)
+	private void HandleRequest(IHttpListenerRequest req, IHttpListenerResponse res, RoutesDict routes)
 	{
 		try
 		{
@@ -229,18 +234,45 @@ internal class ManagementService : IHostedService, IAsyncDisposable
 		{
 			res.Close();
 		}
-		
-		return Task.CompletedTask;
 	}
 	
-	private async Task RunServerAsync(HttpListener listener, RoutesDict routes, CancellationToken cancellationToken)
+	private void RunServerSync(HttpListener listener, RoutesDict routes, string serverType)
 	{
-		while (!cancellationToken.IsCancellationRequested)
+		_logger.LogDebug("{ServerType} HTTP server thread started", serverType);
+		
+		try
+		{
+			while (!_cts.Token.IsCancellationRequested)
+			{
+				try
+				{
+					HttpListenerContext context = listener.GetContext();
+					WrappedHttpListenerRequest wrappedReq = new(context.Request);
+					WrappedHttpListenerResponse wrappedRes = new(context.Response);
+					HandleRequest(wrappedReq, wrappedRes, routes);
+				}
+				catch (HttpListenerException ex)
+				{
+					_logger.LogError(ex, "{ServerType} HTTP server listener error", serverType);
+					Thread.Sleep(100); // Small delay to prevent tight error loops
+				}
+				catch (Exception ex)
+				{
+					if (!_cts.Token.IsCancellationRequested)
+					{
+						_logger.LogError(ex, "Unexpected error in {ServerType} HTTP server", serverType);
+						Thread.Sleep(100); // Small delay to prevent tight error loops
+					}
+				}
+			}
+		}
+		catch (Exception ex)
+		{
+			_logger.LogError(ex, "{ServerType} HTTP server thread failed unexpectedly", serverType);
+		}
+		finally
 		{
-			HttpListenerContext context = await listener.GetContextAsync();
-			WrappedHttpListenerRequest wrappedReq = new(context.Request);
-			WrappedHttpListenerResponse wrappedRes = new(context.Response);
-			await HandleRequestAsync(wrappedReq, wrappedRes, routes);
+			_logger.LogInformation("{ServerType} HTTP server thread exiting", serverType);
 		}
 	}
 }
\ No newline at end of file
