Horde: Add support for filtering by status and showing cancelled builds in job search
Minor visual fixes to UI to account for new elements

diff --git a/Engine/Source/Programs/Horde/HordeDashboard/src/components/JobSearchSimple.tsx b/Engine/Source/Programs/Horde/HordeDashboard/src/components/JobSearchSimple.tsx
--- a/Engine/Source/Programs/Horde/HordeDashboard/src/components/JobSearchSimple.tsx
+++ b/Engine/Source/Programs/Horde/HordeDashboard/src/components/JobSearchSimple.tsx
@@ -29,7 +29,10 @@ export const JobSearchSimpleModal: React.FC<{ streamIdIn: string, changeListIn?:
       job: GetJobResponse
    };
 
-   const [searchState, setSearchState] = useState<{ items: JobItem[], groups?: IGroup[], templates: Set<string>, name?: string, userId?: string, global?: boolean, streamIdOverride?: string, preflights: boolean, preflightOnly?: boolean, minDate?: Date, maxDate?: Date, minCL?: string, maxCL?: string, querying: boolean, containsStep?: string, containsParameter?: string }>({ items: [], querying: false, preflights: false, minCL: changeListIn, maxCL: changeListIn, templates: new Set(jobTab?.templates ?? []) });
+   const statuses = ["Waiting", "Running", "Complete", "Cancelled"]
+   const allStatuses = new Set<string>(statuses);
+
+   const [searchState, setSearchState] = useState<{ items: JobItem[], groups?: IGroup[], templates: Set<string>, statuses: Set<string>, name?: string, userId?: string, global?: boolean, streamIdOverride?: string, preflights: boolean, preflightOnly?: boolean, minDate?: Date, maxDate?: Date, minCL?: string, maxCL?: string, querying: boolean, containsStep?: string, containsParameter?: string }>({ items: [], querying: false, preflights: false, minCL: changeListIn, maxCL: changeListIn, templates: new Set(jobTab?.templates ?? []), statuses: allStatuses });
    const [templateData, setTemplateData] = useState<{ streamId: string; templates: GetTemplateRefResponse[] } | undefined>(undefined);
    const [controller] = useState(new CalloutController());
    const [autoload, setAutoload] = useState<boolean | undefined>(undefined);
@@ -211,8 +214,68 @@ export const JobSearchSimpleModal: React.FC<{ streamIdIn: string, changeListIn?:
       items: templateOptions
    };
 
+   const ctemplates = templates.filter(t => searchState.templates.has(t.id));
+
+   // Statuses
+   let statusOptions: IContextualMenuItem[] = [];
+
+   let allStatusesSelected: boolean = statuses.every(s => searchState.statuses.has(s));
+   let someStatusesSelected: boolean = statuses.some(s => searchState.statuses.has(s)) && !allStatusesSelected;
+
+   statusOptions.push({
+      itemType: ContextualMenuItemType.Normal,
+      key: `select_all_statuses`,
+      onRender: () => {
+         return <Stack horizontal verticalFill verticalAlign="center" style={{ padding: "8px 12px 8px 12px", cursor: "pointer" }} onClick={(ev) => {
+            ev?.preventDefault();
+            ev?.stopPropagation();
+
+            allStatusesSelected = !allStatusesSelected;
+
+            if (allStatusesSelected) {
+               searchState.statuses = allStatuses;
+               setSearchState({ ...searchState })
+            } else {
+               searchState.statuses = new Set();
+               setSearchState({ ...searchState })
+            }
+         }} >
+
+            <Checkbox indeterminate={someStatusesSelected} checked={allStatusesSelected} />
+            <Text>Select All</Text></Stack>
+      }
+      
+   });
+
+   statuses.forEach((status) => {
+      statusOptions.push({ 
+         itemType: ContextualMenuItemType.Normal,
+         key: `${status}_key`,
+         onRender: () => {
+            return <Stack horizontal verticalFill verticalAlign="center" style={{ padding: "8px 12px 8px 12px", cursor: "pointer" }} onClick={(ev) => {
+               ev?.preventDefault();
+               ev?.stopPropagation();
+
+               if (searchState.statuses.has(status)) {
+                  searchState.statuses.delete(status);
+               } else {
+                  searchState.statuses.add(status);
+               }
+               setSearchState({ ...searchState })
+            }} >
 
-   const ctemplates = templates.filter(t => searchState.templates.size === 0 || searchState.templates.has(t.id));
+               <Checkbox checked={searchState.statuses.has(status)} />
+               <Text>{status}</Text></Stack>
+         }
+      });
+   })
+
+   const statusMenuProps: IContextualMenuProps = {
+      shouldFocusOnMount: true,
+      items: statusOptions
+   };
+
+   const cstatuses = statuses.filter(s => searchState.statuses.has(s))
 
    // results
    const labelColumnWidth = 396;
@@ -371,7 +434,7 @@ export const JobSearchSimpleModal: React.FC<{ streamIdIn: string, changeListIn?:
 
       if (column.name === "Change") {
 
-         return <Stack horizontal verticalFill={true} verticalAlign="center" tokens={{ childrenGap: 0, padding: 0 }} style={{ paddingTop: 4 }} ><ChangeButton job={item.job} /></Stack>;
+         return <Stack horizontal verticalFill={true} verticalAlign="center" tokens={{ childrenGap: 0, padding: 0 }} style={{ paddingTop: 4 }} ><ChangeButton job={item.job} hideAborted={false} /></Stack>;
       }
 
       if (column.name === "Labels") {
@@ -528,7 +591,7 @@ export const JobSearchSimpleModal: React.FC<{ streamIdIn: string, changeListIn?:
                id: jobIds.length ? jobIds : undefined,
                streamId: searchStreamId?.toLowerCase(),
                name: name,
-               filter: "id,labels,defaultLabel,name,change,createTime,streamId,preflightChange,startedByUserInfo,arguments",
+               filter: "id,labels,defaultLabel,name,change,createTime,streamId,preflightChange,startedByUserInfo,arguments,abortedByUserInfo,state",
                minChange: minCL,
                maxChange: maxCL,
                template: (!searchState.global && ctemplates.length && !searchState.containsStep) ? ctemplates.map(t => t.id) : undefined,
@@ -542,6 +605,12 @@ export const JobSearchSimpleModal: React.FC<{ streamIdIn: string, changeListIn?:
 
             jobs = await backend.getJobs(query);
 
+            if (cstatuses.length && cstatuses.length != allStatuses.size) {
+               jobs = jobs.filter(j => {
+                  return (cstatuses.includes('Cancelled') && j.abortedByUserInfo?.id || cstatuses.includes(j.state) && !j.abortedByUserInfo?.id)
+               })
+            }
+
             const p = searchState.containsParameter?.trim().toLowerCase();
             if (p) {
                jobs = jobs.filter(j => {
@@ -641,13 +710,22 @@ export const JobSearchSimpleModal: React.FC<{ streamIdIn: string, changeListIn?:
       }
    }
 
+   let statusText = "None Selected";
+   if (searchState.statuses.size) {
+      statusText = cstatuses.map(s => s).join(", ");
+      if (statusText.length > 84) {
+         statusText = statusText.slice(0, 84);
+         statusText += "...";
+      }
+   }
+
    if (autoloadIn && autoload === undefined) {
       queryJobs();
       setAutoload(autoloadIn);
       return null;
    }
 
-   return (<Modal isOpen={true} topOffsetFixed={true} styles={{ main: { padding: 8, width: 1400, height: '800px', backgroundColor: hordeTheme.horde.contentBackground, hasBeenOpened: false, top: "80px", position: "absolute" } }} className={hordeClasses.modal} onDismiss={() => { onClose() }}>
+   return (<Modal isOpen={true} topOffsetFixed={true} styles={{ main: { padding: 8, width: 1400, height: '820px', backgroundColor: hordeTheme.horde.contentBackground, hasBeenOpened: false, top: "80px", position: "absolute" } }} className={hordeClasses.modal} onDismiss={() => { onClose() }}>
       {searchState.querying && <StatusModal text={"Finding Jobs"} />}
       <JobLabelCallout controller={controller} />
       <Stack styles={{ root: { paddingTop: 8, paddingLeft: 24, paddingRight: 12, paddingBottom: 8 } }}>
@@ -684,6 +762,11 @@ export const JobSearchSimpleModal: React.FC<{ streamIdIn: string, changeListIn?:
                         <DefaultButton disabled={searchState.global} style={{ width: 352, textAlign: "left" }} menuProps={templateMenuProps} text={templateText} />
                      </Stack>
 
+                     <Stack>
+                        <Label>{searchState.statuses.size > 1 ? "Statuses" : "Status"}</Label>
+                        <DefaultButton style={{ width: 352, textAlign: "left" }} menuProps={statusMenuProps} text={statusText} />
+                     </Stack>
+
                      {!searchState.global && <Stack horizontal tokens={{ childrenGap: 12 }} >
                         <TextField value={searchState.minCL ?? ""} style={{ width: 168 }} label="Min Changelist" onChange={(ev, newValue) => {
                            newValue = newValue?.trim();
@@ -714,7 +797,7 @@ export const JobSearchSimpleModal: React.FC<{ streamIdIn: string, changeListIn?:
                         })} userIdHints={[]} noResultsFoundText="No users found" />
                      </Stack>
 
-                     {!searchState.global && <Stack>
+                     {!searchState.global && <Stack tokens={{ childrenGap: 12 }}>
                         <Stack>
                            <TextField disabled={!ctemplates.length} value={searchState.containsParameter ?? ""} spellCheck={false} autoComplete="off" style={{ width: 350 }} label="Contains Parameter" placeholder={!ctemplates.length ? "Please select a template to search parameters" : ""} onChange={(event, newValue) => {
                               setSearchState({ ...searchState, containsParameter: newValue?.trim() ? newValue : undefined })
@@ -759,12 +842,12 @@ export const JobSearchSimpleModal: React.FC<{ streamIdIn: string, changeListIn?:
                         }} />
                      </Stack>
 
-                     <Stack horizontal style={{ paddingTop: 24 }}>
+                     <Stack horizontal style={{ paddingTop: 20 }}>
                         <Stack style={{ paddingTop: 8, paddingRight: 24 }}>
                            <PrimaryButton disabled={searchState.querying} text="Search" onClick={() => { queryJobs() }} />
                         </Stack>
                         <Stack style={{ paddingTop: 8, paddingRight: 24 }}>
-                           <DefaultButton disabled={searchState.querying} text="Clear" onClick={() => { setSearchState({ items: [], groups: [], querying: false, preflights: false, userId: searchState.userId, templates: new Set() }) }} />
+                           <DefaultButton disabled={searchState.querying} text="Clear" onClick={() => { setSearchState({ items: [], groups: [], querying: false, preflights: false, userId: searchState.userId, templates: new Set(), statuses: new Set()}) }} />
                         </Stack>
                      </Stack>
                   </Stack>
@@ -781,7 +864,7 @@ export const JobSearchSimpleModal: React.FC<{ streamIdIn: string, changeListIn?:
                      </Stack>}
 
                      {!!searchState.items.length && <Stack>
-                        <div style={{ overflowY: 'auto', overflowX: 'hidden', height: "670px" }} data-is-scrollable={true} onScroll={() => { controller.setState(undefined, true) }}>
+                        <div style={{ overflowY: 'auto', overflowX: 'hidden', height: "700px" }} data-is-scrollable={true} onScroll={() => { controller.setState(undefined, true) }}>
                            <Stack tokens={{ childrenGap: 12 }} style={{ paddingRight: 12 }}>
                               <DetailsList
                                  compact={true}
