Horde: Report agent's tool ID as a property

diff --git a/Engine/Source/Programs/Horde/HordeAgent/Services/CapabilitiesService.cs b/Engine/Source/Programs/Horde/HordeAgent/Services/CapabilitiesService.cs
--- a/Engine/Source/Programs/Horde/HordeAgent/Services/CapabilitiesService.cs
+++ b/Engine/Source/Programs/Horde/HordeAgent/Services/CapabilitiesService.cs
@@ -13,6 +13,7 @@ using Amazon.Util;
 using EpicGames.Core;
 using EpicGames.Horde.Agents;
 using EpicGames.Horde.Compute;
+using HordeAgent.Utility;
 using HordeCommon.Rpc.Messages;
 using Microsoft.Extensions.Logging;
 using Microsoft.Extensions.Options;
@@ -402,6 +403,11 @@ namespace HordeAgent.Services
 			capabilities.Properties.Add($"ProcessArch={RuntimeInformation.ProcessArchitecture.ToString().ToLowerInvariant()}");
 			capabilities.Properties.Add($"OsArch={RuntimeInformation.OSArchitecture.ToString().ToLowerInvariant()}");
 #pragma warning restore CA1308
+			
+			if (VersionInfo.ToolId != null)
+			{
+				capabilities.Properties.Add($"{KnownPropertyNames.ToolId}={VersionInfo.ToolId}");
+			}
 
 			// Add any additional properties from the config file
 			capabilities.Properties.AddRange(_settings.CurrentValue.Properties.Select(kvp => $"{kvp.Key}={kvp.Value}"));
diff --git a/Engine/Source/Programs/Shared/EpicGames.Horde/Agents/KnownPropertyNames.cs b/Engine/Source/Programs/Shared/EpicGames.Horde/Agents/KnownPropertyNames.cs
--- a/Engine/Source/Programs/Shared/EpicGames.Horde/Agents/KnownPropertyNames.cs
+++ b/Engine/Source/Programs/Shared/EpicGames.Horde/Agents/KnownPropertyNames.cs
@@ -36,6 +36,11 @@ public static class KnownPropertyNames
 	/// Whether the agent is a .NET self-contained app
 	/// </summary>
 	public const string SelfContained = "SelfContained";
+	
+	/// <summary>
+	/// Self-reported tool ID of an agent. Used by the server when pushing updates to select the correct tool package.
+	/// </summary>
+	public const string ToolId = "ToolId";
 
 	/// <summary>
 	/// Pools that this agent belongs to
