Horde: Improve empty shelf error message on preflights
#skipci

diff --git a/Engine/Source/Programs/Horde/HordeDashboard/src/components/build/NewBuildV2.tsx b/Engine/Source/Programs/Horde/HordeDashboard/src/components/build/NewBuildV2.tsx
--- a/Engine/Source/Programs/Horde/HordeDashboard/src/components/build/NewBuildV2.tsx
+++ b/Engine/Source/Programs/Horde/HordeDashboard/src/components/build/NewBuildV2.tsx
@@ -1996,8 +1996,8 @@ const BuildModal: React.FC<{ setUseLegacyDialog: (value: boolean) => void }> = o
 
                if (errorReason) {
 
-                  if (errorReason?.trim().endsWith("does not exist")) {
-                     errorReason += ".  Perforce edge server replication for the change may be in progress."
+                  if (errorReason?.trim().endsWith("does not exist") || errorReason?.trim().endsWith("does not contain any shelved files")) {
+                     errorReason += ".\nPerforce edge server replication for the change may be in progress. Please try again in a few minutes."
                   }
 
                   ErrorHandler.set({
