Horde: Use short template query param to avoid URL char limit in job search

diff --git a/Engine/Source/Programs/Horde/HordeDashboard/src/backend/Api.ts b/Engine/Source/Programs/Horde/HordeDashboard/src/backend/Api.ts
--- a/Engine/Source/Programs/Horde/HordeDashboard/src/backend/Api.ts
+++ b/Engine/Source/Programs/Horde/HordeDashboard/src/backend/Api.ts
@@ -307,6 +307,7 @@ export type JobQuery = {
 	streamId?: string;
 	name?: string;
 	template?: string[];
+	t?: string[];
 	state?: string[];
 	outcome?: string[];
 	minChange?: number;
diff --git a/Engine/Source/Programs/Horde/HordeDashboard/src/components/JobSearchSimple.tsx b/Engine/Source/Programs/Horde/HordeDashboard/src/components/JobSearchSimple.tsx
--- a/Engine/Source/Programs/Horde/HordeDashboard/src/components/JobSearchSimple.tsx
+++ b/Engine/Source/Programs/Horde/HordeDashboard/src/components/JobSearchSimple.tsx
@@ -594,7 +594,7 @@ export const JobSearchSimpleModal: React.FC<{ streamIdIn: string, changeListIn?:
                filter: "id,labels,defaultLabel,name,change,createTime,streamId,preflightChange,startedByUserInfo,arguments,abortedByUserInfo,state",
                minChange: minCL,
                maxChange: maxCL,
-               template: (!searchState.global && ctemplates.length && !searchState.containsStep) ? ctemplates.map(t => t.id) : undefined,
+               t: (!searchState.global && ctemplates.length && !searchState.containsStep) ? ctemplates.map(t => t.id) : undefined,
                minCreateTime: minCreateTime?.toUTCString(),
                maxCreateTime: maxCreateTime?.toUTCString(),
                includePreflight: searchState.preflights || !!searchState.preflightOnly,
