Horde: Use status color user preferences for error and warning issue colors and controls
Minor visual fixes to error and warning buttons in log view

diff --git a/Engine/Source/Programs/Horde/HordeDashboard/src/base/utilities/colors.ts b/Engine/Source/Programs/Horde/HordeDashboard/src/base/utilities/colors.ts
new file mode 100644
--- /dev/null
+++ b/Engine/Source/Programs/Horde/HordeDashboard/src/base/utilities/colors.ts
@@ -0,0 +1,24 @@
+
+
+export const hexToRgb = (hex: string): string => {
+
+    hex = hex.replace("#", "");
+
+    let hexInt = parseInt(hex, 16);
+    let r = (hexInt >> 16) & 255;
+    let g = (hexInt >> 8) & 255;
+    let b = hexInt & 255;
+
+    return r + "," + g + "," + b;
+}
+
+
+export const isBright = (hex: string): boolean => {
+
+    let rgbColor = hexToRgb(hex);
+    let [r, g, b] = rgbColor.split(",").map(n => parseInt(n));
+
+    let luminosity: number = ((r * 0.2126) + (g * 0.7152) + (b * 0.0722)) / 255
+
+    return luminosity > 0.6
+}
diff --git a/Engine/Source/Programs/Horde/HordeDashboard/src/components/LogStyle.tsx b/Engine/Source/Programs/Horde/HordeDashboard/src/components/LogStyle.tsx
--- a/Engine/Source/Programs/Horde/HordeDashboard/src/components/LogStyle.tsx
+++ b/Engine/Source/Programs/Horde/HordeDashboard/src/components/LogStyle.tsx
@@ -2,6 +2,8 @@
 
 import { mergeStyleSets } from "@fluentui/react";
 import dashboard from "../backend/Dashboard";
+import { DashboardPreference } from "horde/backend/Api";
+import { hexToRgb, isBright } from "horde/base/utilities/colors";
 
 export type LogMetricType = {
    lineHeight: number;
@@ -55,6 +57,11 @@ export const getLogStyles = () => {
       }
    });
 
+   const warningColorPref = dashboard.getPreference(DashboardPreference.ColorWarning);
+   const errorColorPref = dashboard.getPreference(DashboardPreference.ColorError);
+   const warningHighlight = warningColorPref ? isBright(warningColorPref) ? "brightness(0.9)" : "brightness(1.2)" : "";
+   const errorHighlight = errorColorPref ? isBright(errorColorPref) ? "brightness(0.9)" : "brightness(1.2)" : "";
+
    const logStyleBase = _logStyleBase ?? mergeStyleSets({
       container: {
          overflow: 'auto',         
@@ -69,17 +76,18 @@ export const getLogStyles = () => {
       },
       itemWarning: [
          {
-            background: dashboard.darktheme ? "#302402": "#FEF8E7"
+            background: warningColorPref ? `rgb(${hexToRgb(warningColorPref)}, 0.2)` : dashboard.darktheme ? "#302402": "#FEF8E7"
          }
       ],
       errorButton: {
-         backgroundColor: dashboard.darktheme ? "#9D1410" : "#EC4C47",
+         backgroundColor: errorColorPref ? errorColorPref : dashboard.darktheme ? "#9D1410" : "#EC4C47",
          borderStyle: "hidden",
          color: "#FFFFFF",
          selectors: {
             ':active,:hover': {
                color: "#FFFFFF",
-               backgroundColor: "#DC3C37"
+               backgroundColor: errorColorPref ? errorColorPref : "#DC3C37",
+               filter: errorHighlight
             }
          }
       },
@@ -88,11 +96,12 @@ export const getLogStyles = () => {
          backgroundColor: dashboard.darktheme ? "#1F2223" : "#f3f2f1"
       },
       warningButton: {
-         backgroundColor: dashboard.darktheme ? "#9D840E" : "#F7D154",
+         backgroundColor: warningColorPref ? warningColorPref : dashboard.darktheme ? "#9D840E" : "#F7D154",
          borderStyle: "hidden",
          selectors: {
             ':active,:hover': {
-               backgroundColor: "#E7C144"
+               backgroundColor: warningColorPref ? warningColorPref : "#E7C144",
+               filter: warningHighlight
             }
          }
       },
@@ -114,7 +123,7 @@ export const getLogStyles = () => {
       gutterError: [
          {
             borderLeftStyle: 'solid',
-            borderLeftColor: "#EC4C47",
+            borderLeftColor: errorColorPref ? errorColorPref : "#EC4C47",
             borderLeftWidth: 6,
             padding: 0,
             margin: 0,
@@ -128,7 +137,7 @@ export const getLogStyles = () => {
       gutterWarning: [
          {
             borderLeftStyle: 'solid',
-            borderLeftColor: "#F7D154",
+            borderLeftColor: warningColorPref ? warningColorPref : "#F7D154",
             borderLeftWidth: 6,
             padding: 0,
             margin: 0,
@@ -141,7 +150,7 @@ export const getLogStyles = () => {
       ],
       itemError: [
          {
-            background: dashboard.darktheme ? "#330606" : "#FEF6F6",
+           background: errorColorPref ? `rgb(${hexToRgb(errorColorPref)}, 0.2)` : dashboard.darktheme ? "#330606" : "#FEF6F6",
          }
       ]
    });
diff --git a/Engine/Source/Programs/Horde/HordeDashboard/src/components/LogView.tsx b/Engine/Source/Programs/Horde/HordeDashboard/src/components/LogView.tsx
--- a/Engine/Source/Programs/Horde/HordeDashboard/src/components/LogView.tsx
+++ b/Engine/Source/Programs/Horde/HordeDashboard/src/components/LogView.tsx
@@ -7,7 +7,7 @@ import moment from 'moment-timezone';
 import React, { useEffect, useId, useState } from 'react';
 import { Link, useLocation, useNavigate, useParams } from 'react-router-dom';
 import backend from '../backend';
-import { ArtifactContextType, EventSeverity, GetChangeSummaryResponse, GetJobStepRefResponse, GetLogEventResponse, LogLevel } from '../backend/Api';
+import { ArtifactContextType, DashboardPreference, EventSeverity, GetChangeSummaryResponse, GetJobStepRefResponse, GetLogEventResponse, LogLevel } from '../backend/Api';
 import { CommitCache } from '../backend/CommitCache';
 import dashboard from '../backend/Dashboard';
 import { JobDetails } from '../backend/JobDetails';
@@ -29,6 +29,7 @@ import { getLogStyles, logMetricNormal, logMetricSmall } from "./LogStyle";
 import { PrintException } from './PrintException';
 import { StepRefStatusIcon } from './StatusIcon';
 import { TopNav } from './TopNav';
+import { isBright } from 'horde/base/utilities/colors';
 
 function getQueryLine(): number | undefined {
 
@@ -738,6 +739,12 @@ export const LogList: React.FC<{ logId: string }> = observer(({ logId }) => {
    const hordeTheme = getHordeTheme();
    const { hordeClasses, modeColors } = getHordeStyling();
 
+   const warningColorPref = dashboard.getPreference(DashboardPreference.ColorWarning);
+   const errorColorPref = dashboard.getPreference(DashboardPreference.ColorError);
+   const buttonNoneTextColor = dashboard.darktheme ? "#949898" : "#616e85";
+   const warningButtonTextColor = warningColorPref ? isBright(warningColorPref) ? "#2d3f5f" : "#F9F9FB" : dashboard.darktheme ? modeColors.text : buttonNoneTextColor;
+   const errorButtonTextColor = errorColorPref && isBright(errorColorPref) ? "#2d3f5f" : "#F9F9FB";
+
    if (!logId) {
       console.error("Bad log id settings up LogList");
       return <div>Error</div>;
@@ -874,7 +881,7 @@ export const LogList: React.FC<{ logId: string }> = observer(({ logId }) => {
             const href = `${location.pathname}?${query.toString()}`;
 
             const fontStyle = {
-               fontFamily: "Horde Open Sans SemiBold, sans-serif, sans-serif", color: error ? "#FFFFFF" : modeColors.text, fontSize: handler.fontSize, textDecoration: !item!.issue?.resolvedAt ? undefined : "line-through"
+               fontFamily: "Horde Open Sans SemiBold, sans-serif, sans-serif", color: error ? errorButtonTextColor: warningButtonTextColor, fontSize: handler.fontSize, textDecoration: !item!.issue?.resolvedAt ? undefined : "line-through"
             }
 
             return <Stack style={{width: 80}}>
@@ -1383,7 +1390,6 @@ export const LogList: React.FC<{ logId: string }> = observer(({ logId }) => {
       errorText = "No";
    }
 
-   const buttonNoneText = dashboard.darktheme ? "#949898" : "#616e85";
 
    let heightAdjust = 292;
    if (logTelemetry) {
@@ -1422,14 +1428,14 @@ export const LogList: React.FC<{ logId: string }> = observer(({ logId }) => {
                                     onClick={() => {
                                        nextError();
                                     }}
-                                    style={{ color: errors.length ? "#F9F9FB" : buttonNoneText, height: 30, fontSize: 19, paddingLeft: 12, borderRadius: '4px 0 0 4px' }}>
+                                    style={{ color: errors.length ? errorButtonTextColor : buttonNoneTextColor, height: 30, fontSize: 19, paddingLeft: errors.length ? 12 : 8, borderRadius: errors.length ? '4px 0 0 4px' : '4px' }}>
                                     {!!errors.length && <Icon style={{ fontSize: 19, paddingLeft: 4, paddingRight: 0 }} iconName='ChevronDown' />}
                                  </DefaultButton>
                               </TooltipHost>
-                              {!!errors.length && <div style={{backgroundColor: "#F9F9FB", height: 18, width: 1, position: 'absolute', right: -0.5, zIndex: 1 }} />}
+                              {!!errors.length && <div style={{backgroundColor: errorButtonTextColor, height: 18, width: 1, position: 'absolute', right: -0.5, zIndex: 1 }} />}
                            </Stack>
                            {!!errors.length && <TooltipHost content={'Navigate to previous error'}>
-                              <IconButton disabled={!prevErrorEnabled} className={prevErrorEnabled ? handler.style.errorButton : handler.style.errorButtonDisabled} style={{ color: "#F9F9FB", height: 30, fontSize: 19, padding: 2, borderRadius: '0 4px 4px 0' }} iconProps={{ iconName: 'ChevronUp' }} onClick={(event: any) => {
+                              <IconButton disabled={!prevErrorEnabled} className={prevErrorEnabled ? handler.style.errorButton : handler.style.errorButtonDisabled} style={{ color: errorButtonTextColor, height: 30, fontSize: 19, padding: 2, borderRadius: '0 4px 4px 0' }} iconProps={{ iconName: 'ChevronUp' }} onClick={(event: any) => {
                                  event?.stopPropagation();
                                  prevError();
                               }} />
@@ -1440,18 +1446,18 @@ export const LogList: React.FC<{ logId: string }> = observer(({ logId }) => {
                            <Stack horizontal verticalAlign='center' style={{position: 'relative'}}>
                               <TooltipHost content={'Skip to next block of errors or issue'}>
                                  <IconButton disabled={!nextErrorBlockEnabled} className={nextErrorBlockEnabled ? handler.style.errorButton : handler.style.errorButtonDisabled}
-                                    style={{ color: errors.length ? "#F9F9FB" : buttonNoneText, height: 30, fontSize: 19, padding: 9, borderRadius: '4px 0 0 4px' }}
+                                    style={{ color: errorButtonTextColor, height: 30, fontSize: 19, padding: 9, borderRadius: '4px 0 0 4px' }}
                                     styles={{rootPressed: {border: 'none', boxShadow: 'none'}}} 
                                     iconProps={{ iconName: 'DoubleChevronDown' }}
                                     onClick={() => {
                                        nextErrorBlock();
                                     }} />
                               </TooltipHost>
-                                 <div style={{backgroundColor: "#F9F9FB", height: 18, width: 1, position: 'absolute', right: -0.5, zIndex: 1 }}></div>
+                                 <div style={{backgroundColor: errorButtonTextColor, height: 18, width: 1, position: 'absolute', right: -0.5, zIndex: 1 }}></div>
                            </Stack>
                            <TooltipHost content={'Skip to previous block of errors or issue'}>
                               <IconButton disabled={!prevErrorBlockEnabled} className={prevErrorBlockEnabled ? handler.style.errorButton : handler.style.errorButtonDisabled}
-                                 style={{ color: errors.length ? "#F9F9FB" : buttonNoneText, height: 30, fontSize: 19, padding: 9,  borderRadius: '0 4px 4px 0'}} 
+                                 style={{ color: errorButtonTextColor, height: 30, fontSize: 19, padding: 9,  borderRadius: '0 4px 4px 0'}} 
                                  iconProps={{ iconName: 'DoubleChevronUp' }}
                                  onClick={(event: any) => {
                                     event?.stopPropagation();
@@ -1469,15 +1475,15 @@ export const LogList: React.FC<{ logId: string }> = observer(({ logId }) => {
                                     onClick={() => {
                                        nextWarning();
                                     }}
-                                    style={{ color: (dashboard.darktheme && nextWarningEnabled) ? "#F9F9FB" : buttonNoneText, height: 30, fontSize: 19, padding: 12, borderRadius: '4px 0 0 4px'}} >
+                                    style={{ color: warnings.length ? warningButtonTextColor : buttonNoneTextColor, height: 30, fontSize: 19, paddingLeft:  warnings.length ? 12 : 8, borderRadius: warnings.length ? '4px 0 0 4px': '4px' }} >
                                     {!!warnings.length && <Icon style={{ fontSize: 19, paddingLeft: 4, paddingRight: 0 }} iconName='ChevronDown' />}
                                  </DefaultButton>
                               </TooltipHost>     
-                              {!!warnings.length && <div style={{backgroundColor: (dashboard.darktheme && nextWarningEnabled) ? "#F9F9FB" : buttonNoneText, height: 18, width: 1, position: 'absolute', right: -0.5, zIndex: 1 }} />}
+                              {!!warnings.length && <div style={{backgroundColor: warningButtonTextColor, height: 18, width: 1, position: 'absolute', right: -0.5, zIndex: 1 }} />}
                            </Stack>
                            {!!warnings.length && <Stack>
                               <TooltipHost content={'Navigate to previous warning'}>
-                                 <IconButton disabled={!prevWarningEnabled} className={prevWarningEnabled ? handler.style.warningButton : handler.style.warningButtonDisabled} style={{ height: 30, fontSize: 19, padding: 2, color: (dashboard.darktheme && nextWarningEnabled) ? "#F9F9FB" : buttonNoneText, borderRadius: '0 4px 4px 0'}} iconProps={{ iconName: 'ChevronUp' }} onClick={(event: any) => {
+                                 <IconButton disabled={!prevWarningEnabled} className={prevWarningEnabled ? handler.style.warningButton : handler.style.warningButtonDisabled} style={{ height: 30, fontSize: 19, padding: 2, color: warningButtonTextColor, borderRadius: '0 4px 4px 0'}} iconProps={{ iconName: 'ChevronUp' }} onClick={(event: any) => {
                                     event?.stopPropagation();
                                     prevWarning();
                                  }} />
@@ -1489,18 +1495,18 @@ export const LogList: React.FC<{ logId: string }> = observer(({ logId }) => {
                            <Stack horizontal verticalAlign='center' style={{position: 'relative'}}>
                               <TooltipHost content={'Skip to next block of warnings or issue'}>
                                  <IconButton disabled={!nextWarningBlockEnabled} className={nextWarningBlockEnabled ? handler.style.warningButton : handler.style.warningButtonDisabled}
-                                    style={{ color: (dashboard.darktheme && nextWarningEnabled) ? "#F9F9FB" : buttonNoneText, height: 30, fontSize: 19, padding: 9, borderRadius: '4px 0 0 4px' }}
+                                    style={{ color: warningButtonTextColor, height: 30, fontSize: 19, padding: 9, borderRadius: '4px 0 0 4px' }}
                                     styles={{rootPressed: {border: 'none', boxShadow: 'none'}}} 
                                     iconProps={{ iconName: 'DoubleChevronDown' }}
                                     onClick={() => {
                                        nextWarningBlock();
                                     }} />
                               </TooltipHost>
-                                 <div style={{backgroundColor: (dashboard.darktheme && nextWarningEnabled) ? "#F9F9FB" : buttonNoneText, height: 18, width: 1, position: 'absolute', right: -0.5, zIndex: 1 }}></div>
+                                 <div style={{backgroundColor: warningButtonTextColor, height: 18, width: 1, position: 'absolute', right: -0.5, zIndex: 1 }}></div>
                            </Stack>
                            <TooltipHost content={'Skip to previous block of warnings or issue'}>
                               <IconButton disabled={!prevWarningBlockEnabled} className={prevWarningBlockEnabled ? handler.style.warningButton : handler.style.warningButtonDisabled}
-                                 style={{ color: (dashboard.darktheme && nextWarningEnabled) ? "#F9F9FB" : buttonNoneText, height: 30, fontSize: 19, padding: 9,  borderRadius: '0 4px 4px 0'}} 
+                                 style={{ color: warningButtonTextColor, height: 30, fontSize: 19, padding: 9,  borderRadius: '0 4px 4px 0'}} 
                                  iconProps={{ iconName: 'DoubleChevronUp' }}
                                  onClick={(event: any) => {
                                     event?.stopPropagation();
diff --git a/Engine/Source/Programs/Horde/HordeDashboard/src/components/jobDetailsV2/StepErrorPanel.tsx b/Engine/Source/Programs/Horde/HordeDashboard/src/components/jobDetailsV2/StepErrorPanel.tsx
--- a/Engine/Source/Programs/Horde/HordeDashboard/src/components/jobDetailsV2/StepErrorPanel.tsx
+++ b/Engine/Source/Programs/Horde/HordeDashboard/src/components/jobDetailsV2/StepErrorPanel.tsx
@@ -6,12 +6,13 @@ import { observer } from "mobx-react-lite";
 import React, { useEffect } from 'react';
 import { Link, useNavigate } from 'react-router-dom';
 import backend from "../../backend";
-import { EventData, EventSeverity, JobStepState } from '../../backend/Api';
+import { DashboardPreference, EventData, EventSeverity, JobStepState } from '../../backend/Api';
 import { ISideRailLink } from "../../base/components/SideRail";
 import { getHordeStyling } from "../../styles/Styles";
 import { getHordeTheme } from "../../styles/theme";
 import { renderLine } from '../LogRender';
 import { JobDataView, JobDetailsV2 } from "./JobDetailsViewCommon";
+import dashboard from "horde/backend/Dashboard";
 
 const errorSideRail: ISideRailLink = { text: "Events", url: "rail_log_events" };
 
@@ -145,11 +146,14 @@ const getStyles = () => {
 
    const theme = getHordeTheme();
 
+   const warningColorPref = dashboard.getPreference(DashboardPreference.ColorWarning);
+   const errorColorPref = dashboard.getPreference(DashboardPreference.ColorError);
+
    const styles = _styles ?? mergeStyleSets({
       gutter: [
          {
             borderLeftStyle: 'solid',
-            borderLeftColor: "#EC4C47",
+            borderLeftColor: errorColorPref ? errorColorPref : "#EC4C47",
             borderLeftWidth: 6,
             padding: 0,
             margin: 0,
@@ -163,7 +167,7 @@ const getStyles = () => {
       gutterWarning: [
          {
             borderLeftStyle: 'solid',
-            borderLeftColor: "rgb(247, 209, 84)",
+            borderLeftColor: warningColorPref ? warningColorPref : "rgb(247, 209, 84)",
             borderLeftWidth: 6,
             padding: 0,
             margin: 0,
