Horde: Remove log warnings in agent happening during a normal, clean shutdown

diff --git a/Engine/Source/Programs/Horde/HordeAgent/Services/ManagementService.cs b/Engine/Source/Programs/Horde/HordeAgent/Services/ManagementService.cs
--- a/Engine/Source/Programs/Horde/HordeAgent/Services/ManagementService.cs
+++ b/Engine/Source/Programs/Horde/HordeAgent/Services/ManagementService.cs
@@ -67,7 +67,8 @@ internal class ManagementService : IHostedService, IDisposable
 	private readonly Thread? _adminThread;
 	private readonly Thread? _healthThread;
 	private readonly ILogger<ManagementService> _logger;
-	private readonly CancellationTokenSource _cts = new ();
+	private bool _isStopping = false;
+	private bool _isDisposed = false;
 	
 	public static readonly JsonSerializerOptions JsonSerializerOptions = new()
 	{
@@ -140,10 +141,26 @@ internal class ManagementService : IHostedService, IDisposable
 	
 	public void Dispose()
 	{
-		_cts.Cancel();
-		_cts.Dispose();
-		_adminListener?.Close();
-		_healthListener?.Close();
+		if (!_isDisposed)
+		{
+			_isStopping = true;
+			
+			if (_adminListener?.IsListening == true)
+			{
+				_adminListener.Stop();
+			}
+			
+			if (_healthListener?.IsListening == true)
+			{
+				_healthListener.Stop();
+			}
+			
+			_adminThread?.Join(TimeSpan.FromSeconds(3));
+			_healthThread?.Join(TimeSpan.FromSeconds(3));
+			_adminListener?.Close();
+			_healthListener?.Close();
+			_isDisposed = true;
+		}
 	}
 	
 	/// <inheritdoc/>
@@ -179,22 +196,10 @@ internal class ManagementService : IHostedService, IDisposable
 	}
 	
 	/// <inheritdoc/>
-	public async Task StopAsync(CancellationToken cancellationToken)
+	public Task StopAsync(CancellationToken cancellationToken)
 	{
-		await _cts.CancelAsync();
-		
-		if (_adminListener?.IsListening == true)
-		{
-			_adminListener.Stop();
-		}
-		
-		if (_healthListener?.IsListening == true)
-		{
-			_healthListener.Stop();
-		}
-
-		_adminThread?.Join(TimeSpan.FromSeconds(3));
-		_healthThread?.Join(TimeSpan.FromSeconds(3));
+		Dispose();
+		return Task.CompletedTask;
 	}
 	
 	internal void TestAdminRoute(IHttpListenerRequest req, IHttpListenerResponse res)
@@ -242,7 +247,7 @@ internal class ManagementService : IHostedService, IDisposable
 		
 		try
 		{
-			while (!_cts.Token.IsCancellationRequested)
+			while (!_isStopping)
 			{
 				try
 				{
@@ -251,14 +256,9 @@ internal class ManagementService : IHostedService, IDisposable
 					WrappedHttpListenerResponse wrappedRes = new(context.Response);
 					HandleRequest(wrappedReq, wrappedRes, routes);
 				}
-				catch (HttpListenerException ex)
-				{
-					_logger.LogError(ex, "{ServerType} HTTP server listener error", serverType);
-					Thread.Sleep(100); // Small delay to prevent tight error loops
-				}
 				catch (Exception ex)
 				{
-					if (!_cts.Token.IsCancellationRequested)
+					if (!_isStopping)
 					{
 						_logger.LogError(ex, "Unexpected error in {ServerType} HTTP server", serverType);
 						Thread.Sleep(100); // Small delay to prevent tight error loops
@@ -272,7 +272,7 @@ internal class ManagementService : IHostedService, IDisposable
 		}
 		finally
 		{
-			_logger.LogInformation("{ServerType} HTTP server thread exiting", serverType);
+			_logger.LogDebug("{ServerType} HTTP server thread exiting", serverType);
 		}
 	}
 }
\ No newline at end of file
diff --git a/Engine/Source/Programs/Horde/HordeAgent/Services/WorkerService.cs b/Engine/Source/Programs/Horde/HordeAgent/Services/WorkerService.cs
--- a/Engine/Source/Programs/Horde/HordeAgent/Services/WorkerService.cs
+++ b/Engine/Source/Programs/Horde/HordeAgent/Services/WorkerService.cs
@@ -93,7 +93,10 @@ namespace HordeAgent.Services
 			}
 			catch (Exception ex)
 			{
-				_logger.LogCritical(ex, "Unhandled exception");
+				if (!stoppingToken.IsCancellationRequested)
+				{
+					_logger.LogCritical(ex, $"Unhandled exception in {nameof(WorkerService)}");	
+				}
 			}
 		}
 
