Horde: Check for null on jira issue
#skipci

diff --git a/Engine/Source/Programs/Horde/Plugins/Build/HordeServer.Build/Issues/External/JiraService.cs b/Engine/Source/Programs/Horde/Plugins/Build/HordeServer.Build/Issues/External/JiraService.cs
--- a/Engine/Source/Programs/Horde/Plugins/Build/HordeServer.Build/Issues/External/JiraService.cs
+++ b/Engine/Source/Programs/Horde/Plugins/Build/HordeServer.Build/Issues/External/JiraService.cs
@@ -328,7 +328,7 @@ namespace HordeServer.Issues.External
 					{
 						JiraIssue jiraIssue = externalIssue.Value._issue;
 
-						if (!String.IsNullOrEmpty(jiraIssue.AssigneeEmailAddress))
+						if (jiraIssue != null && !String.IsNullOrEmpty(jiraIssue.AssigneeEmailAddress))
 						{
 
 							IUser? user = await _userCollection.FindUserByEmailAsync(jiraIssue.AssigneeEmailAddress, cancellationToken: stoppingToken);
