Horde: Revert updateFilterItems optimisations

diff --git a/Engine/Source/Programs/Horde/HordeDashboard/src/components/jobDetailsV2/JobDetailsViewCommon.tsx b/Engine/Source/Programs/Horde/HordeDashboard/src/components/jobDetailsV2/JobDetailsViewCommon.tsx
--- a/Engine/Source/Programs/Horde/HordeDashboard/src/components/jobDetailsV2/JobDetailsViewCommon.tsx
+++ b/Engine/Source/Programs/Horde/HordeDashboard/src/components/jobDetailsV2/JobDetailsViewCommon.tsx
@@ -178,7 +178,6 @@ export class JobDetailsV2 extends PollBase {
       this.views = [];
       this.timing = undefined;
       this.overview = undefined;
-      this.stepsById = new Map();
    }
 
    viewReady(order: number): boolean {
@@ -244,22 +243,9 @@ export class JobDetailsV2 extends PollBase {
    }
 
    getSteps(): StepData[] {
-      return Array.from(this.stepsById.values());
+      return this.batches.map(b => b.steps).flat();
    }
 
-   stepsById: Map<string, StepData> = new Map();
-
-   createStepsById(): void {
-      if (this.stepsById.size) {
-         return;
-      }
-
-      this.batches.forEach(b => {
-         b.steps.forEach( s => {
-            this.stepsById.set(s.id, s);
-         })
-      })
-   }
 
    stepById(id?: string): StepData | undefined {
 
@@ -267,23 +253,17 @@ export class JobDetailsV2 extends PollBase {
          return undefined;
       }
 
-      if (this.stepsById.has(id)) {
-         return this.stepsById.get(id)!;
-      }
-
       let step: StepData | undefined;
 
-      for (const b of this.batches) {
-         for (const s of b.steps) {
-            if (s.id === id) {
-               step = s;
-               break;
-            }
+      this.batches.forEach(b => {
+
+         if (step) {
+            return;
          }
-         if (step) break;
-      }
 
-      this.stepsById.set(id, step!);
+         step = b.steps.find(s => s.id === id);
+
+      });
 
       return step;
    }
@@ -471,15 +451,15 @@ export class JobDetailsV2 extends PollBase {
 
       let step: StepData | undefined;
 
-      for (const b of this.batches) {
-         for (const s of b.steps) {
-            if (s.name === name) {
-               step = s;
-               break;
-            }
+      this.batches.forEach(b => {
+
+         if (step) {
+            return;
          }
-         if (step) break;
-      }
+
+         step = b.steps.find(s => s.name === name);
+
+      });
 
       return step;
    }
@@ -627,7 +607,6 @@ export class JobDetailsV2 extends PollBase {
          }
 
          if (initialRequest) {
-            this.createStepsById();
             this.filter.updateFilterItems();
             this.filterUpdated();
 
