Horde: Update release notes

#rnx

diff --git a/Engine/Source/Programs/Horde/Docs/ReleaseNotes.md b/Engine/Source/Programs/Horde/Docs/ReleaseNotes.md
--- a/Engine/Source/Programs/Horde/Docs/ReleaseNotes.md
+++ b/Engine/Source/Programs/Horde/Docs/ReleaseNotes.md
@@ -1,5 +1,25 @@
 # Release Notes
 
+## 2025-03-27
+
+* Remove log warnings in agent happening during a normal, clean shutdown (41102934)
+* Use status color user preferences for error and warning issue colors and controls (41076656)
+* Use short template query param to avoid URL char limit in job search (41075171)
+* Improve empty shelf error message on preflights (41070379)
+* Report agent's tool ID as a property (41065929)
+* Add support for filtering by status and showing cancelled builds in job search (41028739)
+* Run agent's management server on dedicated thread (41027669)
+* Use a shared ticker update for assigning Jira build health issues (41021755)
+* Added support for WITH_TESTS (41021713)
+* Reference scheduled downtime from config homepage (41019958)
+* Add documentation for Scheduled Downtime (41018403)
+* Fix issue with device health report bombing (40989153)
+* Rate limit jira to horde issue sync for slack and other considerations (40985841)
+* Improve performance of updateFilterItems (40933175)
+* Fix param name in XML comment (40907450)
+* Fixed badly formed XML warning (40906936)
+* Add shortened alias for templates query param on api/v1/jobs endpoint (40905858)
+
 ## 2025-03-13
 
 * Add retry policy for blob writes when getting an upload redirect (40683860)
