Added IoStore Telemetry to Horde Analytics
Changed the Target attribute name to Platform so it aligns with other use cases in the analytics backend
Bumped the IoStoreChunk event schema to 8

diff --git a/Engine/Source/Developer/IoStoreUtilities/Private/IoStoreUtilities.cpp b/Engine/Source/Developer/IoStoreUtilities/Private/IoStoreUtilities.cpp
--- a/Engine/Source/Developer/IoStoreUtilities/Private/IoStoreUtilities.cpp
+++ b/Engine/Source/Developer/IoStoreUtilities/Private/IoStoreUtilities.cpp
@@ -1158,7 +1158,7 @@ public:
 		OutputArchive.Reset(IFileManager::Get().CreateFileWriter(OutputFilename));
 		if (OutputArchive.IsValid())
 		{
-			OutputArchive->Logf(TEXT("OrderInContainer, ChunkId, PackageId, PackageName, Filename, ContainerName, Offset, OffsetOnDisk, Size, CompressedSize, Hash, ChunkType, ClassType, PakChunk, Target, InstallType"));
+			OutputArchive->Logf(TEXT("OrderInContainer, ChunkId, PackageId, PackageName, Filename, ContainerName, Offset, OffsetOnDisk, Size, CompressedSize, Hash, ChunkType, ClassType, PakChunk, Platform, InstallType"));
 		}
 	}
 
@@ -1167,11 +1167,11 @@ public:
 		if (OutputArchive.IsValid())
 		{
 			FString PakChunk = TEXT("global");
-			FString Target = TEXT("Unknown");
+			FString Platform = TEXT("Unknown");
 			FString InstallType = TEXT("Installed");
 
-			// Extract the PakChunk name and the Target from the container name
-			ContainerName.Split(TEXT("-"), &PakChunk, &Target);
+			// Extract the PakChunk name and the Platform from the container name
+			ContainerName.Split(TEXT("-"), &PakChunk, &Platform);
 
 			// Determine the Install Type
 			if (ContainerName.Contains(TEXT("iad")))
@@ -1202,7 +1202,7 @@ public:
 				*LexToString(Info.ChunkType),
 				*ClassType,
 				*PakChunk,
-				*Target,
+				*Platform,
 				*InstallType
 			);
 		}
diff --git a/Engine/Source/Programs/Horde/HordeServer/Defaults/default-metrics.telemetry.json b/Engine/Source/Programs/Horde/HordeServer/Defaults/default-metrics.telemetry.json
--- a/Engine/Source/Programs/Horde/HordeServer/Defaults/default-metrics.telemetry.json
+++ b/Engine/Source/Programs/Horde/HordeServer/Defaults/default-metrics.telemetry.json
@@ -828,6 +828,46 @@
 			"function": "Average",
 			"interval": "00:30:00"
 		},
+		{
+			"id": "iostorechunk-package-size",
+			"filter": "$[?(@.payload.EventName == 'Core.IoStoreChunk' && @.payload.SchemaVersion>=7)]",
+			"groupBy": "$payload.Build_BranchName,$payload.Horde_TemplateName,$payload.Target,$payload.InstallType,$payload.PakChunk",
+			"property": "$.Payload.Duration",
+			"function": "Average",
+			"interval": "00:60:00"
+		},
+		{
+			"id": "iostorechunk-iad-size",
+			"filter": "$[?(@.payload.EventName == 'Core.IoStoreChunk' && && @.payload.InstallType == 'IAD' && @.payload.SchemaVersion>=8)]",
+			"groupBy": "$payload.Build_BranchName,$payload.Horde_TemplateName,$payload.Target",
+			"property": "$.Payload.Duration",
+			"function": "Average",
+			"interval": "00:60:00"
+		},
+		{
+			"id": "iostorechunk-ias-size",
+			"filter": "$[?(@.payload.EventName == 'Core.IoStoreChunk' && && @.payload.InstallType == 'IAS' && @.payload.SchemaVersion>=8)]",
+			"groupBy": "$payload.Build_BranchName,$payload.Horde_TemplateName,$payload.Target",
+			"property": "$.Payload.Duration",
+			"function": "Average",
+			"interval": "00:60:00"
+		},
+		{
+			"id": "iostorechunk-installed-size",
+			"filter": "$[?(@.payload.EventName == 'Core.IoStoreChunk' && && @.payload.InstallType == 'Installed' && @.payload.SchemaVersion>=8)]",
+			"groupBy": "$payload.Build_BranchName,$payload.Horde_TemplateName,$payload.Target",
+			"property": "$.Payload.Duration",
+			"function": "Average",
+			"interval": "00:60:00"
+		},
+		{
+			"id": "iostorechunk-optional-size",
+			"filter": "$[?(@.payload.EventName == 'Core.IoStoreChunk' && && @.payload.InstallType == 'Optional' && @.payload.SchemaVersion>=8)]",
+			"groupBy": "$payload.Build_BranchName,$payload.Horde_TemplateName,$payload.Target",
+			"property": "$.Payload.Duration",
+			"function": "Average",
+			"interval": "00:60:00"
+		},
 		{
 			"id": "debug-probe",
 			"filter": "$[?(@.Payload.SchemaVersion >= 0 )]",
diff --git a/Engine/Source/Programs/Horde/HordeServer/Defaults/default-views.telemetry.json b/Engine/Source/Programs/Horde/HordeServer/Defaults/default-views.telemetry.json
--- a/Engine/Source/Programs/Horde/HordeServer/Defaults/default-views.telemetry.json
+++ b/Engine/Source/Programs/Horde/HordeServer/Defaults/default-views.telemetry.json
@@ -103,6 +103,10 @@
 				{
 					"name": "Hardware_Platform",
 					"group": "Hardware_Platform"
+				},
+				{
+					"name": "InstallType",
+					"group": "InstallType"
 				}
 			],
 			"categories": [
@@ -733,6 +737,36 @@
 						}
 					]
 				},
+				{
+					"name": "Package",
+					"charts": [
+						{
+							"name": "Package Size By Install Type",
+							"metrics": [
+								{
+									"id": "iostorechunk-installed-size"
+								},
+								{
+									"id": "iostorechunk-optional-size"
+								},
+								{
+									"id": "iostorechunk-ias-size"
+								},
+								{
+									"id": "iostorechunk-iad-size"
+								}
+							]
+						},
+						{
+							"name": "Package Size",
+							"metrics": [
+								{
+									"id": "iostorechunk-package-size"
+								}
+							]
+						}
+					]
+				},
 				{
 					"name": "Debug",
 					"charts": [
