Added Package Metrics to Horde Analytics

diff --git a/Engine/Source/Programs/Horde/HordeServer/Defaults/default-metrics.telemetry.json b/Engine/Source/Programs/Horde/HordeServer/Defaults/default-metrics.telemetry.json
--- a/Engine/Source/Programs/Horde/HordeServer/Defaults/default-metrics.telemetry.json
+++ b/Engine/Source/Programs/Horde/HordeServer/Defaults/default-metrics.telemetry.json
@@ -828,6 +828,46 @@
 			"function": "Average",
 			"interval": "00:30:00"
 		},
+		{
+			"id": "package-size",
+			"filter": "$[?(@.Payload.EventName == 'Core.IoStoreChunk' && @.Payload.SchemaVersion>=8 )]",
+			"groupBy": "$.Payload.Build_BranchName,$.Payload.Horde_TemplateName,$.Payload.Platform,$.Payload.InstallType,$.Payload.PakChunk",
+			"property": "$.Payload.CompressedSize",
+			"function": "Sum",
+			"interval": "00:30:00"
+		},
+		{
+			"id": "package-installed-total-size",
+			"filter": "$[?(@.Payload.EventName == 'Core.IoStoreChunk' && @.Payload.SchemaVersion>=8 && @.Payload.InstallType=='Installed' )]",
+			"groupBy": "$.Payload.Build_BranchName,$.Payload.Horde_TemplateName,$.Payload.Platform",
+			"property": "$.Payload.CompressedSize",
+			"function": "Sum",
+			"interval": "00:30:00"
+		},
+		{
+			"id": "package-optional-total-size",
+			"filter": "$[?(@.Payload.EventName == 'Core.IoStoreChunk' && @.Payload.SchemaVersion>=8 && @.Payload.InstallType=='Optional' )]",
+			"groupBy": "$.Payload.Build_BranchName,$.Payload.Horde_TemplateName,$.Payload.Platform",
+			"property": "$.Payload.CompressedSize",
+			"function": "Sum",
+			"interval": "00:30:00"
+		},
+		{
+			"id": "package-ias-total-size",
+			"filter": "$[?(@.Payload.EventName == 'Core.IoStoreChunk' && @.Payload.SchemaVersion>=8 && @.Payload.InstallType=='IAS' )]",
+			"groupBy": "$.Payload.Build_BranchName,$.Payload.Horde_TemplateName,$.Payload.Platform",
+			"property": "$.Payload.CompressedSize",
+			"function": "Sum",
+			"interval": "00:30:00"
+		},
+		{
+			"id": "package-iad-total-size",
+			"filter": "$[?(@.Payload.EventName == 'Core.IoStoreChunk' && @.Payload.SchemaVersion>=8 && @.Payload.InstallType=='IAD' )]",
+			"groupBy": "$.Payload.Build_BranchName,$.Payload.Horde_TemplateName,$.Payload.Platform",
+			"property": "$.Payload.CompressedSize",
+			"function": "Sum",
+			"interval": "00:30:00"
+		},
 		{
 			"id": "debug-probe",
 			"filter": "$[?(@.Payload.SchemaVersion >= 0 )]",
