Horde: Add missing await in GetCapabilitiesWithTimeoutMonitoringAsync

diff --git a/Engine/Source/Programs/Horde/HordeAgent/Services/CapabilitiesService.cs b/Engine/Source/Programs/Horde/HordeAgent/Services/CapabilitiesService.cs
--- a/Engine/Source/Programs/Horde/HordeAgent/Services/CapabilitiesService.cs
+++ b/Engine/Source/Programs/Horde/HordeAgent/Services/CapabilitiesService.cs
@@ -103,7 +103,7 @@ namespace HordeAgent.Services
 				while (!task.IsCompleted)
 				{
 					Task delayTask = Task.Delay(TimeSpan.FromSeconds(30.0));
-					if (Task.WhenAny(task, delayTask) == delayTask)
+					if (await Task.WhenAny(task, delayTask) == delayTask)
 					{
 						_logger.LogWarning("GetCapabilitiesInternalAsync() has been running for {Time}", timer.Elapsed);
 					}
