Added Package view to Horde Analytics Dashboard.
This view will display sizes per package and total package sizes per install type ( Installed, Optional, IAS, IAD ).

Added InstalType variable
Added a space to Hardware Platfrom variable name

diff --git a/Engine/Source/Programs/Horde/HordeServer/Defaults/default-views.telemetry.json b/Engine/Source/Programs/Horde/HordeServer/Defaults/default-views.telemetry.json
--- a/Engine/Source/Programs/Horde/HordeServer/Defaults/default-views.telemetry.json
+++ b/Engine/Source/Programs/Horde/HordeServer/Defaults/default-views.telemetry.json
@@ -101,8 +101,12 @@
 					]
 				},
 				{
-					"name": "Hardware_Platform",
+					"name": "Hardware Platform",
 					"group": "Hardware_Platform"
+				},
+				{
+					"name": "Install Type",
+					"group": "Install_Type"
 				}
 			],
 			"categories": [
@@ -733,6 +737,40 @@
 						}
 					]
 				},
+				{
+					"name": "Package",
+					"charts": [
+						{
+							"name": "Size By Install Type",
+							"metrics": [
+								{
+									"id": "package-installed-total-size",
+									"alias": "Installed"
+								},
+								{
+									"id": "package-optional-total-size",
+									"alias": "Optional"
+								},
+								{
+									"id": "package-ias-total-size",
+									"alias": "IAS"
+								},
+								{
+									"id": "package-iad-total-size",
+									"alias": "IAD"
+								}
+							]
+						},
+						{
+							"name": "Size Per Package",
+							"metrics": [
+								{
+									"id": "package-size"
+								}
+							]
+						}
+					]
+				},
 				{
 					"name": "Debug",
 					"charts": [
