Horde: Add an agent mode distinguishing between dedicated and workstation agents

diff --git a/Engine/Source/Programs/Horde/HordeServer.Shared/Utilities/HordeClaims.cs b/Engine/Source/Programs/Horde/HordeServer.Shared/Utilities/HordeClaims.cs
--- a/Engine/Source/Programs/Horde/HordeServer.Shared/Utilities/HordeClaims.cs
+++ b/Engine/Source/Programs/Horde/HordeServer.Shared/Utilities/HordeClaims.cs
@@ -42,9 +42,14 @@ namespace HordeServer.Utilities
 		public static AclClaimConfig StartChainedJobClaim { get; } = new AclClaimConfig(HordeClaimTypes.Role, "start-chained-job");
 
 		/// <summary>
-		/// Role for all agents
+		/// Role for dedicated agents
 		/// </summary>
-		public static AclClaimConfig AgentRoleClaim { get; } = new AclClaimConfig(HordeClaimTypes.Role, "agent");
+		public static AclClaimConfig AgentDedicatedRoleClaim { get; } = new AclClaimConfig(HordeClaimTypes.Role, "agent");
+		
+		/// <summary>
+		/// Role for workstation agents
+		/// </summary>
+		public static AclClaimConfig AgentWorkstationRoleClaim { get; } = new AclClaimConfig(HordeClaimTypes.Role, "agent-workstation");
 
 		/// <summary>
 		/// Gets the role for a specific agent
diff --git a/Engine/Source/Programs/Horde/Plugins/Build/HordeServer.Build.Tests/Agents/AgentServiceTest.cs b/Engine/Source/Programs/Horde/Plugins/Build/HordeServer.Build.Tests/Agents/AgentServiceTest.cs
--- a/Engine/Source/Programs/Horde/Plugins/Build/HordeServer.Build.Tests/Agents/AgentServiceTest.cs
+++ b/Engine/Source/Programs/Horde/Plugins/Build/HordeServer.Build.Tests/Agents/AgentServiceTest.cs
@@ -63,7 +63,7 @@ public class AgentServiceTest : BuildTestSetup
 	public async Task LastOnlineChangeDuringSessionCreateAsync()
 	{
 		// No session created yet, status change timestamp is empty
-		IAgent agent = await AgentService.CreateAgentAsync(new CreateAgentOptions(new AgentId("agent1"), false, ""));
+		IAgent agent = await AgentService.CreateAgentAsync(new CreateAgentOptions(new AgentId("agent1"), AgentMode.Dedicated, false, ""));
 		Assert.AreEqual(AgentStatus.Stopped, agent.Status);
 		Assert.IsTrue(agent.LastOnlineTime.HasValue);
 
@@ -78,7 +78,7 @@ public class AgentServiceTest : BuildTestSetup
 	private static int s_agentId = 1;
 	private async Task<IAgent> CreateAgentSessionAsync()
 	{
-		IAgent agent = await AgentService.CreateAgentAsync(new CreateAgentOptions(new AgentId("agentServiceTest-" + s_agentId++), false, ""));
+		IAgent agent = await AgentService.CreateAgentAsync(new CreateAgentOptions(new AgentId("agentServiceTest-" + s_agentId++), AgentMode.Dedicated, false, ""));
 		agent = await AgentService.CreateSessionAsync(agent, new RpcAgentCapabilities(), "v1");
 		return agent;
 	}
@@ -135,8 +135,8 @@ public class AgentServiceTest : BuildTestSetup
 	[TestMethod]
 	public async Task GetAgentRateTestAsync()
 	{
-		IAgent agent1 = await AgentService.CreateAgentAsync(new CreateAgentOptions(new AgentId("agent1"), false, ""));
-		IAgent agent2 = await AgentService.CreateAgentAsync(new CreateAgentOptions(new AgentId("agent2"), false, ""));
+		IAgent agent1 = await AgentService.CreateAgentAsync(new CreateAgentOptions(new AgentId("agent1"), AgentMode.Dedicated, false, ""));
+		IAgent agent2 = await AgentService.CreateAgentAsync(new CreateAgentOptions(new AgentId("agent2"), AgentMode.Dedicated, false, ""));
 		await AgentService.CreateSessionAsync(agent1, new RpcAgentCapabilities(new List<string>() { "aws-instance-type=c5.24xlarge", "osfamily=windows" }), "test");
 		await AgentService.CreateSessionAsync(agent2, new RpcAgentCapabilities(new List<string>() { "aws-instance-type=c4.4xLARge", "osfamily=WinDowS" }), "test");
 
@@ -158,7 +158,7 @@ public class AgentServiceTest : BuildTestSetup
 	[TestMethod]
 	public async Task Properties_ServerDefined_AgentCannotOverride_Async()
 	{
-		IAgent newAgent = await AgentService.CreateAgentAsync(new CreateAgentOptions(new AgentId("agent1"), false, ""));
+		IAgent newAgent = await AgentService.CreateAgentAsync(new CreateAgentOptions(new AgentId("agent1"), AgentMode.Dedicated, false, ""));
 		IAgent agent = await AgentService.CreateSessionAsync(newAgent, new RpcAgentCapabilities(["foo=bar", $"{KnownPropertyNames.Trusted}=true", $"{KnownPropertyNames.Trusted}"]), "test");
 		
 		Assert.AreEqual(0, newAgent.Properties.Count);
@@ -171,7 +171,7 @@ public class AgentServiceTest : BuildTestSetup
 	[TestMethod]
 	public async Task Properties_ServerDefined_PropagatedToAgent_Async()
 	{
-		CreateAgentOptions options = new(new AgentId("agent1"), false, "", ["server=foo"]);
+		CreateAgentOptions options = new(new AgentId("agent1"), AgentMode.Dedicated, false, "", ["server=foo"]);
 		IAgent newAgent = await AgentService.CreateAgentAsync(options);
 		IAgent agent = await AgentService.CreateSessionAsync(newAgent, new RpcAgentCapabilities(["agent=bar"]), "test");
 		
diff --git a/Engine/Source/Programs/Horde/Plugins/Build/HordeServer.Build.Tests/BuildTestSetup.cs b/Engine/Source/Programs/Horde/Plugins/Build/HordeServer.Build.Tests/BuildTestSetup.cs
--- a/Engine/Source/Programs/Horde/Plugins/Build/HordeServer.Build.Tests/BuildTestSetup.cs
+++ b/Engine/Source/Programs/Horde/Plugins/Build/HordeServer.Build.Tests/BuildTestSetup.cs
@@ -307,7 +307,7 @@ namespace HordeServer.Tests
 				tempProps.Add(KnownPropertyNames.AwsInstanceId + "=" + awsInstanceId);
 			}
 			
-			IAgent? agent = await AgentService.CreateAgentAsync(new CreateAgentOptions(new AgentId("TestAgent" + s_agentIdCounter++), ephemeral, ""));
+			IAgent? agent = await AgentService.CreateAgentAsync(new CreateAgentOptions(new AgentId("TestAgent" + s_agentIdCounter++), AgentMode.Dedicated, ephemeral, ""));
 			Assert.IsNotNull(agent);
 
 			agent = await agent.TryUpdateAsync(new UpdateAgentOptions { Enabled = enabled, ExplicitPools = poolId != null ? [poolId.Value] : [] });
diff --git a/Engine/Source/Programs/Horde/Plugins/Build/HordeServer.Build.Tests/Fixture.cs b/Engine/Source/Programs/Horde/Plugins/Build/HordeServer.Build.Tests/Fixture.cs
--- a/Engine/Source/Programs/Horde/Plugins/Build/HordeServer.Build.Tests/Fixture.cs
+++ b/Engine/Source/Programs/Horde/Plugins/Build/HordeServer.Build.Tests/Fixture.cs
@@ -112,7 +112,7 @@ namespace HordeServer.Tests
 			Job2 = (await jobService.GetJobAsync(Job2.Id))!;
 
 			Agent1Name = "testAgent1";
-			Agent1 = await agentService.CreateAgentAsync(new CreateAgentOptions(new AgentId(Agent1Name), false, ""));
+			Agent1 = await agentService.CreateAgentAsync(new CreateAgentOptions(new AgentId(Agent1Name), AgentMode.Dedicated, false, ""));
 		}
 	}
 
diff --git a/Engine/Source/Programs/Horde/Plugins/Build/HordeServer.Build.Tests/Jobs/JobServiceTests.cs b/Engine/Source/Programs/Horde/Plugins/Build/HordeServer.Build.Tests/Jobs/JobServiceTests.cs
--- a/Engine/Source/Programs/Horde/Plugins/Build/HordeServer.Build.Tests/Jobs/JobServiceTests.cs
+++ b/Engine/Source/Programs/Horde/Plugins/Build/HordeServer.Build.Tests/Jobs/JobServiceTests.cs
@@ -245,7 +245,7 @@ namespace HordeServer.Tests.Jobs
 
 			// ----
 
-			IAgent? agent = await AgentService.CreateAgentAsync(new CreateAgentOptions(new AgentId("TestAgent"), false, ""));
+			IAgent? agent = await AgentService.CreateAgentAsync(new CreateAgentOptions(new AgentId("TestAgent"), AgentMode.Dedicated, false, ""));
 			Assert.IsNotNull(agent);
 
 			agent = await agent.TryUpdateAsync(new UpdateAgentOptions { Enabled = true, ExplicitPools = new List<PoolId> { new PoolId("win") } });
diff --git a/Engine/Source/Programs/Horde/Plugins/Build/HordeServer.Build.Tests/Jobs/JobTaskSourceTests.cs b/Engine/Source/Programs/Horde/Plugins/Build/HordeServer.Build.Tests/Jobs/JobTaskSourceTests.cs
--- a/Engine/Source/Programs/Horde/Plugins/Build/HordeServer.Build.Tests/Jobs/JobTaskSourceTests.cs
+++ b/Engine/Source/Programs/Horde/Plugins/Build/HordeServer.Build.Tests/Jobs/JobTaskSourceTests.cs
@@ -159,7 +159,7 @@ namespace HordeServer.Tests.Jobs
 
 			if (shouldCreateAgent)
 			{
-				IAgent? agent = await AgentService.CreateAgentAsync(new CreateAgentOptions(new AgentId("TestAgent"), false, ""));
+				IAgent? agent = await AgentService.CreateAgentAsync(new CreateAgentOptions(new AgentId("TestAgent"), AgentMode.Dedicated, false, ""));
 				Assert.IsNotNull(agent);
 
 				agent = await agent.TryUpdateAsync(new UpdateAgentOptions { Enabled = isAgentEnabled, ExplicitPools = new List<PoolId> { pool.Id } });
diff --git a/Engine/Source/Programs/Horde/Plugins/Build/HordeServer.Build.Tests/Server/RpcServiceTest.cs b/Engine/Source/Programs/Horde/Plugins/Build/HordeServer.Build.Tests/Server/RpcServiceTest.cs
--- a/Engine/Source/Programs/Horde/Plugins/Build/HordeServer.Build.Tests/Server/RpcServiceTest.cs
+++ b/Engine/Source/Programs/Horde/Plugins/Build/HordeServer.Build.Tests/Server/RpcServiceTest.cs
@@ -1,23 +1,25 @@
 // Copyright Epic Games, Inc. All Rights Reserved.
 
-using System.Reflection;
+using System.IdentityModel.Tokens.Jwt;
+using System.IO.Compression;
 using System.Security.Claims;
-using System.Threading.Channels;
+using System.Text;
 using EpicGames.Horde.Agents;
 using EpicGames.Horde.Agents.Pools;
+using EpicGames.Horde.Tools;
 using Grpc.Core;
 using HordeCommon.Rpc;
 using HordeCommon.Rpc.Messages;
+using HordeServer.Acls;
 using HordeServer.Agents;
-using HordeServer.Server;
+using HordeServer.Tools;
 using HordeServer.Utilities;
 using Microsoft.AspNetCore.Http;
 using Microsoft.AspNetCore.Http.Features;
+using Microsoft.Extensions.DependencyInjection;
 
 namespace HordeServer.Tests.Server
 {
-	using ISession = Microsoft.AspNetCore.Http.ISession;
-
 	sealed class HttpContextStub : HttpContext
 	{
 		public override ConnectionInfo Connection { get; } = null!;
@@ -32,12 +34,9 @@ namespace HordeServer.Tests.Server
 		public override ClaimsPrincipal User { get; set; }
 		public override WebSocketManager WebSockets { get; } = null!;
 
-		public HttpContextStub(Claim roleClaimType)
+		public HttpContextStub(Claim? roleClaimType)
 		{
-			User = new ClaimsPrincipal(new ClaimsIdentity(new List<Claim>
-			{
-				roleClaimType
-			}, "TestAuthType"));
+			User = new ClaimsPrincipal(new ClaimsIdentity(roleClaimType != null ? [roleClaimType] : [], "TestAuthType"));
 		}
 
 		public HttpContextStub(ClaimsPrincipal user)
@@ -86,7 +85,7 @@ namespace HordeServer.Tests.Server
 			}, "TestAuthType")));
 		}
 
-		public ServerCallContextStub(Claim roleClaimType)
+		public ServerCallContextStub(Claim? roleClaimType)
 		{
 			// The GetHttpContext extension falls back to getting the HttpContext from UserState
 			// We can piggyback on that behavior during tests
@@ -119,171 +118,95 @@ namespace HordeServer.Tests.Server
 	[TestClass]
 	public class RpcServiceTest : BuildTestSetup
 	{
+		// Defined as consts so they can be used in DataRows in tests
+		public const string DedicatedAgentTool = "agent-dedicated-tool";
+		public const string DedicatedAgentRole = "agent";
+		public const string DedicatedAgentRegistrationRole = "agent-registration";
+		public const string WorkstationAgentTool = "agent-workstation-tool";
+		public const string WorkstationAgentRole = "agent-workstation";
+		public const string AdminRole = "admin";
+		
 		private readonly ServerCallContext _adminContext = new ServerCallContextStub(HordeClaims.AdminClaim.ToClaim());
+		private readonly ToolId _dedicatedAgentToolId = new (DedicatedAgentTool);
+		private readonly ToolId _workstationAgentToolId = new (WorkstationAgentTool);
 
-		class RpcServiceInvoker : CallInvoker
+		[TestMethod]
+		public async Task CreateSessionTestAsync()
 		{
-			private readonly RpcService _rpcService;
-			private readonly ServerCallContext _serverCallContext;
-
-			public RpcServiceInvoker(RpcService rpcService, ServerCallContext serverCallContext)
-			{
-				_rpcService = rpcService;
-				_serverCallContext = serverCallContext;
-			}
-
-			public override TResponse BlockingUnaryCall<TRequest, TResponse>(Method<TRequest, TResponse> method, string? host, CallOptions options, TRequest request)
-			{
-				throw new NotImplementedException("Blocking calls are not supported! Method " + method.FullName);
-			}
-
-			public override AsyncUnaryCall<TResponse> AsyncUnaryCall<TRequest, TResponse>(Method<TRequest, TResponse> method, string? host, CallOptions options, TRequest request)
-			{
-				MethodInfo methodInfo = GetMethod(method.Name);
-				Task<TResponse> res = (methodInfo.Invoke(_rpcService, new object[] { request, _serverCallContext }) as Task<TResponse>)!;
-				return new AsyncUnaryCall<TResponse>(res, null!, null!, null!, null!, null!);
-			}
-
-			public override AsyncServerStreamingCall<TResponse> AsyncServerStreamingCall<TRequest, TResponse>(Method<TRequest, TResponse> method, string? host, CallOptions options,
-				TRequest request)
-			{
-				Console.WriteLine($"RpcServiceInvoker.AsyncServerStreamingCall(method={method.FullName} request={request})");
-				throw new NotImplementedException();
-			}
-
-			public override AsyncClientStreamingCall<TRequest, TResponse> AsyncClientStreamingCall<TRequest, TResponse>(Method<TRequest, TResponse> method, string? host, CallOptions options)
-			{
-				Console.WriteLine($"RpcServiceInvoker.AsyncClientStreamingCall(method={method.FullName})");
-				throw new NotImplementedException();
-			}
-
-			public override AsyncDuplexStreamingCall<TRequest, TResponse> AsyncDuplexStreamingCall<TRequest, TResponse>(Method<TRequest, TResponse> method, string? host, CallOptions options)
-			{
-				Console.WriteLine($"RpcServiceInvoker.AsyncDuplexStreamingCall(method={method.FullName})");
-
-				GrpcDuplexStreamHandler<TRequest> requestStream = new GrpcDuplexStreamHandler<TRequest>(_serverCallContext);
-				GrpcDuplexStreamHandler<TResponse> responseStream = new GrpcDuplexStreamHandler<TResponse>(_serverCallContext);
-
-				MethodInfo methodInfo = GetMethod(method.Name);
-				Task methodTask = (methodInfo.Invoke(_rpcService, new object[] { requestStream, responseStream, _serverCallContext }) as Task)!;
-				methodTask.ContinueWith(t => { Console.Error.WriteLine($"Uncaught exception in {method.Name}: {t}"); }, CancellationToken.None, TaskContinuationOptions.OnlyOnFaulted, TaskScheduler.Default);
+			RpcCreateSessionRequest req = new RpcCreateSessionRequest();
+			await Assert.ThrowsExceptionAsync<StructuredRpcException>(() => RpcService.CreateSession(req, _adminContext));
 
-				return new AsyncDuplexStreamingCall<TRequest, TResponse>(requestStream, responseStream, null!, null!, null!, null!);
-			}
+			req.Id = new AgentId("MyName").ToString();
+			await Assert.ThrowsExceptionAsync<StructuredRpcException>(() => RpcService.CreateSession(req, _adminContext));
 
-			private MethodInfo GetMethod(string methodName)
-			{
-				MethodInfo? method = _rpcService.GetType().GetMethod(methodName);
-				if (method == null)
-				{
-					throw new ArgumentException($"Method {methodName} not found in RpcService");
-				}
+			req.Capabilities = new RpcAgentCapabilities();
+			RpcCreateSessionResponse res = await RpcService.CreateSession(req, _adminContext);
 
-				return method;
-			}
+			Assert.AreEqual("MYNAME", res.AgentId);
+			// TODO: Check Token, ExpiryTime, SessionId 
 		}
-
-		/// <summary>
-		/// Combines and cross-writes streams for a duplex streaming call in gRPC
-		/// </summary>
-		/// <typeparam name="T"></typeparam>
-		class GrpcDuplexStreamHandler<T> : IServerStreamWriter<T>, IClientStreamWriter<T>, IAsyncStreamReader<T> where T : class
+		
+		[TestMethod]
+		[DataRow(Mode.Dedicated, AdminRole, true)]
+		[DataRow(Mode.Dedicated, DedicatedAgentRegistrationRole, true)]
+		[DataRow(Mode.Workstation, AdminRole, true)]
+		[DataRow(Mode.Workstation, WorkstationAgentRole, true)]
+		[DataRow(Mode.Dedicated, DedicatedAgentRole, false)] // Dedicated agent must use agent registration claim
+		[DataRow(Mode.Dedicated, WorkstationAgentRole, false)] // Workstation agent cannot create a dedicated agent
+		[DataRow(Mode.Workstation, DedicatedAgentRole, false)]
+		[DataRow(Mode.Workstation, DedicatedAgentRegistrationRole, false)]
+		public async Task CreateAgent_Async(Mode agentMode, string? roleClaim, bool expectSuccess)
 		{
-			private readonly ServerCallContext _serverCallContext;
-			private readonly Channel<T> _channel;
-
-			public WriteOptions? WriteOptions { get; set; }
-
-			public GrpcDuplexStreamHandler(ServerCallContext serverCallContext)
+			ServerCallContextStub context = CreateServerContext(roleClaim);
+			RpcCreateAgentRequest req = new () { Name = "myName", Mode = agentMode };
+			if (expectSuccess)
 			{
-				_channel = System.Threading.Channels.Channel.CreateUnbounded<T>();
-
-				_serverCallContext = serverCallContext;
-			}
-
-			public void Complete()
-			{
-				_channel.Writer.Complete();
-			}
-
-			public IAsyncEnumerable<T> ReadAllAsync()
-			{
-				return _channel.Reader.ReadAllAsync();
-			}
-
-			public async Task<T?> ReadNextAsync()
-			{
-				if (await _channel.Reader.WaitToReadAsync())
-				{
-					_channel.Reader.TryRead(out T? message);
-					return message;
-				}
-				else
+				RpcCreateAgentResponse res = await RpcService.CreateAgent(req, context);
+				switch (agentMode)
 				{
-					return null;
+					case Mode.Dedicated:
+						{
+							JwtSecurityToken jwtToken = new JwtSecurityTokenHandler().ReadJwtToken(res.Token);
+							List<Claim> claims = jwtToken.Claims.ToList();
+							Assert.AreEqual("HS256", jwtToken.SignatureAlgorithm);
+							Assert.AreEqual(2, claims.Count);
+							Assert.AreEqual("http://epicgames.com/ue/horde/agent", claims[0].Type);
+							Assert.AreEqual("MYNAME", claims[0].Value);
+							break;
+						}
+					case Mode.Workstation: Assert.AreEqual("", res.Token); break;
+					case Mode.Unspecified: default: throw new ArgumentOutOfRangeException(nameof(agentMode), agentMode, null);
 				}
 			}
-
-			public Task WriteAsync(T message)
+			else
 			{
-				if (_serverCallContext.CancellationToken.IsCancellationRequested)
-				{
-					return Task.FromCanceled(_serverCallContext.CancellationToken);
-				}
-
-				if (!_channel.Writer.TryWrite(message))
-				{
-					throw new InvalidOperationException("Unable to write message.");
-				}
-
-				return Task.CompletedTask;
+				StructuredRpcException ex = await Assert.ThrowsExceptionAsync<StructuredRpcException>(() => RpcService.CreateAgent(req, context));
+				Assert.AreEqual(StatusCode.PermissionDenied, ex.StatusCode);
 			}
-
-			public Task CompleteAsync()
+		}
+		
+		[TestMethod]
+		[DataRow(DedicatedAgentTool, "999", DedicatedAgentRole, StatusCode.NotFound)] // Invalid version
+		[DataRow("not-valid", "1", DedicatedAgentRole, StatusCode.NotFound)] // Invalid tool ID
+		[DataRow(DedicatedAgentTool, "1", DedicatedAgentRole, StatusCode.OK)] // All valid
+		[DataRow(DedicatedAgentTool, "1", WorkstationAgentRole, StatusCode.PermissionDenied)] // Workstation role cannot access dedicated agent
+		[DataRow(DedicatedAgentTool, "1", null, StatusCode.PermissionDenied)] // Unauthed user cannot access dedicated agent
+		[DataRow(WorkstationAgentTool, "1", null, StatusCode.OK)] // Unauthed user can access public workstation agent
+		public async Task DownloadSoftware_Async(string toolId, string version, string? roleClaim, StatusCode expectedStatusCode)
+		{
+			await SetupToolsAsync(CancellationToken.None);
+			ServerCallContextStub context = CreateServerContext(roleClaim);
+			TestServerStreamWriter<RpcDownloadSoftwareResponse> responseStream = new (context);
+			RpcDownloadSoftwareRequest req = new () { Version = $"{toolId}:{version}" };
+			if (expectedStatusCode == StatusCode.OK)
 			{
-				Complete();
-				return Task.CompletedTask;
+				await RpcService.DownloadSoftware(req, responseStream, context);
 			}
-
-			public async Task<bool> MoveNext(CancellationToken cancellationToken)
+			else
 			{
-				_serverCallContext.CancellationToken.ThrowIfCancellationRequested();
-
-				if (await _channel.Reader.WaitToReadAsync(cancellationToken))
-				{
-					if (_channel.Reader.TryRead(out T? message))
-					{
-						Current = message;
-						return true;
-					}
-				}
-
-				Current = null!;
-				return false;
+				StructuredRpcException ex = await Assert.ThrowsExceptionAsync<StructuredRpcException>(() => RpcService.DownloadSoftware(req, responseStream, context));
+				Assert.AreEqual(expectedStatusCode, ex.StatusCode);
 			}
-
-			public T Current { get; private set; } = null!;
-		}
-
-		public RpcServiceTest()
-		{
-			UpdateConfig(x => x.Plugins.GetComputeConfig().Pools.Clear());
-		}
-
-		[TestMethod]
-		public async Task CreateSessionTestAsync()
-		{
-			RpcCreateSessionRequest req = new RpcCreateSessionRequest();
-			await Assert.ThrowsExceptionAsync<StructuredRpcException>(() => RpcService.CreateSession(req, _adminContext));
-
-			req.Id = new AgentId("MyName").ToString();
-			await Assert.ThrowsExceptionAsync<StructuredRpcException>(() => RpcService.CreateSession(req, _adminContext));
-
-			req.Capabilities = new RpcAgentCapabilities();
-			RpcCreateSessionResponse res = await RpcService.CreateSession(req, _adminContext);
-
-			Assert.AreEqual("MYNAME", res.AgentId);
-			// TODO: Check Token, ExpiryTime, SessionId 
 		}
 
 		[TestMethod]
@@ -351,10 +274,8 @@ namespace HordeServer.Tests.Server
 			string agentId = createRes.AgentId;
 			string sessionId = createRes.SessionId;
 
-			TestAsyncStreamReader<RpcUpdateSessionRequest> requestStream =
-				new TestAsyncStreamReader<RpcUpdateSessionRequest>(_adminContext);
-			TestServerStreamWriter<RpcUpdateSessionResponse> responseStream =
-				new TestServerStreamWriter<RpcUpdateSessionResponse>(_adminContext);
+			TestAsyncStreamReader<RpcUpdateSessionRequest> requestStream = new (_adminContext);
+			TestServerStreamWriter<RpcUpdateSessionResponse> responseStream = new (_adminContext);
 			Task call = RpcService.UpdateSession(requestStream, responseStream, _adminContext);
 
 			requestStream.AddMessage(new RpcUpdateSessionRequest { AgentId = "does-not-exist", SessionId = sessionId });
@@ -375,10 +296,8 @@ namespace HordeServer.Tests.Server
 			string agentId = createRes.AgentId;
 			string sessionId = createRes.SessionId;
 
-			TestAsyncStreamReader<RpcUpdateSessionRequest> requestStream =
-				new TestAsyncStreamReader<RpcUpdateSessionRequest>(_adminContext);
-			TestServerStreamWriter<RpcUpdateSessionResponse> responseStream =
-				new TestServerStreamWriter<RpcUpdateSessionResponse>(_adminContext);
+			TestAsyncStreamReader<RpcUpdateSessionRequest> requestStream = new (_adminContext);
+			TestServerStreamWriter<RpcUpdateSessionResponse> responseStream =new (_adminContext);
 			Task call = RpcService.UpdateSession(requestStream, responseStream, _adminContext);
 
 			requestStream.AddMessage(new RpcUpdateSessionRequest { AgentId = "does-not-exist", SessionId = sessionId });
@@ -386,30 +305,48 @@ namespace HordeServer.Tests.Server
 			Assert.AreEqual(StatusCode.NotFound, re.StatusCode);
 			Assert.IsTrue(re.Message.Contains("Invalid agent name", StringComparison.OrdinalIgnoreCase));
 		}
-
-		/*
-		[TestMethod]
-		public async Task UploadSoftwareAsync()
+		
+		private async Task SetupToolsAsync(CancellationToken cancellationToken)
+		{
+			AddPlugin<ToolsPlugin>();
+			AclEntryConfig aclEntryConfig = new(HordeClaims.AgentDedicatedRoleClaim, [ToolAclAction.DownloadTool]);
+			ToolsConfig toolsConfig = new ();
+			toolsConfig.Tools.Add(new ToolConfig(_dedicatedAgentToolId) { Acl = new AclConfig() { Entries = [aclEntryConfig] } });
+			toolsConfig.Tools.Add(new ToolConfig(_workstationAgentToolId) { Public = true });
+			UpdateConfig(gc => gc.Plugins.AddToolsConfig(toolsConfig) );
+			
+			IToolCollection collection = ServiceProvider.GetRequiredService<IToolCollection>();
+			await CreateAgentToolDeploymentAsync(collection, _dedicatedAgentToolId, "1", cancellationToken);
+			await CreateAgentToolDeploymentAsync(collection, _dedicatedAgentToolId, "2", cancellationToken);
+			await CreateAgentToolDeploymentAsync(collection, _workstationAgentToolId, "1", cancellationToken);
+		}
+		
+		private static async Task<ITool?> CreateAgentToolDeploymentAsync(IToolCollection collection, ToolId toolId, string version, CancellationToken cancellationToken)
 		{
-			TestSetup TestSetup = await GetTestSetup();
+			ITool? tool = await collection.GetAsync(toolId, cancellationToken);
+			Assert.IsNotNull(tool);
+			ToolDeploymentConfig tdc = new () { Version = version, Duration = TimeSpan.FromMinutes(5.0), CreatePaused = true };
 			
-			MemoryStream OutputStream = new MemoryStream();
-			using (ZipArchive ZipFile = new ZipArchive(OutputStream, ZipArchiveMode.Create, false))
+			const string FileName = "version.txt";
+			byte[] fileData = Encoding.UTF8.GetBytes(version);
 			{
-				string TempFilename = Path.GetTempFileName();
-				File.WriteAllText(TempFilename, "{\"Horde\": {\"Version\": \"myVersion\"}}");
-				ZipFile.CreateEntryFromFile(TempFilename, "appsettings.json");
+				using MemoryStream stream = new();
+				using (ZipArchive archive = new (stream, ZipArchiveMode.Create, leaveOpen: true))
+				{
+					ZipArchiveEntry entry = archive.CreateEntry(FileName);
+					await using (Stream entryStream = entry.Open())
+					{
+						await entryStream.WriteAsync(fileData, cancellationToken);
+					}
+				}
+				stream.Position = 0;
+				return await tool.CreateDeploymentAsync(tdc, stream, cancellationToken);
 			}
-
-			ByteString Data = ByteString.CopyFrom(OutputStream.ToArray());
-			UploadSoftwareRequest Req = new UploadSoftwareRequest { Channel = "boguschannel", Data = Data };
-
-			UploadSoftwareResponse Res1 = await TestSetup.RpcService.UploadSoftware(Req, AdminContext);
-			Assert.AreEqual("r1", Res1.Version);
-			
-			UploadSoftwareResponse Res2 = await TestSetup.RpcService.UploadSoftware(Req, AdminContext);
-			Assert.AreEqual("r2", Res2.Version);
 		}
-		*/
+		
+		private static ServerCallContextStub CreateServerContext(string? roleClaim)
+		{
+			return new ServerCallContextStub(roleClaim != null ? new AclClaimConfig(HordeClaimTypes.Role, roleClaim).ToClaim() : null);
+		}
 	}
 }
\ No newline at end of file
diff --git a/Engine/Source/Programs/Horde/Plugins/Build/HordeServer.Build/BuildAclModifier.cs b/Engine/Source/Programs/Horde/Plugins/Build/HordeServer.Build/BuildAclModifier.cs
--- a/Engine/Source/Programs/Horde/Plugins/Build/HordeServer.Build/BuildAclModifier.cs
+++ b/Engine/Source/Programs/Horde/Plugins/Build/HordeServer.Build/BuildAclModifier.cs
@@ -1,6 +1,7 @@
 // Copyright Epic Games, Inc. All Rights Reserved.
 
 using HordeServer.Acls;
+using HordeServer.Agents;
 using HordeServer.Agents.Software;
 using HordeServer.Artifacts;
 using HordeServer.Devices;
@@ -20,7 +21,8 @@ namespace HordeServer
 		public void Apply(DefaultAclBuilder acl)
 		{
 			acl.AddCustomRole(HordeClaims.AgentRegistrationClaim, new[] { StreamAclAction.ViewStream, ProjectAclAction.ViewProject, JobAclAction.ViewJob });
-			acl.AddCustomRole(HordeClaims.AgentRoleClaim, new[] { ProjectAclAction.ViewProject, StreamAclAction.ViewStream, LogAclAction.CreateEvent, AgentSoftwareAclAction.DownloadSoftware });
+			acl.AddCustomRole(HordeClaims.AgentDedicatedRoleClaim, new[] { ProjectAclAction.ViewProject, StreamAclAction.ViewStream, LogAclAction.CreateEvent, AgentSoftwareAclAction.DownloadSoftware });
+			acl.AddCustomRole(HordeClaims.AgentWorkstationRoleClaim, new[] { AgentAclAction.CreateWorkstationAgent, AgentSoftwareAclAction.DownloadSoftware });
 			acl.AddCustomRole(HordeClaims.ConfigureProjectsClaim, new[] { ProjectAclAction.CreateProject, ProjectAclAction.UpdateProject, ProjectAclAction.ViewProject, StreamAclAction.CreateStream, StreamAclAction.UpdateStream, StreamAclAction.ViewStream });
 			acl.AddCustomRole(HordeClaims.StartChainedJobClaim, new[] { JobAclAction.CreateJob, JobAclAction.ExecuteJob, JobAclAction.UpdateJob, JobAclAction.ViewJob, StreamAclAction.ViewTemplate, StreamAclAction.ViewStream });
 
diff --git a/Engine/Source/Programs/Horde/Plugins/Build/HordeServer.Build/Jobs/JobTaskSource.cs b/Engine/Source/Programs/Horde/Plugins/Build/HordeServer.Build/Jobs/JobTaskSource.cs
--- a/Engine/Source/Programs/Horde/Plugins/Build/HordeServer.Build/Jobs/JobTaskSource.cs
+++ b/Engine/Source/Programs/Horde/Plugins/Build/HordeServer.Build/Jobs/JobTaskSource.cs
@@ -885,7 +885,7 @@ namespace HordeServer.Jobs
 
 			// Create a bearer token for the job executor
 			List<AclClaimConfig> claims = new List<AclClaimConfig>();
-			claims.Add(HordeClaims.AgentRoleClaim);
+			claims.Add(HordeClaims.AgentDedicatedRoleClaim);
 			claims.Add(new AclClaimConfig(HordeClaimTypes.Lease, leaseId.ToString()));
 			claims.Add(new AclClaimConfig(HordeClaimTypes.LeaseStream, streamConfig.Id.ToString()));
 			claims.Add(new AclClaimConfig(HordeClaimTypes.LeaseProject, streamConfig.ProjectConfig.Id.ToString()));
diff --git a/Engine/Source/Programs/Horde/Plugins/Compute/HordeServer.Compute.Tests/Agents/AgentCollectionTests.cs b/Engine/Source/Programs/Horde/Plugins/Compute/HordeServer.Compute.Tests/Agents/AgentCollectionTests.cs
--- a/Engine/Source/Programs/Horde/Plugins/Compute/HordeServer.Compute.Tests/Agents/AgentCollectionTests.cs
+++ b/Engine/Source/Programs/Horde/Plugins/Compute/HordeServer.Compute.Tests/Agents/AgentCollectionTests.cs
@@ -36,7 +36,7 @@ public class AgentCollectionTests : ComputeTestSetup
 	[TestInitialize]
 	public async Task SetupAsync()
 	{
-		_agent = await AgentCollection.AddAsync(new CreateAgentOptions(new AgentId("test"), true, ""));
+		_agent = await AgentCollection.AddAsync(new CreateAgentOptions(new AgentId("test"), AgentMode.Dedicated, true, ""));
 	}
 
 	[TestMethod]
diff --git a/Engine/Source/Programs/Horde/Plugins/Compute/HordeServer.Compute.Tests/Agents/AgentsControllerTests.cs b/Engine/Source/Programs/Horde/Plugins/Compute/HordeServer.Compute.Tests/Agents/AgentsControllerTests.cs
--- a/Engine/Source/Programs/Horde/Plugins/Compute/HordeServer.Compute.Tests/Agents/AgentsControllerTests.cs
+++ b/Engine/Source/Programs/Horde/Plugins/Compute/HordeServer.Compute.Tests/Agents/AgentsControllerTests.cs
@@ -107,7 +107,7 @@ public class AgentsControllerTests : ComputeTestSetup
 	
 	private async Task<IAgent> CreateAgentAsync(string name, List<string> properties)
 	{
-		IAgent agent = await AgentService.CreateAgentAsync(new CreateAgentOptions(new AgentId(name), false, ""));
+		IAgent agent = await AgentService.CreateAgentAsync(new CreateAgentOptions(new AgentId(name), AgentMode.Dedicated, false, ""));
 		agent = await AgentService.CreateSessionAsync(agent, new RpcAgentCapabilities(properties), null);
 		return agent;
 	}
diff --git a/Engine/Source/Programs/Horde/Plugins/Compute/HordeServer.Compute/Agents/AgentAclAction.cs b/Engine/Source/Programs/Horde/Plugins/Compute/HordeServer.Compute/Agents/AgentAclAction.cs
--- a/Engine/Source/Programs/Horde/Plugins/Compute/HordeServer.Compute/Agents/AgentAclAction.cs
+++ b/Engine/Source/Programs/Horde/Plugins/Compute/HordeServer.Compute/Agents/AgentAclAction.cs
@@ -13,6 +13,11 @@ namespace HordeServer.Agents
 		/// Ability to create an agent. This may be done explicitly, or granted to agents to allow them to self-register.
 		/// </summary>
 		public static AclAction CreateAgent { get; } = new AclAction("CreateAgent");
+		
+		/// <summary>
+		/// Ability to create a workstation agent. Granted to users with hardware to offer.
+		/// </summary>
+		public static AclAction CreateWorkstationAgent { get; } = new AclAction("CreateWorkstationAgent");
 
 		/// <summary>
 		/// Update an agent's name, pools, etc...
diff --git a/Engine/Source/Programs/Horde/Plugins/Compute/HordeServer.Compute/Agents/AgentCollection.cs b/Engine/Source/Programs/Horde/Plugins/Compute/HordeServer.Compute/Agents/AgentCollection.cs
--- a/Engine/Source/Programs/Horde/Plugins/Compute/HordeServer.Compute/Agents/AgentCollection.cs
+++ b/Engine/Source/Programs/Horde/Plugins/Compute/HordeServer.Compute/Agents/AgentCollection.cs
@@ -55,6 +55,7 @@ namespace HordeServer.Agents
 			DateTime? IAgent.LastOnlineTime => _document.LastOnlineTime;
 			bool IAgent.Enabled => _document.Enabled;
 			bool IAgent.Ephemeral => _document.Ephemeral;
+			AgentMode IAgent.Mode => _document.Mode;
 			bool IAgent.Deleted => _document.Deleted;
 			string? IAgent.Version => _document.Version;
 			string? IAgent.Comment => _document.Comment;
@@ -157,6 +158,8 @@ namespace HordeServer.Agents
 			public bool Enabled { get; set; } = true;
 
 			public bool Ephemeral { get; set; }
+			
+			public AgentMode Mode { get; set; } = AgentMode.Dedicated;
 
 			[BsonIgnoreIfDefault, BsonDefaultValue(false)]
 			public bool Deleted { get; set; }
diff --git a/Engine/Source/Programs/Horde/Plugins/Compute/HordeServer.Compute/Agents/AgentService.cs b/Engine/Source/Programs/Horde/Plugins/Compute/HordeServer.Compute/Agents/AgentService.cs
--- a/Engine/Source/Programs/Horde/Plugins/Compute/HordeServer.Compute/Agents/AgentService.cs
+++ b/Engine/Source/Programs/Horde/Plugins/Compute/HordeServer.Compute/Agents/AgentService.cs
@@ -176,7 +176,7 @@ namespace HordeServer.Agents
 		public async ValueTask<string> IssueSessionTokenAsync(AgentId agentId, SessionId sessionId, CancellationToken cancellationToken = default)
 		{
 			List<AclClaimConfig> claims = new List<AclClaimConfig>();
-			claims.Add(HordeClaims.AgentRoleClaim);
+			claims.Add(HordeClaims.AgentDedicatedRoleClaim);
 			claims.Add(HordeClaims.GetAgentClaim(agentId));
 			claims.Add(HordeClaims.GetSessionClaim(sessionId));
 			return await _aclService.IssueBearerTokenAsync(claims, null, cancellationToken);
diff --git a/Engine/Source/Programs/Horde/Plugins/Compute/HordeServer.Compute/Agents/Enrollment/EnrollmentRpc.cs b/Engine/Source/Programs/Horde/Plugins/Compute/HordeServer.Compute/Agents/Enrollment/EnrollmentRpc.cs
--- a/Engine/Source/Programs/Horde/Plugins/Compute/HordeServer.Compute/Agents/Enrollment/EnrollmentRpc.cs
+++ b/Engine/Source/Programs/Horde/Plugins/Compute/HordeServer.Compute/Agents/Enrollment/EnrollmentRpc.cs
@@ -58,7 +58,7 @@ namespace HordeServer.Agents.Enrollment
 						return;
 					}
 
-					IAgent agent = await _agentService.CreateAgentAsync(new CreateAgentOptions(agentId, false, request.Key), cancellationSource.Token);
+					IAgent agent = await _agentService.CreateAgentAsync(new CreateAgentOptions(agentId, AgentMode.Dedicated, false, request.Key), cancellationSource.Token);
 
 					List<AclClaimConfig> claims = new List<AclClaimConfig>();
 					claims.Add(new AclClaimConfig(HordeClaimTypes.Agent, agentId.ToString()));
diff --git a/Engine/Source/Programs/Horde/Plugins/Compute/HordeServer.Compute/Agents/IAgent.cs b/Engine/Source/Programs/Horde/Plugins/Compute/HordeServer.Compute/Agents/IAgent.cs
--- a/Engine/Source/Programs/Horde/Plugins/Compute/HordeServer.Compute/Agents/IAgent.cs
+++ b/Engine/Source/Programs/Horde/Plugins/Compute/HordeServer.Compute/Agents/IAgent.cs
@@ -24,6 +24,34 @@ using MongoDB.Bson.Serialization.Attributes;
 
 namespace HordeServer.Agents
 {
+	/// <summary>
+	/// Defines the operation mode of the agent
+	/// </summary>
+	public enum AgentMode
+	{
+		/// <summary>
+		/// For backwards compatibility
+		/// </summary>
+		Unspecified = 0,
+		
+		/// <summary>
+		/// Dedicated mode
+		/// - Trusted agent running in a controlled environment (e.g., build farm)
+		/// - Does not run any other non-Horde workloads
+		/// - Capable of executing all lease types
+		/// </summary>
+		Dedicated = 1,
+		
+		/// <summary>
+		/// Workstation mode
+		/// - Agent authenticates as the current user through interactive authentication via EpicGames.OIDC
+		/// - Low trust level, typically used for spare on-premise hardware like workstations
+		/// - Lease execution is inherently unreliable as these agents yield to non-Horde workloads
+		/// - Limited to compute leases only, primarily to support remote execution
+		/// </summary>
+		Workstation = 2
+	}
+	
 	/// <summary>
 	/// Mirrors an Agent document in the database
 	/// </summary>
@@ -63,6 +91,11 @@ namespace HordeServer.Agents
 		/// Whether the agent is ephemeral
 		/// </summary>
 		public bool Ephemeral { get; }
+		
+		/// <summary>
+		/// Operation mode of agent
+		/// </summary>
+		public AgentMode Mode { get; }
 
 		/// <summary>
 		/// Whether the agent should be included on the dashboard. This is set to true for ephemeral agents once they are no longer online, or agents that are explicitly deleted.
@@ -295,10 +328,11 @@ namespace HordeServer.Agents
 	/// Options for creating an agent
 	/// </summary>
 	/// <param name="Id">ID of the agent</param>
+	/// <param name="Mode">Operation mode</param>
 	/// <param name="Ephemeral">Whether the agent is ephemeral or not</param>
 	/// <param name="EnrollmentKey">Key used to identify a unique enrollment for the agent with this ID</param>
 	/// <param name="ServerDefinedProperties">Server-defined properties</param>
-	public record CreateAgentOptions(AgentId Id, bool Ephemeral, string EnrollmentKey, IReadOnlyList<string>? ServerDefinedProperties = null);
+	public record CreateAgentOptions(AgentId Id, AgentMode Mode, bool Ephemeral, string EnrollmentKey, IReadOnlyList<string>? ServerDefinedProperties = null);
 	
 	/// <summary>
 	/// Options for updating an agent
diff --git a/Engine/Source/Programs/Horde/Plugins/Compute/HordeServer.Compute/Server/RpcService.cs b/Engine/Source/Programs/Horde/Plugins/Compute/HordeServer.Compute/Server/RpcService.cs
--- a/Engine/Source/Programs/Horde/Plugins/Compute/HordeServer.Compute/Server/RpcService.cs
+++ b/Engine/Source/Programs/Horde/Plugins/Compute/HordeServer.Compute/Server/RpcService.cs
@@ -3,6 +3,7 @@
 using System.Buffers;
 using System.Net;
 using System.Security.Claims;
+using EpicGames.Horde.Acls;
 using EpicGames.Horde.Agents;
 using EpicGames.Horde.Agents.Sessions;
 using EpicGames.Horde.Tools;
@@ -149,32 +150,47 @@ namespace HordeServer.Server
 
 		/// <summary>
 		/// Creates a new agent
-		/// Assumes user has permission to create new agents as it skips the enrollment process
+		/// The enrollment process is bypassed for all agent types created through this endpoint.
 		/// </summary>
 		/// <param name="request">Request to create a new agent</param>
 		/// <param name="context">Context for the RPC call</param>
-		/// <returns>Information about the new agent</returns>
+		/// <returns>New agent's ID and authentication token (if applicable)</returns>
 		public override async Task<RpcCreateAgentResponse> CreateAgent(RpcCreateAgentRequest request, ServerCallContext context)
 		{
 			using IDisposable? scope = _logger.BeginScope("CreateAgent({AgentId})", request.Name.ToString());
-
-			if (!_computeConfig.Value.Authorize(AgentAclAction.CreateAgent, context.GetHttpContext().User))
+			
+			switch (request.Mode)
 			{
-				throw new StructuredRpcException(StatusCode.PermissionDenied, "User is not authenticated to create new agents");
+				case Mode.Unspecified or Mode.Dedicated:
+					EnsurePermission(AgentAclAction.CreateAgent, "User is not authorized to create a dedicated agent", context);
+					// Dedicated agents are assumed trusted
+					CreateAgentOptions options = new(new AgentId(request.Name), AgentMode.Dedicated, request.Ephemeral, "", [$"{KnownPropertyNames.Trusted}=true"]);
+					IAgent agent = await _agentService.CreateAgentAsync(options, context.CancellationToken);
+					
+					List<AclClaimConfig> claims = [new AclClaimConfig(HordeClaimTypes.Agent, agent.Id.ToString())];
+					return new RpcCreateAgentResponse
+					{
+						Id = agent.Id.ToString(),
+						Token = await _aclService.IssueBearerTokenAsync(claims, null, context.CancellationToken)
+					};
+				case Mode.Workstation:
+					EnsurePermission(AgentAclAction.CreateWorkstationAgent, "User is not authorized to create a workstation agent", context);
+					CreateAgentOptions workstationOptions = new(new AgentId(request.Name), AgentMode.Workstation, request.Ephemeral, "");
+					IAgent workstationAgent = await _agentService.CreateAgentAsync(workstationOptions, context.CancellationToken);
+					// Workstation-based agents always authenticate as the current user and should *not* be issued a token
+					return new RpcCreateAgentResponse { Id = workstationAgent.Id.ToString(), Token = "" };
+					
+				default:
+					throw new StructuredRpcException(StatusCode.InvalidArgument, $"Invalid agent mode: {request.Mode}");
+			}
+		}
+		
+		private void EnsurePermission(AclAction action, string errorMessage, ServerCallContext context)
+		{
+			if (!_computeConfig.Value.Authorize(action, context.GetHttpContext().User))
+			{
+				throw new StructuredRpcException(StatusCode.PermissionDenied, errorMessage);
 			}
-
-			// If the agent/user has permission to create an agent, it's assumed the agent is also trusted
-			CreateAgentOptions options = new(new AgentId(request.Name), request.Ephemeral, "", [$"{KnownPropertyNames.Trusted}=true"]);
-			IAgent agent = await _agentService.CreateAgentAsync(options, context.CancellationToken);
-
-			List<AclClaimConfig> claims = new List<AclClaimConfig>();
-			claims.Add(new AclClaimConfig(HordeClaimTypes.Agent, agent.Id.ToString()));
-
-			RpcCreateAgentResponse response = new RpcCreateAgentResponse();
-			response.Id = agent.Id.ToString();
-			response.Token = await _aclService.IssueBearerTokenAsync(claims, null, context.CancellationToken);
-
-			return response;
 		}
 
 		/// <summary>
@@ -205,7 +221,7 @@ namespace HordeServer.Server
 					throw new StructuredRpcException(StatusCode.PermissionDenied, "User is not authenticated to create new agents");
 				}
 
-				agent = await _agentService.CreateAgentAsync(new CreateAgentOptions(agentId, false, ""));
+				agent = await _agentService.CreateAgentAsync(new CreateAgentOptions(agentId, AgentMode.Dedicated, false, ""));
 			}
 
 			// Check the enrollment key in the user token matches
diff --git a/Engine/Source/Programs/Horde/Plugins/Tools/HordeServer.Tools.Tests/Tools/ToolsControllerTests.cs b/Engine/Source/Programs/Horde/Plugins/Tools/HordeServer.Tools.Tests/Tools/ToolsControllerTests.cs
--- a/Engine/Source/Programs/Horde/Plugins/Tools/HordeServer.Tools.Tests/Tools/ToolsControllerTests.cs
+++ b/Engine/Source/Programs/Horde/Plugins/Tools/HordeServer.Tools.Tests/Tools/ToolsControllerTests.cs
@@ -23,7 +23,7 @@ public class ToolsControllerTests
 	[TestMethod]
 	public async Task DownloadToolWithAclAsync()
 	{
-		AclEntryConfig aclEntryConfig = new(HordeClaims.AgentRoleClaim, [ToolAclAction.DownloadTool]);
+		AclEntryConfig aclEntryConfig = new(HordeClaims.AgentDedicatedRoleClaim, [ToolAclAction.DownloadTool]);
 		ToolId toolId = new("foo");
 
 		PluginCollection pluginCollection = new PluginCollection();
