Fix for Horde Analytics Package Dashbaord
Install Type variable was incorrectly named.

diff --git a/Engine/Source/Programs/Horde/HordeServer/Defaults/default-views.telemetry.json b/Engine/Source/Programs/Horde/HordeServer/Defaults/default-views.telemetry.json
--- a/Engine/Source/Programs/Horde/HordeServer/Defaults/default-views.telemetry.json
+++ b/Engine/Source/Programs/Horde/HordeServer/Defaults/default-views.telemetry.json
@@ -106,7 +106,7 @@
 				},
 				{
 					"name": "Install Type",
-					"group": "Install_Type"
+					"group": "InstallType"
 				}
 			],
 			"categories": [
