Changed display type to Value, rather than Time, for Package data dashboard in Horde Analytics

diff --git a/Engine/Source/Programs/Horde/HordeServer/Defaults/default-views.telemetry.json b/Engine/Source/Programs/Horde/HordeServer/Defaults/default-views.telemetry.json
--- a/Engine/Source/Programs/Horde/HordeServer/Defaults/default-views.telemetry.json
+++ b/Engine/Source/Programs/Horde/HordeServer/Defaults/default-views.telemetry.json
@@ -742,6 +742,7 @@
 					"charts": [
 						{
 							"name": "Size By Install Type",
+							"display": "Value",
 							"metrics": [
 								{
 									"id": "package-installed-total-size",
@@ -763,6 +764,7 @@
 						},
 						{
 							"name": "Size Per Package",
+							"display": "Value",
 							"metrics": [
 								{
 									"id": "package-size"
