Horde: Updating schema documentation

#rnx

diff --git a/Engine/Source/Programs/Horde/Docs/Config/Schema/AclActions.md b/Engine/Source/Programs/Horde/Docs/Config/Schema/AclActions.md
--- a/Engine/Source/Programs/Horde/Docs/Config/Schema/AclActions.md
+++ b/Engine/Source/Programs/Horde/Docs/Config/Schema/AclActions.md
@@ -163,6 +163,7 @@
 | Name | Description |
 | ---- | ----------- |
 | `CreateAgent` | Ability to create an agent. This may be done explicitly, or granted to agents to allow them to self-register. |
+| `CreateWorkstationAgent` | Ability to create a workstation agent. Granted to users with hardware to offer. |
 | `UpdateAgent` | Update an agent's name, pools, etc... |
 | `DeleteAgent` | Soft-delete an agent |
 | `ViewAgent` | View an agent |
@@ -318,6 +319,7 @@
 | Name | Description |
 | ---- | ----------- |
 | `CreateAgent` |  |
+| `CreateWorkstationAgent` |  |
 | `UpdateAgent` |  |
 | `DeleteAgent` |  |
 | `ViewAgent` |  |
@@ -486,6 +488,7 @@
 | Name | Description |
 | ---- | ----------- |
 | `CreateAgent` |  |
+| `CreateWorkstationAgent` |  |
 | `UpdateAgent` |  |
 | `DeleteAgent` |  |
 | `ViewAgent` |  |
@@ -654,6 +657,7 @@
 | Name | Description |
 | ---- | ----------- |
 | `CreateAgent` |  |
+| `CreateWorkstationAgent` |  |
 | `UpdateAgent` |  |
 | `DeleteAgent` |  |
 | `ViewAgent` |  |
@@ -828,6 +832,7 @@
 | Name | Description |
 | ---- | ----------- |
 | `CreateAgent` |  |
+| `CreateWorkstationAgent` |  |
 | `UpdateAgent` |  |
 | `DeleteAgent` |  |
 | `ViewAgent` |  |
@@ -1015,6 +1020,7 @@
 | Name | Description |
 | ---- | ----------- |
 | `CreateAgent` |  |
+| `CreateWorkstationAgent` |  |
 | `UpdateAgent` |  |
 | `DeleteAgent` |  |
 | `ViewAgent` |  |
@@ -1188,12 +1194,11 @@
 | `ViewChanges` |  |
 | `ViewTemplate` |  |
 
-## Tools
+## Symbols
 
 | Name | Description |
 | ---- | ----------- |
-| `DownloadTool` | Ability to download a tool |
-| `UploadTool` | Ability to upload new tool versions |
+| `ReadSymbols` | Ability to download symbols |
 
 ## Accounts
 
@@ -1209,6 +1214,7 @@
 | Name | Description |
 | ---- | ----------- |
 | `CreateAgent` |  |
+| `CreateWorkstationAgent` |  |
 | `UpdateAgent` |  |
 | `DeleteAgent` |  |
 | `ViewAgent` |  |
@@ -1386,14 +1392,14 @@
 
 | Name | Description |
 | ---- | ----------- |
-| `ReadSymbols` | Ability to download symbols |
+| `ReadSymbols` |  |
 
 ## Tools
 
 | Name | Description |
 | ---- | ----------- |
-| `DownloadTool` |  |
-| `UploadTool` |  |
+| `DownloadTool` | Ability to download a tool |
+| `UploadTool` | Ability to upload new tool versions |
 
 ## Accounts
 
@@ -1409,6 +1415,7 @@
 | Name | Description |
 | ---- | ----------- |
 | `CreateAgent` |  |
+| `CreateWorkstationAgent` |  |
 | `UpdateAgent` |  |
 | `DeleteAgent` |  |
 | `ViewAgent` |  |
