Added Horde_TemplateName as a variable in the Horde Analytics dashboard

diff --git a/Engine/Source/Programs/Horde/HordeServer/Defaults/default-views.telemetry.json b/Engine/Source/Programs/Horde/HordeServer/Defaults/default-views.telemetry.json
--- a/Engine/Source/Programs/Horde/HordeServer/Defaults/default-views.telemetry.json
+++ b/Engine/Source/Programs/Horde/HordeServer/Defaults/default-views.telemetry.json
@@ -40,6 +40,10 @@
 					"name": "Cook Platform",
 					"group": "Cook_Profile_TargetPlatforms"
 				},
+				{
+					"name": "Horde TemplateName",
+					"group": "Horde TemplateName"
+				},
 				{
 					"name": "Horde JobName",
 					"group": "Horde_JobName"
