Added missing _ to Horde Template Name variable

diff --git a/Engine/Source/Programs/Horde/HordeServer/Defaults/default-views.telemetry.json b/Engine/Source/Programs/Horde/HordeServer/Defaults/default-views.telemetry.json
--- a/Engine/Source/Programs/Horde/HordeServer/Defaults/default-views.telemetry.json
+++ b/Engine/Source/Programs/Horde/HordeServer/Defaults/default-views.telemetry.json
@@ -42,7 +42,7 @@
 				},
 				{
 					"name": "Horde TemplateName",
-					"group": "Horde TemplateName"
+					"group": "Horde_TemplateName"
 				},
 				{
 					"name": "Horde JobName",
