Horde: Use agent's self-reported tool ID for tool selection during upgrades

Agents now send this information to provide correct hinting for what tool to use. For example multi-platform, self-contained, or even custom tools for workstation-based agents.

diff --git a/Engine/Source/Programs/Horde/Plugins/Compute/HordeServer.Compute.Tests/Agents/AgentCollectionTests.cs b/Engine/Source/Programs/Horde/Plugins/Compute/HordeServer.Compute.Tests/Agents/AgentCollectionTests.cs
--- a/Engine/Source/Programs/Horde/Plugins/Compute/HordeServer.Compute.Tests/Agents/AgentCollectionTests.cs
+++ b/Engine/Source/Programs/Horde/Plugins/Compute/HordeServer.Compute.Tests/Agents/AgentCollectionTests.cs
@@ -11,6 +11,7 @@ using HordeCommon.Rpc;
 using HordeCommon.Rpc.Messages;
 using HordeServer.Agents;
 using System.Security.Cryptography;
+using EpicGames.Horde.Tools;
 using Google.Protobuf;
 
 namespace HordeServer.Tests.Agents;
@@ -342,6 +343,24 @@ public class AgentCollectionTests : ComputeTestSetup
 		Assert.IsTrue(leases.Contains(_leaseWithParent3.Id));
 		Assert.IsTrue(leases.Contains(_leaseWithParent4.Id));
 	}
+	
+	public static IEnumerable<object?[]> GetSoftwareToolIdTestData()
+	{
+		yield return [ null, null, AgentExtensions.AgentToolId.ToString() ];
+		yield return [ "my-tool", null, "my-tool" ];
+		yield return [ "my-tool", new AgentSoftwareConfig { ToolId = new ToolId("other-tool"), Condition = "" }, "other-tool" ];
+		yield return [ null, new AgentSoftwareConfig { ToolId = new ToolId("other-tool"), Condition = "" }, "other-tool" ];
+	}
+	
+	[TestMethod]
+	[DynamicData(nameof(GetSoftwareToolIdTestData), DynamicDataSourceType.Method)]
+	public async Task GetSoftwareToolIdAsync(string? reportedToolId, AgentSoftwareConfig? softwareConfig, string? expectedToolId)
+	{
+		List<string> props = reportedToolId != null ? [$"{KnownPropertyNames.ToolId}={reportedToolId}"] : [];
+		IAgent? agent = await _agent.TryCreateSessionAsync(new CreateSessionOptions(new RpcAgentCapabilities(props), Array.Empty<PoolId>(), "foo"));
+		ToolId actual = agent!.GetSoftwareToolId(new ComputeConfig() { Software = softwareConfig != null ? [softwareConfig] : [] });
+		Assert.AreEqual(expectedToolId, actual.ToString());
+	}
 
 	private async Task UpdateAgentAsync()
 	{
diff --git a/Engine/Source/Programs/Horde/Plugins/Compute/HordeServer.Compute/Agents/IAgent.cs b/Engine/Source/Programs/Horde/Plugins/Compute/HordeServer.Compute/Agents/IAgent.cs
--- a/Engine/Source/Programs/Horde/Plugins/Compute/HordeServer.Compute/Agents/IAgent.cs
+++ b/Engine/Source/Programs/Horde/Plugins/Compute/HordeServer.Compute/Agents/IAgent.cs
@@ -404,10 +404,17 @@ namespace HordeServer.Agents
 		/// <returns>Identifier for the tool that the agent should be using</returns>
 		public static ToolId GetSoftwareToolId(this IAgent agent, ComputeConfig computeConfig)
 		{
-			ToolId toolId = AgentToolId;
-
-			if (agent.IsSelfContained())
+			ToolId result = AgentToolId;
+			
+			string? reportedToolId = agent.GetToolId();
+			if (reportedToolId != null)
+			{
+				result = new ToolId(reportedToolId);
+			}
+			else if (agent.IsSelfContained())
 			{
+				// Let reported tool ID take precedence, this self-contained path and reporting should be deprecated in favor of that 
+				
 				// Skip support for condition-based software configs below by returning early when self-contained
 				// Getting this wrong can lead to a self-contained agent getting non-self-contained updates and vice versa.
 				return agent.GetOsFamily() switch
@@ -423,11 +430,11 @@ namespace HordeServer.Agents
 			{
 				if (softwareConfig.Condition != null && agent.SatisfiesCondition(softwareConfig.Condition))
 				{
-					toolId = softwareConfig.ToolId;
+					result = softwareConfig.ToolId;
 					break;
 				}
 			}
-			return toolId;
+			return result;
 		}
 
 		/// <summary>
@@ -460,6 +467,16 @@ namespace HordeServer.Agents
 			List<string> values = agent.GetPropertyValues(KnownPropertyNames.SelfContained).ToList();
 			return values.Count > 0 && values[0].Equals("true", StringComparison.OrdinalIgnoreCase);
 		}
+		
+		/// <summary>
+		/// Get optional tool ID as reported by agent
+		/// <see cref="KnownPropertyNames.ToolId" />
+		/// </summary>
+		public static string? GetToolId(this IAgent agent)
+		{
+			List<string> values = agent.GetPropertyValues(KnownPropertyNames.ToolId).ToList();
+			return values.Count > 0 ? values[0] : null;
+		}
 
 		/// <summary>
 		/// Get free disk space on agent (as reported through primary device capabilities)
