Horde: Add Wine override setting to JobDriver

Applies to all jobs executed by this driver process, and overrides the value set by JobOption (which is better used for per-job settings).

diff --git a/Engine/Source/Programs/Horde/Drivers/JobDriver/DriverSettings.cs b/Engine/Source/Programs/Horde/Drivers/JobDriver/DriverSettings.cs
--- a/Engine/Source/Programs/Horde/Drivers/JobDriver/DriverSettings.cs
+++ b/Engine/Source/Programs/Horde/Drivers/JobDriver/DriverSettings.cs
@@ -1,5 +1,6 @@
 ﻿// Copyright Epic Games, Inc. All Rights Reserved.
 
+using HordeCommon.Rpc.Tasks;
 using JobDriver.Execution;
 using JobDriver.Utility;
 
@@ -14,6 +15,11 @@ namespace JobDriver
 		/// The executor to use for jobs
 		/// </summary>
 		public string Executor { get; set; } = WorkspaceExecutor.Name;
+		
+		/// <summary>
+		/// Whether to use Wine for executing the job. This normally is configured via <see cref="RpcJobOptions"/>
+		/// </summary>
+		public bool UseWine { get; set; } = false;
 
 		/// <summary>
 		/// Settings for the local executor
diff --git a/Engine/Source/Programs/Horde/Drivers/JobDriver/Execution/JobExecutor.cs b/Engine/Source/Programs/Horde/Drivers/JobDriver/Execution/JobExecutor.cs
--- a/Engine/Source/Programs/Horde/Drivers/JobDriver/Execution/JobExecutor.cs
+++ b/Engine/Source/Programs/Horde/Drivers/JobDriver/Execution/JobExecutor.cs
@@ -2292,7 +2292,11 @@ namespace JobDriver.Execution
 			JobId jobId = JobId.Parse(executeTask.JobId);
 			JobStepBatchId batchId = JobStepBatchId.Parse(executeTask.BatchId);
 
-			JobExecutorOptions options = new JobExecutorOptions(hordeClient, workingDir, driverSettings.ProcessesToTerminate, jobId, batchId, leaseId, executeTask.JobOptions);
+			if (driverSettings.UseWine)
+			{
+				executeTask.JobOptions.UseWine = true;
+			}
+			JobExecutorOptions options = new (hordeClient, workingDir, driverSettings.ProcessesToTerminate, jobId, batchId, leaseId, executeTask.JobOptions);
 
 			using JobExecutor executor = await executorFactory.CreateExecutorAsync(executeTask.Workspace, executeTask.AutoSdkWorkspace, options, cancellationToken);
 			await executor.ExecuteAsync(logger, cancellationToken);
