Implemented logic to consolidate multiple related jobs into a single notification

#rnx

diff --git a/Engine/Source/Programs/Horde/Plugins/Experimental/HordeServer.Experimental.Tests/ExperimentalTestSetup.cs b/Engine/Source/Programs/Horde/Plugins/Experimental/HordeServer.Experimental.Tests/ExperimentalTestSetup.cs
--- a/Engine/Source/Programs/Horde/Plugins/Experimental/HordeServer.Experimental.Tests/ExperimentalTestSetup.cs
+++ b/Engine/Source/Programs/Horde/Plugins/Experimental/HordeServer.Experimental.Tests/ExperimentalTestSetup.cs
@@ -1,7 +1,6 @@
 // Copyright Epic Games, Inc. All Rights Reserved.
 
-using HordeServer.Notifications;
-using HordeServer.Notifications.Sinks;
+using HordeServer.Experimental.Notifications;
 using Microsoft.Extensions.DependencyInjection;
 
 namespace HordeServer.Tests
@@ -20,7 +19,8 @@ namespace HordeServer.Tests
 		{
 			base.ConfigureServices(services);
 
-			services.AddSingleton<INotificationSink, ExperimentalSlackNotificationSink>();
+			services.AddSingleton<JobNotificationCollection>();
+			services.AddSingleton<IJobNotificationCollection>(sp => sp.GetRequiredService<JobNotificationCollection>());
 		}
 	}
 }
\ No newline at end of file
diff --git a/Engine/Source/Programs/Horde/Plugins/Experimental/HordeServer.Experimental.Tests/Notifications/ExperimentalNotificationTests.cs b/Engine/Source/Programs/Horde/Plugins/Experimental/HordeServer.Experimental.Tests/Notifications/ExperimentalNotificationTests.cs
new file mode 100644
--- /dev/null
+++ b/Engine/Source/Programs/Horde/Plugins/Experimental/HordeServer.Experimental.Tests/Notifications/ExperimentalNotificationTests.cs
@@ -0,0 +1,1290 @@
+// Copyright Epic Games, Inc. All Rights Reserved.
+
+using EpicGames.Core;
+using EpicGames.Horde;
+using EpicGames.Horde.Agents;
+using EpicGames.Horde.Agents.Leases;
+using EpicGames.Horde.Agents.Pools;
+using EpicGames.Horde.Agents.Sessions;
+using EpicGames.Horde.Commits;
+using EpicGames.Horde.Jobs;
+using EpicGames.Horde.Jobs.Graphs;
+using EpicGames.Horde.Jobs.Templates;
+using EpicGames.Horde.Logs;
+using EpicGames.Horde.Projects;
+using EpicGames.Horde.Streams;
+using EpicGames.Horde.Users;
+using HordeServer.Agents;
+using HordeServer.Configuration;
+using HordeServer.Devices;
+using HordeServer.Experimental.Notifications;
+using HordeServer.Issues;
+using HordeServer.Jobs;
+using HordeServer.Logs;
+using HordeServer.Notifications;
+using HordeServer.Projects;
+using HordeServer.Server;
+using HordeServer.Notifications.Sinks;
+using HordeServer.Streams;
+using HordeServer.Tests;
+using HordeServer.Users;
+using HordeServer.Utilities;
+using Microsoft.Extensions.DependencyInjection;
+using Microsoft.Extensions.Logging;
+using Microsoft.Extensions.Options;
+using Moq;
+
+namespace HordeServer.Experimental.Tests.Notifications
+{
+	class TemplateStub(ContentHash hash, string name) : ITemplate
+	{
+		public ContentHash Hash { get; } = hash;
+		public string Name { get; } = name;
+		public string? Description { get; init; } = null;
+		public Priority? Priority { get; init; } = null;
+		public bool AllowPreflights { get; init; } = false;
+		public bool UpdateIssues { get; init; } = false;
+		public bool PromoteIssuesByDefault { get; init; } = false;
+		public string? InitialAgentType { get; init; } = null;
+		public string? SubmitNewChange { get; init; } = null;
+		public string? SubmitDescription { get; init; } = null;
+		public IReadOnlyList<string> Arguments { get; init; } = [];
+		public IReadOnlyList<ITemplateParameter> Parameters { get; init; } = [];
+	}
+
+	struct JobNotificationState
+	{
+		public JobNotificationState(JobId jobId, TemplateId templateId, string channel)
+		{
+			JobId = jobId;
+			TemplateId = templateId;
+			Channel = channel;
+		}
+
+		public JobId JobId { get; }
+		public TemplateId TemplateId { get; }
+		public string Channel { get; }
+	}
+
+	public class JobNotificationInformation
+	{
+		public IJob Job { get; }
+		public IJobStep? Step { get; } = null;
+		public List<ILogEventData>? JobStepEventData { get; } = null;
+		public IEnumerable<IUser>? UsersToNotify { get; } = null;
+		public JobNotificationInformation(IJob job, IJobStep? step = null, List<ILogEventData>? jobStepEventData = null, IEnumerable<IUser>? usersToNotify = null)
+		{
+			Job = job;
+			Step = step;
+			JobStepEventData = jobStepEventData;
+			UsersToNotify = usersToNotify;
+		}
+	}
+
+	public class FakeExperimentalNotificationSink : INotificationSink, IDisposable
+	{
+		readonly SlackNotificationProcessor _slackNotificationProcessor;
+
+		public JobNotificationInformation? ProcessedJobNotification { get; set; } = null;
+
+		public Task NotifyJobScheduledAsync(List<JobScheduledNotification> notifications, CancellationToken cancellationToken) { throw new NotImplementedException(); }
+
+		public async Task NotifyJobStepCompleteAsync(IEnumerable<IUser>? usersToNotify, IJob job, IJobStepBatch batch, IJobStep step, INode node, List<ILogEventData> jobStepEventData, CancellationToken cancellationToken)
+		{
+			// Currently we only want to notify from the plugin's configuration which would not have this set
+			if (usersToNotify is not null)
+			{
+				return;
+			}
+
+			await _slackNotificationProcessor.ProcessJobStepCompleteAsync(job, step, jobStepEventData, cancellationToken);
+			ProcessedJobNotification = new JobNotificationInformation(job, step, jobStepEventData, usersToNotify);
+		}
+
+		public async Task NotifyJobStepAbortedAsync(IEnumerable<IUser>? usersToNotify, IJob job, IJobStepBatch batch, IJobStep step, INode node, List<ILogEventData> jobStepEventData, CancellationToken cancellationToken)
+		{
+			await _slackNotificationProcessor.ProcessJobStepAbortedAsync(job, step, cancellationToken);
+			ProcessedJobNotification = new JobNotificationInformation(job, step, jobStepEventData, usersToNotify);
+		}
+
+		public async Task NotifyJobCompleteAsync(IJob job, IGraph graph, LabelOutcome outcome, CancellationToken cancellationToken)
+		{
+			await _slackNotificationProcessor.ProcessJobCompleteAsync(job, cancellationToken);
+			ProcessedJobNotification = new JobNotificationInformation(job, null, null, null);
+		}
+
+		public Task NotifyJobCompleteAsync(IUser user, IJob job, IGraph graph, LabelOutcome outcome, CancellationToken cancellationToken)
+		{
+			// Not handled, but we don't want to throw an exception
+			return Task.CompletedTask;
+		}
+
+		public Task NotifyLabelCompleteAsync(IUser user, IJob job, ILabel label, int labelIdx, LabelOutcome outcome, List<(string, JobStepOutcome, Uri)> stepData, CancellationToken cancellationToken) { throw new NotImplementedException(); }
+		public Task NotifyIssueUpdatedAsync(IIssue issue, CancellationToken cancellationToken) { throw new NotImplementedException(); }
+		public Task NotifyConfigUpdateAsync(ConfigUpdateInfo info, CancellationToken cancellationToken) => Task.CompletedTask;
+		public Task NotifyConfigUpdateFailureAsync(string errorMessage, string fileName, int? change = null, IUser? author = null, string? description = null, CancellationToken cancellationToken = default) { throw new NotImplementedException(); }
+		public Task NotifyDeviceServiceAsync(string message, IDevice? device = null, IDevicePool? pool = null, StreamConfig? stream = null, IJob? job = null, IJobStep? step = null, INode? node = null, IUser? user = null, CancellationToken cancellationToken = default) { throw new NotImplementedException(); }
+		public Task SendAgentReportAsync(AgentReport report, CancellationToken cancellationToken) => throw new NotImplementedException();
+		public Task SendIssueReportAsync(IssueReportGroup report, CancellationToken cancellationToken) => throw new NotImplementedException();
+		public Task SendDeviceIssueReportAsync(DeviceIssueReport report, CancellationToken cancellationToken) => throw new NotImplementedException();
+		public void Dispose()
+		{
+			Dispose(true);
+			GC.SuppressFinalize(this);
+		}
+
+		protected virtual void Dispose(bool disposing)
+		{
+			if (disposing)
+			{
+				_slackNotificationProcessor.Dispose();
+			}
+		}
+
+		public FakeExperimentalNotificationSink(IOptions<BuildServerConfig> buildServerConfig,
+			IOptionsMonitor<BuildConfig> buildConfig,
+			IOptionsMonitor<ExperimentalConfig> experimentalConfig,
+			IJobNotificationCollection jobNotificationCollection,
+			IJobCollection jobCollection,
+			IUserCollection userCollection,
+			IServerInfo serverInfo,
+			ILogger<FakeExperimentalNotificationSink> logger)
+		{
+			_slackNotificationProcessor = new SlackNotificationProcessor(buildServerConfig, buildConfig, experimentalConfig, jobNotificationCollection, jobCollection, userCollection, serverInfo, logger);
+		}
+	}
+
+	[TestClass]
+	public class ExperimentalNotificationTests : ExperimentalTestSetup
+	{
+		const string MainStreamName = "//UE5/Main";
+		readonly StreamId _mainStreamId = new StreamId(StringId.Sanitize(MainStreamName));
+
+		readonly TemplateId _buildAndTestTemplateId = new TemplateId("test-build-test-job-template");
+		readonly TemplateId _buildTemplateId = new TemplateId("test-build-job-template");
+		readonly TemplateId _spawnedTestTemplateId = new TemplateId("test-spawned-job-template");
+
+		readonly StreamConfig _mainStreamConfig;
+
+		readonly TemplateRefConfig _testBuildTestJobTemplate;
+		readonly TemplateRefConfig _testBuildJobTemplate;
+		readonly TemplateRefConfig _testSpawnedJobTemplate;
+
+		const string ChannelId = "TEST_CHANNEL";
+		const string SpawnJobTestGroup = "Spawn Test Job";
+		const string ClientBootTestGroup = "Client Boot Test";
+		const string ClientProjectTestGroup = "Client Project Tests";
+		const string EditorBootTestGroup = "Editor Boot Test";
+		const string EditorProjectTestGroup = "Editor Project Tests";
+
+		/// <summary>
+		/// Slack emoji used to emphasize state
+		/// </summary>
+		const string PendingBadge = ":large_blue_square:";
+		const string WarningBadge = ":large_yellow_square:";
+		const string FailureBadge = ":large_red_square:";
+		const string SkippedBadge = ":black_square:";
+		const string PassingBadge = ":large_green_square:";
+
+		readonly List<JobNotificationState> _testJobNotifications = new List<JobNotificationState>();
+
+		public ExperimentalNotificationTests()
+		{
+			_testBuildTestJobTemplate = new TemplateRefConfig { Id = _buildAndTestTemplateId };
+			_testBuildJobTemplate = new TemplateRefConfig { Id = _buildTemplateId };
+			_testSpawnedJobTemplate = new TemplateRefConfig { Id = _spawnedTestTemplateId };
+
+			_mainStreamConfig = CreateStream(_mainStreamId, MainStreamName);
+
+			ProjectConfig projectConfig = new ProjectConfig { Id = new ProjectId("ue5"), Name = "UE5" };
+			projectConfig.Streams.Add(_mainStreamConfig);
+
+			BuildConfig buildConfig = new BuildConfig();
+			buildConfig.Projects.Add(projectConfig);
+
+			GlobalConfig globalConfig = new GlobalConfig();
+			globalConfig.Plugins.AddBuildConfig(buildConfig);
+
+			SetConfig(globalConfig);
+		}
+
+		protected override void ConfigureServices(IServiceCollection services)
+		{
+			base.ConfigureServices(services);
+
+			// Replace the ExperimentalSlackNotificationSink with the one used here for our tests
+			// Otherwise, we'll be adding duplicate records to our database
+			ServiceDescriptor descriptor = services.First(x => x.ImplementationType == typeof(ExperimentalSlackNotificationSink));
+			services.Remove(descriptor);
+
+			services.AddSingleton<FakeExperimentalNotificationSink>();
+			services.AddSingleton<INotificationSink>(sp => sp.GetRequiredService<FakeExperimentalNotificationSink>());
+		}
+
+		StreamConfig CreateStream(StreamId streamId, string streamName)
+		{
+			StreamConfig streamConfig = new StreamConfig { Id = streamId, Name = streamName };
+			streamConfig.Tabs.Add(new TabConfig { Title = "General", Templates = new List<TemplateId> { _buildAndTestTemplateId, _buildTemplateId, _spawnedTestTemplateId } });
+			streamConfig.Templates.Add(_testBuildTestJobTemplate);
+			streamConfig.Templates.Add(_testBuildJobTemplate);
+			streamConfig.Templates.Add(_testSpawnedJobTemplate);
+			return streamConfig;
+		}
+
+		void UpdateGlobalConfig(ExperimentalConfig experimentalConfig)
+		{
+			UpdateConfig(globalConfig =>
+			{
+				globalConfig.Plugins.AddExperimentalConfig(experimentalConfig);
+			});
+		}
+
+		ExperimentalConfig CreateExperimentalNotificationConfig()
+		{
+			NotificationFormatConfig editorNotificationFormatConfig = CreateNotificationFormatConfig(EditorBootTestGroup, "^Editor BootTest\\s");
+			NotificationFormatConfig clientNotificationFormatConfig = CreateNotificationFormatConfig(ClientBootTestGroup, "^Client BootTest\\s");
+
+			JobNotificationConfig jobNotificationConfig = new JobNotificationConfig();
+			jobNotificationConfig.Enabled = true;
+			jobNotificationConfig.Template = _buildAndTestTemplateId;
+			jobNotificationConfig.NamePattern = "Platform1|Platform3";
+			jobNotificationConfig.NotificationFormats.Add(editorNotificationFormatConfig);
+			jobNotificationConfig.NotificationFormats.Add(clientNotificationFormatConfig);
+			jobNotificationConfig.Channels.Add(ChannelId);
+
+			NotificationStreamConfig notificationStreamConfig = new NotificationStreamConfig();
+			notificationStreamConfig.Streams.Add(_mainStreamId);
+			notificationStreamConfig.JobNotifications.Add(jobNotificationConfig);
+
+			NotificationConfig notificationConfig = new NotificationConfig();
+			notificationConfig.Id = new NotificationConfigId("test-notification");
+			notificationConfig.NotificationStreams.Add(notificationStreamConfig);
+
+			ExperimentalConfig experimentalConfig = new ExperimentalConfig();
+			experimentalConfig.Notifications.Add(notificationConfig);
+
+			return experimentalConfig;
+		}
+
+		static NotificationFormatConfig CreateNotificationFormatConfig(string group, string stepPattern, string? alternateNamePattern = null)
+		{
+			NotificationFormatConfig notificationFormatConfig = new NotificationFormatConfig();
+			notificationFormatConfig.Group = group;
+			notificationFormatConfig.StepPattern = stepPattern;
+			notificationFormatConfig.AlternativeNamePattern = alternateNamePattern;
+			return notificationFormatConfig;
+		}
+
+		static NewGroup AddGroup(List<NewGroup> groups, string agentType = "win64")
+		{
+			NewGroup group = new NewGroup(agentType, new List<NewNode>());
+			groups.Add(group);
+			return group;
+		}
+
+		static void AddNode(NewGroup group, string name, string[]? inputDependencies, Action<NewNode>? action = null)
+		{
+			NewNode node = new NewNode(name, inputDependencies: inputDependencies?.ToList(), orderDependencies: inputDependencies?.ToList());
+			action?.Invoke(node);
+			group.Nodes.Add(node);
+		}
+
+		static void AddPlatformTestSteps(NewGroup group, string platform, List<string> targetArguments, string[]? editorInputDependencies = null, string[]? clientInputDependencies = null)
+		{
+			AddNode(group, $"Editor BootTest {platform}", editorInputDependencies);
+			AddNode(group, $"Editor Project Tests {platform}", new[] { $"Editor BootTest {platform}"});
+			targetArguments.Add($"-Target=Editor Project Tests {platform}");
+
+			AddNode(group, $"Client BootTest {platform}", clientInputDependencies);
+			AddNode(group, $"Client Project Tests {platform}", new[] { $"Client BootTest {platform}"});
+			targetArguments.Add($"-Target=Client Project Tests {platform}");
+		}
+
+		async Task<IJob> CreateNightlyBuildTestJobAsync(StreamConfig config, int commitId)
+		{
+			Mock<ITemplate> templateMock = new Mock<ITemplate>(MockBehavior.Strict);
+			templateMock.SetupGet(x => x.InitialAgentType).Returns((string?)null);
+
+			IGraph baseGraph = await GraphCollection.AddAsync(templateMock.Object);
+
+			CreateJobOptions options = new CreateJobOptions();
+
+			List<NewGroup> newGroups = new List<NewGroup>();
+
+			NewGroup group = AddGroup(newGroups);
+			AddNode(group, "Update Version Files", null);
+			AddNode(group, "Compile Editor", new[] { "Update Version Files" });
+			AddNode(group, "Compile Client", new[] { "Update Version Files" });
+			AddNode(group, "Cook Client", new[] { "Compile Editor" });
+
+			string[] editorInputDependencies = new[] { "Compile Editor" };
+			string[] clientInputDependencies = new[] { "Compile Client", "Cook Client" };
+			AddPlatformTestSteps(group, "Platform1", options.Arguments, editorInputDependencies, clientInputDependencies);
+			AddPlatformTestSteps(group, "Platform2", options.Arguments, editorInputDependencies, clientInputDependencies);
+			AddPlatformTestSteps(group, "Platform3", options.Arguments, editorInputDependencies, clientInputDependencies);
+
+			IGraph graph = await GraphCollection.AppendAsync(baseGraph, newGroups);
+
+			ITemplate template = new TemplateStub(ContentHash.MD5("graphHash"), _buildAndTestTemplateId.ToString());
+			return await JobService.CreateJobAsync(null, config, _buildAndTestTemplateId, template.Hash, graph, "Test job", CommitIdWithOrder.FromPerforceChange(commitId), CommitIdWithOrder.FromPerforceChange(commitId), options);
+		}
+
+		async Task<IJob> CreateNightlyBuildJobAsync(StreamConfig config, int commitId)
+		{
+			Mock<ITemplate> templateMock = new Mock<ITemplate>(MockBehavior.Strict);
+			templateMock.SetupGet(x => x.InitialAgentType).Returns((string?)null);
+
+			IGraph baseGraph = await GraphCollection.AddAsync(templateMock.Object);
+
+			CreateJobOptions options = new CreateJobOptions();
+			options.Arguments.Add("-Target=Spawn Platform1 Tests");
+			options.Arguments.Add("-Target=Spawn Platform2 Tests");
+			options.Arguments.Add("-Target=Spawn Platform3 Tests");
+
+			List<NewGroup> newGroups = new List<NewGroup>();
+
+			NewGroup group = AddGroup(newGroups);
+			AddNode(group, "Spawn Platform1 Tests", null);
+			AddNode(group, "Spawn Platform2 Tests", null);
+			AddNode(group, "Spawn Platform3 Tests", null);
+
+			IGraph graph = await GraphCollection.AppendAsync(baseGraph, newGroups);
+
+			ITemplate template = new TemplateStub(ContentHash.MD5("graphHash"), _buildTemplateId.ToString());
+			return await JobService.CreateJobAsync(null, config, _buildTemplateId, template.Hash, graph, "Test Build job", CommitIdWithOrder.FromPerforceChange(commitId), CommitIdWithOrder.FromPerforceChange(commitId), options);
+		}
+
+		async Task<IJob> SpawnNightlyBuildTestJobAsync(StreamConfig config, JobId parentJobId, JobStepId parentJobStepId, string platform, int commitId)
+		{
+			Mock<ITemplate> templateMock = new Mock<ITemplate>(MockBehavior.Strict);
+			templateMock.SetupGet(x => x.InitialAgentType).Returns((string?)null);
+
+			IGraph baseGraph = await GraphCollection.AddAsync(templateMock.Object);
+
+			CreateJobOptions options = new CreateJobOptions();
+			options.ParentJobId = parentJobId;
+			options.ParentJobStepId = parentJobStepId;
+
+			List<NewGroup> newGroups = new List<NewGroup>();
+
+			NewGroup group = AddGroup(newGroups);
+			AddPlatformTestSteps(group, platform, options.Arguments);
+
+			IGraph graph = await GraphCollection.AppendAsync(baseGraph, newGroups);
+
+			ITemplate template = new TemplateStub(ContentHash.MD5("graphHash"), _spawnedTestTemplateId.ToString());
+			return await JobService.CreateJobAsync(null, config, _spawnedTestTemplateId, template.Hash, graph, $"Test job - {platform}", CommitIdWithOrder.FromPerforceChange(commitId), CommitIdWithOrder.FromPerforceChange(commitId), options);
+		}
+
+		async Task<IJob> AbortJobAsync(IJob job)
+		{
+			job = Deref(await JobService.UpdateJobAsync(job, abortedByUserId: UserId.Anonymous));
+			Assert.AreEqual(JobState.Complete, job.GetState());
+			return job;
+		}
+
+		Task<IJob> RunBatchAsync(StreamConfig config, IJob job, int batchIdx)
+		{
+			return UpdateBatchAsync(config, job, batchIdx, JobStepBatchState.Running);
+		}
+
+		async Task<IJob> UpdateBatchAsync(StreamConfig config, IJob job, int batchIdx, JobStepBatchState state)
+		{
+			Assert.AreEqual(JobStepBatchState.Ready, job.Batches[batchIdx].State);
+
+			LogId logId = LogIdUtils.GenerateNewId();
+			job = Deref(await job.TryAssignLeaseAsync(batchIdx, new PoolId("foo"), new AgentId("agent"), new SessionId(BinaryIdUtils.CreateNew()), new LeaseId(BinaryIdUtils.CreateNew()), logId));
+			if (job.Batches[batchIdx].State != state)
+			{
+				job = Deref(await JobService.UpdateBatchAsync(job, job.Batches[batchIdx].Id, config, logId, state));
+				Assert.AreEqual(state, job.Batches[batchIdx].State);
+			}
+			return job;
+		}
+
+		async Task<IJob> AbortStepAsync(StreamConfig config, IJob job, int batchIdx, int stepIdx)
+		{
+			LogId logId = LogIdUtils.GenerateNewId();
+			await LogCollection.AddAsync(job.Id, job.Batches[batchIdx].LeaseId, job.Batches[batchIdx].SessionId, LogType.Json, logId);
+			job = Deref(await JobService.UpdateStepAsync(job, job.Batches[batchIdx].Id, job.Batches[batchIdx].Steps[stepIdx].Id, config, newAbortRequested: true, newAbortByUserId: UserId.Anonymous, newLogId: logId));
+			Assert.AreEqual(JobStepState.Aborted, job.Batches[batchIdx].Steps[stepIdx].State);
+			return job;
+		}
+
+		async Task<IJob> AddSpawnedJobToParentStepAsync(StreamConfig config, IJob job, int batchIdx, int stepIdx, JobId spawnedJobId)
+		{
+			job = Deref(await JobService.UpdateStepAsync(job, job.Batches[batchIdx].Id, job.Batches[batchIdx].Steps[stepIdx].Id, config, newSpawnedJob: spawnedJobId));
+			Assert.IsNotNull(job.Batches[batchIdx].Steps[stepIdx].SpawnedJobs);
+			Assert.IsTrue(job.Batches[batchIdx].Steps[stepIdx].SpawnedJobs!.Contains(spawnedJobId));
+			return job;
+		}
+
+		async Task<IJob> RunStepAsync(StreamConfig config, IJob job, int batchIdx, int stepIdx, JobStepOutcome outcome)
+		{
+			Assert.AreEqual(JobStepState.Ready, job.Batches[batchIdx].Steps[stepIdx].State);
+
+			LogId logId = LogIdUtils.GenerateNewId();
+			await LogCollection.AddAsync(job.Id, job.Batches[batchIdx].LeaseId, job.Batches[batchIdx].SessionId, LogType.Json, logId);
+			job = Deref(await JobService.UpdateStepAsync(job, job.Batches[batchIdx].Id, job.Batches[batchIdx].Steps[stepIdx].Id, config, JobStepState.Running, JobStepOutcome.Success, newLogId: logId));
+			Assert.AreEqual(JobStepState.Running, job.Batches[batchIdx].Steps[stepIdx].State);
+			job = Deref(await JobService.UpdateStepAsync(job, job.Batches[batchIdx].Id, job.Batches[batchIdx].Steps[stepIdx].Id, config, JobStepState.Completed, outcome, newLogId: logId));
+			Assert.AreEqual(JobStepState.Completed, job.Batches[batchIdx].Steps[stepIdx].State);
+			Assert.AreEqual(outcome, job.Batches[batchIdx].Steps[stepIdx].Outcome);
+			return job;
+		}
+
+		static void ValidateJobStepInformation(IJobStepNotificationStateRef jobStepNotificationState, JobId jobId, TemplateId templateId, string recipient, string group, string platform, string badge, JobId? parentJobId = null, TemplateId? parentJobTemplateId = null)
+		{
+			Assert.AreEqual(jobId, jobStepNotificationState.JobId);
+			Assert.AreEqual(templateId, jobStepNotificationState.TemplateId);
+			Assert.AreEqual(recipient, jobStepNotificationState.Recipient);
+			Assert.AreEqual(group, jobStepNotificationState.Group);
+			Assert.AreEqual(platform, jobStepNotificationState.TargetPlatform);
+			Assert.AreEqual(badge, jobStepNotificationState.Badge);
+			Assert.AreEqual(parentJobId, jobStepNotificationState.ParentJobId);
+			Assert.AreEqual(parentJobTemplateId, jobStepNotificationState.ParentJobTemplateId);
+		}
+
+		static async Task ValidateJobStepInformation(JobNotificationCollection jobNotificationCollection, JobId jobId, TemplateId templateId, string recipient, JobStepId jobStepId, string group, string platform, string badge, JobId? parentJobId = null, TemplateId? parentJobTemplateId = null)
+		{
+			IJobStepNotificationStateQueryBuilder queryBuilder = JobNotificationCollection.CreateJobStepNotificationStateQueryBuilder();
+			queryBuilder.AddJobFilter(jobId);
+			queryBuilder.AddJobStepFilter(jobStepId);
+			queryBuilder.AddTemplateFilter(templateId);
+			queryBuilder.AddRecipientFilter(recipient);
+			IReadOnlyList<IJobStepNotificationStateRef>? jobStepNotificationStates = await jobNotificationCollection.GetJobStepNotificationStatesAsync(queryBuilder);
+			Assert.IsNotNull(jobStepNotificationStates);
+			Assert.AreEqual(1, jobStepNotificationStates.Count);
+
+			ValidateJobStepInformation(jobStepNotificationStates[0], jobId, templateId, recipient, group, platform, badge, parentJobId, parentJobTemplateId);
+		}
+
+		static async Task WaitForNotificationAsync(FakeExperimentalNotificationSink sink, IJob job, IJobStep? jobStep, int timeoutMs = 30000)
+		{
+			DateTime endTime = DateTime.Now + TimeSpan.FromMilliseconds(timeoutMs);
+
+			while (endTime >= DateTime.Now)
+			{
+				if (sink.ProcessedJobNotification is not null && sink.ProcessedJobNotification.Job == job && sink.ProcessedJobNotification.Step == jobStep)
+				{
+					return;
+				}
+				await Task.Delay(100);
+			}
+			Assert.Fail("Expected notification was not triggered within the time allotted.");
+		}
+
+		[TestCleanup]
+		public async Task TestCleanupAsync()
+		{
+			if (_testJobNotifications.Count == 0)
+			{
+				return;
+			}
+
+			// Cleanup all database records
+			JobNotificationCollection jobNotificationCollection = ServiceProvider.GetRequiredService<JobNotificationCollection>();
+
+			foreach (JobNotificationState jobNotificationState in _testJobNotifications)
+			{
+				await jobNotificationCollection.DeleteJobStepNotificationStatesAsync(jobNotificationState.JobId, jobNotificationState.TemplateId, jobNotificationState.Channel);
+				await jobNotificationCollection.DeleteJobNotificationStatesAsync(jobNotificationState.JobId, jobNotificationState.TemplateId, jobNotificationState.Channel);
+			}
+
+			_testJobNotifications.Clear();
+		}
+
+		[TestMethod]
+		public async Task SetupBuildJobStepFailedTestAsync()
+		{
+			// Scenario: Job matched with the configuration, but due to the Setup Build failing, nothing should be recorded as a Slack message would've been sent and the flow would be completed
+			// Expected: No records should exist in the JobNotificationCollection
+
+			ExperimentalConfig experimentalConfig = CreateExperimentalNotificationConfig();
+			UpdateGlobalConfig(experimentalConfig);
+
+			JobNotificationCollection jobNotificationCollection = ServiceProvider.GetRequiredService<JobNotificationCollection>();
+
+			FakeExperimentalNotificationSink fakeSink = ServiceProvider.GetRequiredService<FakeExperimentalNotificationSink>();
+
+			NotificationService service = (NotificationService)ServiceProvider.GetRequiredService<INotificationService>();
+			await service.StartAsync(cancellationToken: default);
+
+			IJob job = await CreateNightlyBuildTestJobAsync(_mainStreamConfig, 10);
+			_testJobNotifications.Add(new JobNotificationState(job.Id, job.TemplateId, ChannelId));
+
+			job = await RunBatchAsync(_mainStreamConfig, job, 0);
+			job = await RunStepAsync(_mainStreamConfig, job, 0, 0, JobStepOutcome.Failure); // Setup Build
+
+			await WaitForNotificationAsync(fakeSink, job, job.Batches[0].Steps[0]);
+			
+			// There should be nothing in the collection as a failed job has nothing to store
+			IJobNotificationStateQueryBuilder jobNotificationQueryBuilder = JobNotificationCollection.CreateJobNotificationStateQueryBuilder();
+			jobNotificationQueryBuilder.AddJobFilter(job.Id);
+			IReadOnlyList<IJobNotificationStateRef>? jobNotificationStates = await jobNotificationCollection.GetJobNotificationStatesAsync(jobNotificationQueryBuilder);
+			Assert.IsNull(jobNotificationStates);
+
+			IJobStepNotificationStateQueryBuilder jobStepNotificationQueryBuilder = JobNotificationCollection.CreateJobStepNotificationStateQueryBuilder();
+			jobStepNotificationQueryBuilder.AddJobFilter(job.Id);
+			IReadOnlyList<IJobStepNotificationStateRef>? jobStepNotificationStates = await jobNotificationCollection.GetJobStepNotificationStatesAsync(jobStepNotificationQueryBuilder);
+			Assert.IsNull(jobStepNotificationStates);
+		}
+
+		[TestMethod]
+		public async Task JobNotificationTestAsync()
+		{
+			// Scenario: Job matched with the configuration and recorded a job with its steps
+			// Expected: 1 job and 4 job step records with the job step records sorted in order by group, and platform
+
+			NotificationFormatConfig editorNotificationFormatConfig = CreateNotificationFormatConfig(EditorBootTestGroup, "^Editor BootTest\\s");
+			NotificationFormatConfig clientNotificationFormatConfig = CreateNotificationFormatConfig(ClientBootTestGroup, "^Client BootTest\\s", "Platform1");
+			NotificationFormatConfig editorProjectNotificationFormatConfig = CreateNotificationFormatConfig(EditorProjectTestGroup, "^Editor Project Tests\\s");
+			NotificationFormatConfig clientProjectNotificationFormatConfig = CreateNotificationFormatConfig(ClientProjectTestGroup, "^Client Project Tests\\s", "Platform1");
+
+			JobNotificationConfig jobNotificationConfig = new JobNotificationConfig();
+			jobNotificationConfig.Enabled = true;
+			jobNotificationConfig.Template = _buildAndTestTemplateId;
+			jobNotificationConfig.NamePattern = "Platform1|Platform2";
+			jobNotificationConfig.NotificationFormats.Add(editorNotificationFormatConfig);
+			jobNotificationConfig.NotificationFormats.Add(clientNotificationFormatConfig);
+			jobNotificationConfig.NotificationFormats.Add(editorProjectNotificationFormatConfig);
+			jobNotificationConfig.NotificationFormats.Add(clientProjectNotificationFormatConfig);
+			jobNotificationConfig.Channels.Add(ChannelId);
+
+			NotificationStreamConfig notificationStreamConfig = new NotificationStreamConfig();
+			notificationStreamConfig.Streams.Add(_mainStreamId);
+			notificationStreamConfig.JobNotifications.Add(jobNotificationConfig);
+
+			NotificationConfig notificationConfig = new NotificationConfig();
+			notificationConfig.Id = new NotificationConfigId("test-notification");
+			notificationConfig.NotificationStreams.Add(notificationStreamConfig);
+
+			ExperimentalConfig experimentalConfig = new ExperimentalConfig();
+			experimentalConfig.Notifications.Add(notificationConfig);
+
+			UpdateGlobalConfig(experimentalConfig);
+
+			JobNotificationCollection jobNotificationCollection = ServiceProvider.GetRequiredService<JobNotificationCollection>();
+
+			FakeExperimentalNotificationSink fakeSink = ServiceProvider.GetRequiredService<FakeExperimentalNotificationSink>();
+
+			NotificationService service = (NotificationService)ServiceProvider.GetRequiredService<INotificationService>();
+			await service.StartAsync(cancellationToken: default);
+
+			IJob job = await CreateNightlyBuildTestJobAsync(_mainStreamConfig, 20);
+			_testJobNotifications.Add(new JobNotificationState(job.Id, job.TemplateId, ChannelId));
+
+			job = await RunBatchAsync(_mainStreamConfig, job, 0);
+			job = await RunStepAsync(_mainStreamConfig, job, 0, 0, JobStepOutcome.Success); // Setup Build
+
+			await WaitForNotificationAsync(fakeSink, job, job.Batches[0].Steps[0]);
+			
+			// We should have a single job notification record in the database after the Setup Build job step has completed
+			IJobNotificationStateQueryBuilder jobNotificationQueryBuilder = JobNotificationCollection.CreateJobNotificationStateQueryBuilder();
+			jobNotificationQueryBuilder.AddJobFilter(job.Id);
+			jobNotificationQueryBuilder.AddTemplateFilter(job.TemplateId);
+			jobNotificationQueryBuilder.AddRecipientFilter(ChannelId);
+			IReadOnlyList<IJobNotificationStateRef>? jobNotificationStates = await jobNotificationCollection.GetJobNotificationStatesAsync(jobNotificationQueryBuilder);
+			Assert.IsNotNull(jobNotificationStates);
+			Assert.AreEqual(1, jobNotificationStates.Count);
+
+			// There should be 6 job step notification records sorted in order:
+			// 1 record for the 'Client Boot Test' group for Platform1
+			// 1 record for the 'Client Project Test' group for Platform1
+			// 2 records for the 'Editor Boot Test' group (Platform1 and Platform2)
+			// 2 records for the 'Editor Project Test' group (Platform1 and Platform2)
+			IJobStepNotificationStateQueryBuilder jobStepNotificationQueryBuilder = JobNotificationCollection.CreateJobStepNotificationStateQueryBuilder();
+			jobStepNotificationQueryBuilder.AddJobFilter(job.Id);
+			jobStepNotificationQueryBuilder.AddTemplateFilter(job.TemplateId);
+			jobStepNotificationQueryBuilder.AddRecipientFilter(ChannelId);
+			IReadOnlyList<IJobStepNotificationStateRef>? jobStepNotificationStates = await jobNotificationCollection.GetJobStepNotificationStatesAsync(jobStepNotificationQueryBuilder);
+			Assert.IsNotNull(jobStepNotificationStates);
+			Assert.AreEqual(6, jobStepNotificationStates.Count);
+
+			// Job Step notification for 'Client Boot Test Platform1'
+			ValidateJobStepInformation(jobStepNotificationStates[0], job.Id, job.TemplateId, ChannelId, ClientBootTestGroup, "Platform1", PendingBadge);
+
+			// Job Step notification for 'Client Project Test Platform1'
+			ValidateJobStepInformation(jobStepNotificationStates[1], job.Id, job.TemplateId, ChannelId, ClientProjectTestGroup, "Platform1", PendingBadge);
+
+			// Job Step notification for 'Editor Boot Test Platform1'
+			ValidateJobStepInformation(jobStepNotificationStates[2], job.Id, job.TemplateId, ChannelId, EditorBootTestGroup, "Platform1", PendingBadge);
+
+			// Job Step notification for 'Editor Project Test Platform2'
+			ValidateJobStepInformation(jobStepNotificationStates[3], job.Id, job.TemplateId, ChannelId, EditorBootTestGroup, "Platform2", PendingBadge);
+
+			// Job Step notification for 'Editor Boot Test Platform1'
+			ValidateJobStepInformation(jobStepNotificationStates[4], job.Id, job.TemplateId, ChannelId, EditorProjectTestGroup, "Platform1", PendingBadge);
+
+			// Job Step notification for 'Editor Project Test Platform2'
+			ValidateJobStepInformation(jobStepNotificationStates[5], job.Id, job.TemplateId, ChannelId, EditorProjectTestGroup, "Platform2", PendingBadge);
+
+			job = await RunBatchAsync(_mainStreamConfig, job, 1);
+			job = await RunStepAsync(_mainStreamConfig, job, 1, 0, JobStepOutcome.Success); // Update Version Files
+			job = await RunStepAsync(_mainStreamConfig, job, 1, 1, JobStepOutcome.Success); // Compile Editor
+			job = await RunStepAsync(_mainStreamConfig, job, 1, 2, JobStepOutcome.Success); // Compile Client
+			job = await RunStepAsync(_mainStreamConfig, job, 1, 3, JobStepOutcome.Success); // Cook Client
+			job = await RunStepAsync(_mainStreamConfig, job, 1, 4, JobStepOutcome.Success); // Editor BootTest Platform1
+
+			await WaitForNotificationAsync(fakeSink, job, job.Batches[1].Steps[4]);
+
+			// Validate that the 'Editor BootTest Platform1' step is recorded as passing
+			await ValidateJobStepInformation(jobNotificationCollection, job.Id, job.TemplateId,ChannelId, job.Batches[1].Steps[4].Id, EditorBootTestGroup, "Platform1", PassingBadge);
+
+			job = await RunStepAsync(_mainStreamConfig, job, 1, 5, JobStepOutcome.Warnings); // Editor Project Tests Platform1
+
+			await WaitForNotificationAsync(fakeSink, job, job.Batches[1].Steps[5]);
+
+			// Validate that the 'Editor Project Test Platform1' step is recorded as having warnings
+			await ValidateJobStepInformation(jobNotificationCollection, job.Id, job.TemplateId, ChannelId, job.Batches[1].Steps[5].Id, EditorProjectTestGroup, "Platform1", WarningBadge);
+
+			job = await RunStepAsync(_mainStreamConfig, job, 1, 6, JobStepOutcome.Failure); // Client BootTest Platform1
+
+			await WaitForNotificationAsync(fakeSink, job, job.Batches[1].Steps[6]);
+
+			// Validate that the 'Client BootTest Platform1' step is recorded as having failures and 'Client Project Tests Platform1' is recorded as being skipped due to having a dependency on the BootTest
+			await ValidateJobStepInformation(jobNotificationCollection, job.Id, job.TemplateId, ChannelId, job.Batches[1].Steps[6].Id, ClientBootTestGroup, "Platform1", FailureBadge);
+			await ValidateJobStepInformation(jobNotificationCollection, job.Id, job.TemplateId, ChannelId, job.Batches[1].Steps[7].Id, ClientProjectTestGroup, "Platform1", SkippedBadge);
+
+			job = await RunStepAsync(_mainStreamConfig, job, 1, 8, JobStepOutcome.Failure); // Editor BootTest Platform2
+
+			await WaitForNotificationAsync(fakeSink, job, job.Batches[1].Steps[8]);
+
+			// The 'Editor BootTest Platform2' step is recorded as having failures which will mark the 'Editor Project Tests Platform2' as being skipped due to being a dependency
+			// At this point, all notifications would be sent out and the notifications should be removed from the collections
+			jobStepNotificationQueryBuilder = JobNotificationCollection.CreateJobStepNotificationStateQueryBuilder();
+			jobStepNotificationQueryBuilder.AddJobFilter(job.Id);
+			jobStepNotificationStates = await jobNotificationCollection.GetJobStepNotificationStatesAsync(jobStepNotificationQueryBuilder);
+			Assert.IsNull(jobStepNotificationStates);
+
+			jobNotificationQueryBuilder = JobNotificationCollection.CreateJobNotificationStateQueryBuilder();
+			jobNotificationQueryBuilder.AddJobFilter(job.Id);
+			jobNotificationStates = await jobNotificationCollection.GetJobNotificationStatesAsync(jobNotificationQueryBuilder);
+			Assert.IsNull(jobNotificationStates);
+		}
+
+		[TestMethod]
+		public async Task AbortedJobTestAsync()
+		{
+			ExperimentalConfig experimentalConfig = CreateExperimentalNotificationConfig();
+			UpdateGlobalConfig(experimentalConfig);
+
+			JobNotificationCollection jobNotificationCollection = ServiceProvider.GetRequiredService<JobNotificationCollection>();
+
+			FakeExperimentalNotificationSink fakeSink = ServiceProvider.GetRequiredService<FakeExperimentalNotificationSink>();
+
+			NotificationService service = (NotificationService)ServiceProvider.GetRequiredService<INotificationService>();
+			await service.StartAsync(cancellationToken: default);
+
+			IJob job = await CreateNightlyBuildTestJobAsync(_mainStreamConfig, 30);
+			_testJobNotifications.Add(new JobNotificationState(job.Id, job.TemplateId, ChannelId));
+
+			job = await RunBatchAsync(_mainStreamConfig, job, 0);
+			job = await RunStepAsync(_mainStreamConfig, job, 0, 0, JobStepOutcome.Warnings); // Setup Build
+
+			await WaitForNotificationAsync(fakeSink, job, job.Batches[0].Steps[0]);
+			
+			// We should have a single job notification record in the database after the Setup Build job step has completed, even with warnings
+			IJobNotificationStateQueryBuilder jobNotificationQueryBuilder = JobNotificationCollection.CreateJobNotificationStateQueryBuilder();
+			jobNotificationQueryBuilder.AddJobFilter(job.Id);
+			jobNotificationQueryBuilder.AddTemplateFilter(job.TemplateId);
+			jobNotificationQueryBuilder.AddRecipientFilter(ChannelId);
+			IReadOnlyList<IJobNotificationStateRef>? jobNotificationStates = await jobNotificationCollection.GetJobNotificationStatesAsync(jobNotificationQueryBuilder);
+			Assert.IsNotNull(jobNotificationStates);
+			Assert.AreEqual(1, jobNotificationStates.Count);
+
+			// There should be 4 job step notification records: 2 records for the Client Boot Test group (Platform1 and Platform3) and 2 records for the Editor Boot Test group (Platform1 and Platform3) sorted in order
+			IJobStepNotificationStateQueryBuilder jobStepNotificationQueryBuilder = JobNotificationCollection.CreateJobStepNotificationStateQueryBuilder();
+			jobStepNotificationQueryBuilder.AddJobFilter(job.Id);
+			jobStepNotificationQueryBuilder.AddTemplateFilter(job.TemplateId);
+			jobStepNotificationQueryBuilder.AddRecipientFilter(ChannelId);
+			IReadOnlyList<IJobStepNotificationStateRef>? jobStepNotificationStates = await jobNotificationCollection.GetJobStepNotificationStatesAsync(jobStepNotificationQueryBuilder);
+			Assert.IsNotNull(jobStepNotificationStates);
+			Assert.AreEqual(4, jobStepNotificationStates.Count);
+
+			// Job Step notification for 'Client Boot Test Platform1'
+			ValidateJobStepInformation(jobStepNotificationStates[0], job.Id, job.TemplateId, ChannelId, ClientBootTestGroup, "Platform1", PendingBadge);
+
+			// Job Step notification for 'Client Boot Test Platform3'
+			ValidateJobStepInformation(jobStepNotificationStates[1], job.Id, job.TemplateId, ChannelId, ClientBootTestGroup, "Platform3", PendingBadge);
+
+			// Job Step notification for 'Editor Boot Test Platform1'
+			ValidateJobStepInformation(jobStepNotificationStates[2], job.Id, job.TemplateId, ChannelId, EditorBootTestGroup, "Platform1", PendingBadge);
+
+			// Job Step notification for 'Editor Boot Test Platform3'
+			ValidateJobStepInformation(jobStepNotificationStates[3], job.Id, job.TemplateId, ChannelId, EditorBootTestGroup, "Platform3", PendingBadge);
+
+			// #1
+			// Scenario: Cancellation of a step will report the step as failing
+			// Expected: Canceled step will be labeled as failing, but should not have any side-effects on the overall job
+			{
+				// Simulate the cancellation of the 'Editor BootTest Platform1' step
+				job = await AbortStepAsync(_mainStreamConfig, job, 1, 4);
+
+				await WaitForNotificationAsync(fakeSink, job, job.Batches[1].Steps[4]);
+
+				// Aside from the 'Editor BootTest Platform1' now being labeled as a failing step, there should still be 4 steps recorded
+				jobStepNotificationQueryBuilder = JobNotificationCollection.CreateJobStepNotificationStateQueryBuilder();
+				jobStepNotificationQueryBuilder.AddJobFilter(job.Id);
+				jobStepNotificationQueryBuilder.AddTemplateFilter(job.TemplateId);
+				jobStepNotificationQueryBuilder.AddRecipientFilter(ChannelId);
+				jobStepNotificationStates = await jobNotificationCollection.GetJobStepNotificationStatesAsync(jobStepNotificationQueryBuilder);
+				Assert.IsNotNull(jobStepNotificationStates);
+				Assert.AreEqual(4, jobStepNotificationStates.Count);
+
+				// Job Step notification for the 'Client Boot Test Platform1' group
+				ValidateJobStepInformation(jobStepNotificationStates[0], job.Id, job.TemplateId, ChannelId, ClientBootTestGroup, "Platform1", PendingBadge);
+
+				// Job Step notification for the 'Client Boot Test Platform3' group
+				ValidateJobStepInformation(jobStepNotificationStates[1], job.Id, job.TemplateId, ChannelId, ClientBootTestGroup, "Platform3", PendingBadge);
+
+				// Job Step notification for the 'Editor Boot Test Platform1' group
+				ValidateJobStepInformation(jobStepNotificationStates[2], job.Id, job.TemplateId, ChannelId, EditorBootTestGroup, "Platform1", FailureBadge);
+
+				// Job Step notification for the 'Editor Boot Test Platform3' group
+				ValidateJobStepInformation(jobStepNotificationStates[3], job.Id, job.TemplateId, ChannelId, EditorBootTestGroup, "Platform3", PendingBadge);
+			}
+
+			// #2
+			// Scenario: Cancellation of a dependency step will report the steps as being skipped
+			// Expected: Canceled step will be labeled as failing and the steps which had a dependency on the canceled step will be skipped
+			{
+				// Simulate the cancellation of the 'Compile Client' step
+				job = await AbortStepAsync(_mainStreamConfig, job, 1, 2);
+
+				await WaitForNotificationAsync(fakeSink, job, job.Batches[1].Steps[2]);
+
+				// Aside from the Client BootTest step for Platform1 and Platform3 now being labeled as skipped steps, there should still be 4 steps recorded
+				jobStepNotificationQueryBuilder = JobNotificationCollection.CreateJobStepNotificationStateQueryBuilder();
+				jobStepNotificationQueryBuilder.AddJobFilter(job.Id);
+				jobStepNotificationQueryBuilder.AddTemplateFilter(job.TemplateId);
+				jobStepNotificationQueryBuilder.AddRecipientFilter(ChannelId);
+				jobStepNotificationStates = await jobNotificationCollection.GetJobStepNotificationStatesAsync(jobStepNotificationQueryBuilder);
+				Assert.IsNotNull(jobStepNotificationStates);
+				Assert.AreEqual(4, jobStepNotificationStates.Count);
+
+				// Job Step notification for the 'Client Boot Test Platform1' group
+				ValidateJobStepInformation(jobStepNotificationStates[0], job.Id, job.TemplateId, ChannelId, ClientBootTestGroup, "Platform1", SkippedBadge);
+
+				// Job Step notification for the 'Client Boot Test Platform3' group
+				ValidateJobStepInformation(jobStepNotificationStates[1], job.Id, job.TemplateId, ChannelId, ClientBootTestGroup, "Platform3", SkippedBadge);
+
+				// Job Step notification for the 'Editor Boot Test Platform1' group
+				ValidateJobStepInformation(jobStepNotificationStates[2], job.Id, job.TemplateId, ChannelId, EditorBootTestGroup, "Platform1", FailureBadge);
+
+				// Job Step notification for the 'Editor Boot Test Platform3' group
+				ValidateJobStepInformation(jobStepNotificationStates[3], job.Id, job.TemplateId, ChannelId, EditorBootTestGroup, "Platform3", PendingBadge);
+			}
+
+			// #3
+			// Scenario: Cancellation of a job will cleanup all notifications
+			// Expected: All notification records will be cleaned up
+			{
+				// Simulate the cancellation of a job
+				job = await AbortJobAsync(job);
+
+				await WaitForNotificationAsync(fakeSink, job, null);
+
+				// The job should be removed from the collection
+				jobNotificationQueryBuilder = JobNotificationCollection.CreateJobNotificationStateQueryBuilder();
+				jobNotificationQueryBuilder.AddJobFilter(job.Id);
+				jobNotificationQueryBuilder.AddTemplateFilter(job.TemplateId);
+				jobNotificationQueryBuilder.AddRecipientFilter(ChannelId);
+				jobNotificationStates = await jobNotificationCollection.GetJobNotificationStatesAsync(jobNotificationQueryBuilder);
+				Assert.IsNull(jobNotificationStates);
+
+				// The job steps should be removed from the collection
+				jobStepNotificationQueryBuilder = JobNotificationCollection.CreateJobStepNotificationStateQueryBuilder();
+				jobStepNotificationQueryBuilder.AddJobFilter(job.Id);
+				jobStepNotificationStates = await jobNotificationCollection.GetJobStepNotificationStatesAsync(jobStepNotificationQueryBuilder);
+				Assert.IsNull(jobStepNotificationStates);
+			}
+		}
+
+		[TestMethod]
+		public async Task ConsolidatedNotificationTestAsync()
+		{
+			int commitId = 40;
+			NotificationFormatConfig spawnNotificationFormatConfig = CreateNotificationFormatConfig(SpawnJobTestGroup, "^Spawn\\s");
+			NotificationFormatConfig editorNotificationFormatConfig = CreateNotificationFormatConfig(EditorBootTestGroup, "^Editor BootTest\\s");
+			NotificationFormatConfig clientNotificationFormatConfig = CreateNotificationFormatConfig(ClientBootTestGroup, "^Client BootTest\\s");
+			NotificationFormatConfig editorProjectNotificationFormatConfig = CreateNotificationFormatConfig(EditorProjectTestGroup, "^Editor Project Tests\\s");
+			NotificationFormatConfig clientProjectNotificationFormatConfig = CreateNotificationFormatConfig(ClientProjectTestGroup, "^Client Project Tests\\s");
+
+			JobNotificationConfig jobNotificationConfig = new JobNotificationConfig();
+			jobNotificationConfig.Enabled = true;
+			jobNotificationConfig.Template = _buildTemplateId;
+			jobNotificationConfig.NamePattern = "Platform1|Platform2|Platform3";
+			jobNotificationConfig.NotificationFormats.Add(spawnNotificationFormatConfig);
+			jobNotificationConfig.NotificationFormats.Add(editorNotificationFormatConfig);
+			jobNotificationConfig.NotificationFormats.Add(clientNotificationFormatConfig);
+			jobNotificationConfig.NotificationFormats.Add(editorProjectNotificationFormatConfig);
+			jobNotificationConfig.NotificationFormats.Add(clientProjectNotificationFormatConfig);
+			jobNotificationConfig.Channels.Add(ChannelId);
+
+			NotificationStreamConfig notificationStreamConfig = new NotificationStreamConfig();
+			notificationStreamConfig.Streams.Add(_mainStreamId);
+			notificationStreamConfig.JobNotifications.Add(jobNotificationConfig);
+
+			NotificationConfig notificationConfig = new NotificationConfig();
+			notificationConfig.Id = new NotificationConfigId("test-notification");
+			notificationConfig.NotificationStreams.Add(notificationStreamConfig);
+
+			ExperimentalConfig experimentalConfig = new ExperimentalConfig();
+			experimentalConfig.Notifications.Add(notificationConfig);
+
+			UpdateGlobalConfig(experimentalConfig);
+
+			JobNotificationCollection jobNotificationCollection = ServiceProvider.GetRequiredService<JobNotificationCollection>();
+
+			FakeExperimentalNotificationSink fakeSink = ServiceProvider.GetRequiredService<FakeExperimentalNotificationSink>();
+
+			NotificationService service = (NotificationService)ServiceProvider.GetRequiredService<INotificationService>();
+			await service.StartAsync(cancellationToken: default);
+
+			IJob parentJob = await CreateNightlyBuildJobAsync(_mainStreamConfig, commitId);
+			_testJobNotifications.Add(new JobNotificationState(parentJob.Id, parentJob.TemplateId, ChannelId));
+
+			parentJob = await RunBatchAsync(_mainStreamConfig, parentJob, 0);
+			parentJob = await RunStepAsync(_mainStreamConfig, parentJob, 0, 0, JobStepOutcome.Success); // Setup Build
+
+			await WaitForNotificationAsync(fakeSink, parentJob, parentJob.Batches[0].Steps[0]);
+
+			// We should have a single job notification record in the database after the Setup Build job step has completed, even with warnings
+			IJobNotificationStateQueryBuilder jobNotificationQueryBuilder = JobNotificationCollection.CreateJobNotificationStateQueryBuilder();
+			jobNotificationQueryBuilder.AddJobFilter(parentJob.Id);
+			jobNotificationQueryBuilder.AddTemplateFilter(parentJob.TemplateId);
+			jobNotificationQueryBuilder.AddRecipientFilter(ChannelId);
+			IReadOnlyList<IJobNotificationStateRef>? jobNotificationStates = await jobNotificationCollection.GetJobNotificationStatesAsync(jobNotificationQueryBuilder);
+			Assert.IsNotNull(jobNotificationStates);
+			Assert.AreEqual(1, jobNotificationStates.Count);
+
+			// There should be 3 job step notification records for the Spawn Test Job group for each platform (Platform1, Platform2, and Platform3)
+			IJobStepNotificationStateQueryBuilder jobStepNotificationQueryBuilder = JobNotificationCollection.CreateJobStepNotificationStateQueryBuilder();
+			jobStepNotificationQueryBuilder.AddJobFilter(parentJob.Id);
+			jobStepNotificationQueryBuilder.AddTemplateFilter(parentJob.TemplateId);
+			jobStepNotificationQueryBuilder.AddRecipientFilter(ChannelId);
+			IReadOnlyList<IJobStepNotificationStateRef>? jobStepNotificationStates = await jobNotificationCollection.GetJobStepNotificationStatesAsync(jobStepNotificationQueryBuilder);
+			Assert.IsNotNull(jobStepNotificationStates);
+			Assert.AreEqual(3, jobStepNotificationStates.Count);
+
+			// Job Step notification for 'Spawn Platform1 Test'
+			ValidateJobStepInformation(jobStepNotificationStates[0], parentJob.Id, parentJob.TemplateId, ChannelId, SpawnJobTestGroup, "Platform1", PendingBadge);
+
+			// Job Step notification for 'Spawn Platform2 Test'
+			ValidateJobStepInformation(jobStepNotificationStates[1], parentJob.Id, parentJob.TemplateId, ChannelId, SpawnJobTestGroup, "Platform2", PendingBadge);
+
+			// Job Step notification for 'Spawn Platform3 Test'
+			ValidateJobStepInformation(jobStepNotificationStates[2], parentJob.Id, parentJob.TemplateId, ChannelId, SpawnJobTestGroup, "Platform3", PendingBadge);
+
+			// #1
+			// Scenario: Spawned job has a failing Setup Build step
+			// Expected: The spawned job will not append any notification information
+			{
+				parentJob = await RunBatchAsync(_mainStreamConfig, parentJob, 1);
+
+				// The spawned job will be created before the job step completes
+				IJob spawnedJob = await SpawnNightlyBuildTestJobAsync(_mainStreamConfig, parentJob.Id, parentJob.Batches[1].Steps[0].Id, "Platform1", commitId);
+				parentJob = await AddSpawnedJobToParentStepAsync(_mainStreamConfig, parentJob, 1, 0, spawnedJob.Id);
+				_testJobNotifications.Add(new JobNotificationState(spawnedJob.Id, spawnedJob.TemplateId, ChannelId));
+
+				parentJob = await RunStepAsync(_mainStreamConfig, parentJob, 1, 0, JobStepOutcome.Success); // Spawn Platform1 Tests
+
+				await WaitForNotificationAsync(fakeSink, parentJob, parentJob.Batches[1].Steps[0]);
+
+				await ValidateJobStepInformation(jobNotificationCollection, parentJob.Id, parentJob.TemplateId, ChannelId, parentJob.Batches[1].Steps[0].Id, SpawnJobTestGroup, "Platform1", PassingBadge);
+
+				// Execute the Setup Build of the newly spawned job
+				spawnedJob = await RunBatchAsync(_mainStreamConfig, spawnedJob, 0);
+				spawnedJob = await RunStepAsync(_mainStreamConfig, spawnedJob, 0, 0, JobStepOutcome.Failure); // Setup Build
+
+				await WaitForNotificationAsync(fakeSink, spawnedJob, spawnedJob.Batches[0].Steps[0]);
+
+				// There should be nothing in the job collection for the spawned job as the job has both failed and should've been recognized as a spawned job
+				IJobNotificationStateQueryBuilder spawnedJobNotificationQueryBuilder = JobNotificationCollection.CreateJobNotificationStateQueryBuilder();
+				spawnedJobNotificationQueryBuilder.AddJobFilter(spawnedJob.Id);
+				IReadOnlyList<IJobNotificationStateRef>? spawnedJobNotificationStates = await jobNotificationCollection.GetJobNotificationStatesAsync(spawnedJobNotificationQueryBuilder);
+				Assert.IsNull(spawnedJobNotificationStates);
+
+				IJobStepNotificationStateQueryBuilder spawnedJobStepNotificationQueryBuilder = JobNotificationCollection.CreateJobStepNotificationStateQueryBuilder();
+				spawnedJobStepNotificationQueryBuilder.AddJobFilter(spawnedJob.Id);
+				IReadOnlyList<IJobStepNotificationStateRef>? spawnedJobStepNotificationStates = await jobNotificationCollection.GetJobStepNotificationStatesAsync(spawnedJobStepNotificationQueryBuilder);
+				Assert.IsNull(spawnedJobStepNotificationStates);
+
+				// The parent job should only have the notifications regarding the spawn test steps as the failing spawned job did not append its steps to the parent job
+				jobNotificationQueryBuilder = JobNotificationCollection.CreateJobNotificationStateQueryBuilder();
+				jobNotificationQueryBuilder.AddJobFilter(parentJob.Id);
+				jobNotificationQueryBuilder.AddTemplateFilter(parentJob.TemplateId);
+				jobNotificationQueryBuilder.AddRecipientFilter(ChannelId);
+				jobNotificationStates = await jobNotificationCollection.GetJobNotificationStatesAsync(jobNotificationQueryBuilder);
+				Assert.IsNotNull(jobNotificationStates);
+				Assert.AreEqual(1, jobNotificationStates.Count);
+
+				// There should be 3 job step notification records for the Spawn Test Job group for each platform (Platform1, Platform2, and Platform3)
+				jobStepNotificationQueryBuilder = JobNotificationCollection.CreateJobStepNotificationStateQueryBuilder();
+				jobStepNotificationQueryBuilder.AddJobFilter(parentJob.Id);
+				jobStepNotificationQueryBuilder.AddTemplateFilter(parentJob.TemplateId);
+				jobStepNotificationQueryBuilder.AddRecipientFilter(ChannelId);
+				jobStepNotificationStates = await jobNotificationCollection.GetJobStepNotificationStatesAsync(jobStepNotificationQueryBuilder);
+				Assert.IsNotNull(jobStepNotificationStates);
+				Assert.AreEqual(3, jobStepNotificationStates.Count);
+
+				// Job Step notification for the 'Spawn Platform1 Test' will now being labeled as passing since the step spawned the new job
+				// Even though the actual job had issues which will be reported in the Slack thread
+				ValidateJobStepInformation(jobStepNotificationStates[0], parentJob.Id, parentJob.TemplateId, ChannelId, SpawnJobTestGroup, "Platform1", PassingBadge);
+
+				// Job Step notification for the Spawn Platform2 Test group
+				ValidateJobStepInformation(jobStepNotificationStates[1], parentJob.Id, parentJob.TemplateId, ChannelId, SpawnJobTestGroup, "Platform2", PendingBadge);
+
+				// Job Step notification for the Spawn Platform3 Test group
+				ValidateJobStepInformation(jobStepNotificationStates[2], parentJob.Id, parentJob.TemplateId, ChannelId, SpawnJobTestGroup, "Platform3", PendingBadge);
+			}
+
+			// #2
+			// Scenario: Spawned job has a successfully completed Setup Build, but later has a step and job manually aborted
+			// Expected: Spawned job will append its information to the parent job. The later canceled step will be labeled as a failure and the canceled job will then have the information regarding the spawned job removed
+			{
+				// The spawned job will be created before the job step completes
+				IJob spawnedJob = await SpawnNightlyBuildTestJobAsync(_mainStreamConfig, parentJob.Id, parentJob.Batches[1].Steps[1].Id, "Platform2", commitId);
+				parentJob = await AddSpawnedJobToParentStepAsync(_mainStreamConfig, parentJob, 1, 0, spawnedJob.Id);
+				_testJobNotifications.Add(new JobNotificationState(spawnedJob.Id, spawnedJob.TemplateId, ChannelId));
+
+				parentJob = await RunStepAsync(_mainStreamConfig, parentJob, 1, 1, JobStepOutcome.Success); // Spawn Platform2 Tests
+
+				await WaitForNotificationAsync(fakeSink, parentJob, parentJob.Batches[1].Steps[1]);
+
+				await ValidateJobStepInformation(jobNotificationCollection, parentJob.Id, parentJob.TemplateId, ChannelId, parentJob.Batches[1].Steps[1].Id, SpawnJobTestGroup, "Platform2", PassingBadge);
+
+				// Execute the Setup Build of the newly spawned job
+				spawnedJob = await RunBatchAsync(_mainStreamConfig, spawnedJob, 0);
+				spawnedJob = await RunStepAsync(_mainStreamConfig, spawnedJob, 0, 0, JobStepOutcome.Success); // Setup Build
+
+				await WaitForNotificationAsync(fakeSink, spawnedJob, spawnedJob.Batches[0].Steps[0]);
+
+				// There should not be a job notification state for the spawned job as it is linked to the parent job
+				IJobNotificationStateQueryBuilder spawnedJobNotificationQueryBuilder = JobNotificationCollection.CreateJobNotificationStateQueryBuilder();
+				spawnedJobNotificationQueryBuilder.AddJobFilter(spawnedJob.Id);
+				IReadOnlyList<IJobNotificationStateRef>? spawnedJobNotificationStates = await jobNotificationCollection.GetJobNotificationStatesAsync(spawnedJobNotificationQueryBuilder);
+				Assert.IsNull(spawnedJobNotificationStates);
+
+				// The spawned job will append its notification records to the parent job and there will be a total of have 7 job step notification records with the parent job step notifications appearing at the beginning of the list:
+				// Spawn Platform1 Test, Spawn Platform2 Test, Spawn Platform3 Test, Client BootTest Platform2, Client Project Tests Platform2, Editor BootTest Platform2, and Editor Project Tests Platform2
+				IJobStepNotificationStateQueryBuilder consolidatedJobStepNotificationQueryBuilder = JobNotificationCollection.CreateJobStepNotificationStateQueryBuilder();
+				consolidatedJobStepNotificationQueryBuilder.AddJobAndParentJobFilter(parentJob.Id);
+				consolidatedJobStepNotificationQueryBuilder.AddTemplateAndParentTemplateFilter(parentJob.TemplateId);
+				consolidatedJobStepNotificationQueryBuilder.AddRecipientFilter(ChannelId);
+				IReadOnlyList<IJobStepNotificationStateRef>? consolidatedJobStepNotificationStates = await jobNotificationCollection.GetJobStepNotificationStatesAsync(consolidatedJobStepNotificationQueryBuilder);
+				Assert.IsNotNull(consolidatedJobStepNotificationStates);
+				Assert.AreEqual(7, consolidatedJobStepNotificationStates.Count);
+
+				// Job Step notification for 'Spawn Platform1 Test'
+				ValidateJobStepInformation(consolidatedJobStepNotificationStates[0], parentJob.Id, parentJob.TemplateId, ChannelId, SpawnJobTestGroup, "Platform1", PassingBadge);
+
+				// Job Step notification for 'Spawn Platform2 Test' group'
+				ValidateJobStepInformation(consolidatedJobStepNotificationStates[1], parentJob.Id, parentJob.TemplateId, ChannelId, SpawnJobTestGroup, "Platform2", PassingBadge);
+
+				// Job Step notification for 'Spawn Platform3 Test'
+				ValidateJobStepInformation(consolidatedJobStepNotificationStates[2], parentJob.Id, parentJob.TemplateId, ChannelId, SpawnJobTestGroup, "Platform3", PendingBadge);
+
+				// Job Step notification for 'Client BootTest Platform2'
+				ValidateJobStepInformation(consolidatedJobStepNotificationStates[3], spawnedJob.Id, spawnedJob.TemplateId, ChannelId, ClientBootTestGroup, "Platform2", PendingBadge, parentJob.Id, parentJob.TemplateId);
+
+				// Job Step notification for 'Client Project Tests Platform2'
+				ValidateJobStepInformation(consolidatedJobStepNotificationStates[4], spawnedJob.Id, spawnedJob.TemplateId, ChannelId, ClientProjectTestGroup, "Platform2", PendingBadge, parentJob.Id, parentJob.TemplateId);
+
+				// Job Step notification for 'Editor BootTest Platform2'
+				ValidateJobStepInformation(consolidatedJobStepNotificationStates[5], spawnedJob.Id, spawnedJob.TemplateId, ChannelId, EditorBootTestGroup, "Platform2", PendingBadge, parentJob.Id, parentJob.TemplateId);
+
+				// Job Step notification for 'Editor Project Tests Platform2'
+				ValidateJobStepInformation(consolidatedJobStepNotificationStates[6], spawnedJob.Id, spawnedJob.TemplateId, ChannelId, EditorProjectTestGroup, "Platform2", PendingBadge, parentJob.Id, parentJob.TemplateId);
+
+				// Simulate the cancellation of the 'Editor BootTest Platform1' step
+				spawnedJob = await AbortStepAsync(_mainStreamConfig, spawnedJob, 1, 0);
+
+				await WaitForNotificationAsync(fakeSink, spawnedJob, spawnedJob.Batches[1].Steps[0]);
+
+				// Aside from 'Editor BootTest Platform2' now being labeled as a failing step and its follow-up step 'Editor Project Tests Platform2' being labeled as skipped, there should still be 7 steps recorded
+				consolidatedJobStepNotificationQueryBuilder = JobNotificationCollection.CreateJobStepNotificationStateQueryBuilder();
+				consolidatedJobStepNotificationQueryBuilder.AddJobAndParentJobFilter(parentJob.Id);
+				consolidatedJobStepNotificationQueryBuilder.AddTemplateAndParentTemplateFilter(parentJob.TemplateId);
+				consolidatedJobStepNotificationQueryBuilder.AddRecipientFilter(ChannelId);
+				consolidatedJobStepNotificationStates = await jobNotificationCollection.GetJobStepNotificationStatesAsync(consolidatedJobStepNotificationQueryBuilder);
+				Assert.IsNotNull(consolidatedJobStepNotificationStates);
+				Assert.AreEqual(7, consolidatedJobStepNotificationStates.Count);
+
+				// Job Step notification for 'Spawn Platform1 Test'
+				ValidateJobStepInformation(consolidatedJobStepNotificationStates[0], parentJob.Id, parentJob.TemplateId, ChannelId, SpawnJobTestGroup, "Platform1", PassingBadge);
+
+				// Job Step notification for 'Spawn Platform2 Test'
+				ValidateJobStepInformation(consolidatedJobStepNotificationStates[1], parentJob.Id, parentJob.TemplateId, ChannelId, SpawnJobTestGroup, "Platform2", PassingBadge);
+
+				// Job Step notification for 'Spawn Platform3 Test'
+				ValidateJobStepInformation(consolidatedJobStepNotificationStates[2], parentJob.Id, parentJob.TemplateId, ChannelId, SpawnJobTestGroup, "Platform3", PendingBadge);
+
+				// Job Step notification for 'Client BootTest Platform2'
+				ValidateJobStepInformation(consolidatedJobStepNotificationStates[3], spawnedJob.Id, spawnedJob.TemplateId, ChannelId, ClientBootTestGroup, "Platform2", PendingBadge, parentJob.Id, parentJob.TemplateId);
+
+				// Job Step notification for 'Client Project Tests Platform2'
+				ValidateJobStepInformation(consolidatedJobStepNotificationStates[4], spawnedJob.Id, spawnedJob.TemplateId, ChannelId, ClientProjectTestGroup, "Platform2", PendingBadge, parentJob.Id, parentJob.TemplateId);
+
+				// Job Step notification for 'Editor BootTest Platform2'
+				ValidateJobStepInformation(consolidatedJobStepNotificationStates[5], spawnedJob.Id, spawnedJob.TemplateId, ChannelId, EditorBootTestGroup, "Platform2", FailureBadge, parentJob.Id, parentJob.TemplateId);
+
+				// Job Step notification for 'Editor Project Tests Platform2'
+				ValidateJobStepInformation(consolidatedJobStepNotificationStates[6], spawnedJob.Id, spawnedJob.TemplateId, ChannelId, EditorProjectTestGroup, "Platform2", SkippedBadge, parentJob.Id, parentJob.TemplateId);
+
+				// Simulate the cancellation of a job
+				spawnedJob = await AbortJobAsync(spawnedJob);
+
+				await WaitForNotificationAsync(fakeSink, spawnedJob, null);
+
+				// The parent job notification should still be found after a spawned job has been canceled
+				jobNotificationQueryBuilder = JobNotificationCollection.CreateJobNotificationStateQueryBuilder();
+				jobNotificationQueryBuilder.AddJobFilter(parentJob.Id);
+				jobNotificationQueryBuilder.AddTemplateFilter(parentJob.TemplateId);
+				jobNotificationQueryBuilder.AddRecipientFilter(ChannelId);
+				jobNotificationStates = await jobNotificationCollection.GetJobNotificationStatesAsync(jobNotificationQueryBuilder);
+				Assert.IsNotNull(jobNotificationStates);
+				Assert.AreEqual(1, jobNotificationStates.Count);
+
+				// The spawned job will have its job step notifications removed as the Slack summary will provide the link to the canceled job
+				// There should be 3 job step notification records for the Spawn Test Job group for each platform (Platform1, Platform2, and Platform3)
+				consolidatedJobStepNotificationQueryBuilder = JobNotificationCollection.CreateJobStepNotificationStateQueryBuilder();
+				consolidatedJobStepNotificationQueryBuilder.AddJobAndParentJobFilter(parentJob.Id);
+				consolidatedJobStepNotificationQueryBuilder.AddTemplateAndParentTemplateFilter(parentJob.TemplateId);
+				consolidatedJobStepNotificationQueryBuilder.AddRecipientFilter(ChannelId);
+				consolidatedJobStepNotificationStates = await jobNotificationCollection.GetJobStepNotificationStatesAsync(consolidatedJobStepNotificationQueryBuilder);
+				Assert.IsNotNull(consolidatedJobStepNotificationStates);
+				Assert.AreEqual(3, consolidatedJobStepNotificationStates.Count);
+
+				// Job Step notification for the 'Spawn Platform2 Test' will now being labeled as passing since the step spawned the new job
+				// Even though the actual job had issues and was canceled which will be reported in the Slack thread
+				ValidateJobStepInformation(consolidatedJobStepNotificationStates[0], parentJob.Id, parentJob.TemplateId, ChannelId, SpawnJobTestGroup, "Platform1", PassingBadge);
+
+				// Job Step notification for 'Spawn Platform2 Test'
+				ValidateJobStepInformation(consolidatedJobStepNotificationStates[1], parentJob.Id, parentJob.TemplateId, ChannelId, SpawnJobTestGroup, "Platform2", PassingBadge);
+
+				// Job Step notification for 'Spawn Platform3 Test'
+				ValidateJobStepInformation(consolidatedJobStepNotificationStates[2], parentJob.Id, parentJob.TemplateId, ChannelId, SpawnJobTestGroup, "Platform3", PendingBadge);
+			}
+
+			// #3
+			// Scenario: Spawned job has a completed Setup Build, albeit with warnings, and is updating with various results
+			// Expected: Spawned job will keep updating the job step notification states until the job has run to completion. At which point all jobs have completed and the notification records will be cleaned
+			{
+				// The spawned job will be created before the job step completes
+				IJob spawnedJob = await SpawnNightlyBuildTestJobAsync(_mainStreamConfig, parentJob.Id, parentJob.Batches[1].Steps[2].Id, "Platform3", commitId);
+				parentJob = await AddSpawnedJobToParentStepAsync(_mainStreamConfig, parentJob, 1, 0, spawnedJob.Id);
+				_testJobNotifications.Add(new JobNotificationState(spawnedJob.Id, spawnedJob.TemplateId, ChannelId));
+
+				parentJob = await RunStepAsync(_mainStreamConfig, parentJob, 1, 2, JobStepOutcome.Success); // Spawn Platform3 Tests
+
+				await WaitForNotificationAsync(fakeSink, parentJob, parentJob.Batches[1].Steps[2]);
+
+				await ValidateJobStepInformation(jobNotificationCollection, parentJob.Id, parentJob.TemplateId, ChannelId, parentJob.Batches[1].Steps[2].Id, SpawnJobTestGroup, "Platform3", PassingBadge);
+
+				// Execute the Setup Build of the newly spawned job
+				spawnedJob = await RunBatchAsync(_mainStreamConfig, spawnedJob, 0);
+				spawnedJob = await RunStepAsync(_mainStreamConfig, spawnedJob, 0, 0, JobStepOutcome.Warnings); // Setup Build
+
+				await WaitForNotificationAsync(fakeSink, spawnedJob, spawnedJob.Batches[0].Steps[0]);
+
+				// There should not be a job notification state for the spawned job as it is linked to the parent job
+				IJobNotificationStateQueryBuilder spawnedJobNotificationQueryBuilder = JobNotificationCollection.CreateJobNotificationStateQueryBuilder();
+				spawnedJobNotificationQueryBuilder.AddJobFilter(spawnedJob.Id);
+				IReadOnlyList<IJobNotificationStateRef>? spawnedJobNotificationStates = await jobNotificationCollection.GetJobNotificationStatesAsync(spawnedJobNotificationQueryBuilder);
+				Assert.IsNull(spawnedJobNotificationStates);
+
+				// The spawned job will append its notification records to the parent job and there will be a total of have 7 job step notification records with the parent job step notifications appearing at the beginning of the list:
+				// Spawn Platform1 Test, Spawn Platform2 Test, Spawn Platform3 Test, Client BootTest Platform3, Client Project Tests Platform3, Editor BootTest Platform3, and Editor Project Tests Platform3
+				IJobStepNotificationStateQueryBuilder consolidatedJobStepNotificationQueryBuilder = JobNotificationCollection.CreateJobStepNotificationStateQueryBuilder();
+				consolidatedJobStepNotificationQueryBuilder.AddJobAndParentJobFilter(parentJob.Id);
+				consolidatedJobStepNotificationQueryBuilder.AddTemplateAndParentTemplateFilter(parentJob.TemplateId);
+				consolidatedJobStepNotificationQueryBuilder.AddRecipientFilter(ChannelId);
+				IReadOnlyList<IJobStepNotificationStateRef>? consolidatedJobStepNotificationStates = await jobNotificationCollection.GetJobStepNotificationStatesAsync(consolidatedJobStepNotificationQueryBuilder);
+				Assert.IsNotNull(consolidatedJobStepNotificationStates);
+				Assert.AreEqual(7, consolidatedJobStepNotificationStates.Count);
+
+				// Job Step notification for 'Spawn Platform1 Test'
+				ValidateJobStepInformation(consolidatedJobStepNotificationStates[0], parentJob.Id, parentJob.TemplateId, ChannelId, SpawnJobTestGroup, "Platform1", PassingBadge);
+
+				// Job Step notification for 'Spawn Platform2 Test'
+				ValidateJobStepInformation(consolidatedJobStepNotificationStates[1], parentJob.Id, parentJob.TemplateId, ChannelId, SpawnJobTestGroup, "Platform2", PassingBadge);
+
+				// Job Step notification for 'Spawn Platform3 Test'
+				ValidateJobStepInformation(consolidatedJobStepNotificationStates[2], parentJob.Id, parentJob.TemplateId, ChannelId, SpawnJobTestGroup, "Platform3", PassingBadge);
+
+				// Job Step notification for 'Client BootTest Platform3'
+				ValidateJobStepInformation(consolidatedJobStepNotificationStates[3], spawnedJob.Id, spawnedJob.TemplateId, ChannelId, ClientBootTestGroup, "Platform3", PendingBadge, parentJob.Id, parentJob.TemplateId);
+
+				// Job Step notification for 'Client Project Tests Platform3'
+				ValidateJobStepInformation(consolidatedJobStepNotificationStates[4], spawnedJob.Id, spawnedJob.TemplateId, ChannelId, ClientProjectTestGroup, "Platform3", PendingBadge, parentJob.Id, parentJob.TemplateId);
+
+				// Job Step notification for 'Editor BootTest Platform3'
+				ValidateJobStepInformation(consolidatedJobStepNotificationStates[5], spawnedJob.Id, spawnedJob.TemplateId, ChannelId, EditorBootTestGroup, "Platform3", PendingBadge, parentJob.Id, parentJob.TemplateId);
+
+				// Job Step notification for 'Editor Project Tests Platform3'
+				ValidateJobStepInformation(consolidatedJobStepNotificationStates[6], spawnedJob.Id, spawnedJob.TemplateId, ChannelId, EditorProjectTestGroup, "Platform3", PendingBadge, parentJob.Id, parentJob.TemplateId);
+
+				spawnedJob = await RunBatchAsync(_mainStreamConfig, spawnedJob, 1);
+				spawnedJob = await RunStepAsync(_mainStreamConfig, spawnedJob, 1, 0, JobStepOutcome.Failure); // Editor BootTest Platform3
+
+				await WaitForNotificationAsync(fakeSink, spawnedJob, spawnedJob.Batches[1].Steps[0]);
+
+				// Validate that the 'Editor BootTest Platform3' step is recorded as having failures and 'Editor Project Tests Platform3' is recorded as being skipped due to having a dependency on the BootTest
+				await ValidateJobStepInformation(jobNotificationCollection, spawnedJob.Id, spawnedJob.TemplateId, ChannelId, spawnedJob.Batches[1].Steps[0].Id, EditorBootTestGroup, "Platform3", FailureBadge, parentJob.Id, parentJob.TemplateId);
+				await ValidateJobStepInformation(jobNotificationCollection, spawnedJob.Id, spawnedJob.TemplateId, ChannelId, spawnedJob.Batches[1].Steps[1].Id, EditorProjectTestGroup, "Platform3", SkippedBadge, parentJob.Id, parentJob.TemplateId);
+
+				spawnedJob = await RunStepAsync(_mainStreamConfig, spawnedJob, 1, 2, JobStepOutcome.Warnings); // Client BootTest Platform3
+
+				await WaitForNotificationAsync(fakeSink, spawnedJob, spawnedJob.Batches[1].Steps[2]);
+
+				// Validate that the Client BootTest Platform3 step is recorded as having warnings
+				await ValidateJobStepInformation(jobNotificationCollection, spawnedJob.Id, spawnedJob.TemplateId, ChannelId, spawnedJob.Batches[1].Steps[2].Id, ClientBootTestGroup, "Platform3", WarningBadge, parentJob.Id, parentJob.TemplateId);
+
+				spawnedJob = await RunStepAsync(_mainStreamConfig, spawnedJob, 1, 3, JobStepOutcome.Success); // Client Project Tests Platform3
+
+				await WaitForNotificationAsync(fakeSink, spawnedJob, spawnedJob.Batches[1].Steps[3]);
+
+				// At this point, the Client Project Tests job step notification would be marked as passing, all notifications would be sent out, and the notifications should be removed from the collections
+				consolidatedJobStepNotificationQueryBuilder = JobNotificationCollection.CreateJobStepNotificationStateQueryBuilder();
+				consolidatedJobStepNotificationQueryBuilder.AddJobAndParentJobFilter(parentJob.Id);
+				consolidatedJobStepNotificationStates = await jobNotificationCollection.GetJobStepNotificationStatesAsync(consolidatedJobStepNotificationQueryBuilder);
+				Assert.IsNull(consolidatedJobStepNotificationStates);
+
+				jobNotificationQueryBuilder = JobNotificationCollection.CreateJobNotificationStateQueryBuilder();
+				jobNotificationQueryBuilder.AddJobFilter(parentJob.Id);
+				jobNotificationStates = await jobNotificationCollection.GetJobNotificationStatesAsync(jobNotificationQueryBuilder);
+				Assert.IsNull(jobNotificationStates);
+			}
+		}
+
+		[TestMethod]
+		public async Task ParentJobCancellationTestAsync()
+		{
+			// Scenario: Job matched with the configuration but is canceled after the Setup Build has been run
+			// Expected: Job steps are found in the JobNotificationCollection after the Setup Build node executes, but is then cleaned up after the job was canceled
+
+			int commitId = 40;
+			NotificationFormatConfig spawnNotificationFormatConfig = CreateNotificationFormatConfig(SpawnJobTestGroup, "^Spawn\\s");
+			NotificationFormatConfig editorNotificationFormatConfig = CreateNotificationFormatConfig(EditorBootTestGroup, "^Editor BootTest\\s");
+			NotificationFormatConfig clientNotificationFormatConfig = CreateNotificationFormatConfig(ClientBootTestGroup, "^Client BootTest\\s");
+
+			JobNotificationConfig jobNotificationConfig = new JobNotificationConfig();
+			jobNotificationConfig.Enabled = true;
+			jobNotificationConfig.Template = _buildTemplateId;
+			jobNotificationConfig.NamePattern = "Platform1|Platform2";
+			jobNotificationConfig.NotificationFormats.Add(spawnNotificationFormatConfig);
+			jobNotificationConfig.NotificationFormats.Add(editorNotificationFormatConfig);
+			jobNotificationConfig.NotificationFormats.Add(clientNotificationFormatConfig);
+			jobNotificationConfig.Channels.Add(ChannelId);
+
+			NotificationStreamConfig notificationStreamConfig = new NotificationStreamConfig();
+			notificationStreamConfig.Streams.Add(_mainStreamId);
+			notificationStreamConfig.JobNotifications.Add(jobNotificationConfig);
+
+			NotificationConfig notificationConfig = new NotificationConfig();
+			notificationConfig.Id = new NotificationConfigId("test-notification");
+			notificationConfig.NotificationStreams.Add(notificationStreamConfig);
+
+			ExperimentalConfig experimentalConfig = new ExperimentalConfig();
+			experimentalConfig.Notifications.Add(notificationConfig);
+
+			UpdateGlobalConfig(experimentalConfig);
+
+			JobNotificationCollection jobNotificationCollection = ServiceProvider.GetRequiredService<JobNotificationCollection>();
+
+			FakeExperimentalNotificationSink fakeSink = ServiceProvider.GetRequiredService<FakeExperimentalNotificationSink>();
+
+			NotificationService service = (NotificationService)ServiceProvider.GetRequiredService<INotificationService>();
+			await service.StartAsync(cancellationToken: default);
+
+			IJob parentJob = await CreateNightlyBuildJobAsync(_mainStreamConfig, commitId);
+			_testJobNotifications.Add(new JobNotificationState(parentJob.Id, parentJob.TemplateId, ChannelId));
+
+			parentJob = await RunBatchAsync(_mainStreamConfig, parentJob, 0);
+			parentJob = await RunStepAsync(_mainStreamConfig, parentJob, 0, 0, JobStepOutcome.Success); // Setup Build
+
+			await WaitForNotificationAsync(fakeSink, parentJob, parentJob.Batches[0].Steps[0]);
+
+			// We should have a single job notification record in the database after the Setup Build job step has completed
+			IJobNotificationStateQueryBuilder jobNotificationQueryBuilder = JobNotificationCollection.CreateJobNotificationStateQueryBuilder();
+			jobNotificationQueryBuilder.AddJobFilter(parentJob.Id);
+			jobNotificationQueryBuilder.AddTemplateFilter(parentJob.TemplateId);
+			jobNotificationQueryBuilder.AddRecipientFilter(ChannelId);
+			IReadOnlyList<IJobNotificationStateRef>? jobNotificationStates = await jobNotificationCollection.GetJobNotificationStatesAsync(jobNotificationQueryBuilder);
+			Assert.IsNotNull(jobNotificationStates);
+			Assert.AreEqual(1, jobNotificationStates.Count);
+
+			// There should be 2 job step notification records for the Spawn Test Job group for each platform (Platform1 and Platform2)
+			IJobStepNotificationStateQueryBuilder jobStepNotificationQueryBuilder = JobNotificationCollection.CreateJobStepNotificationStateQueryBuilder();
+			jobStepNotificationQueryBuilder.AddJobFilter(parentJob.Id);
+			jobStepNotificationQueryBuilder.AddTemplateFilter(parentJob.TemplateId);
+			jobStepNotificationQueryBuilder.AddRecipientFilter(ChannelId);
+			IReadOnlyList<IJobStepNotificationStateRef>? jobStepNotificationStates = await jobNotificationCollection.GetJobStepNotificationStatesAsync(jobStepNotificationQueryBuilder);
+			Assert.IsNotNull(jobStepNotificationStates);
+			Assert.AreEqual(2, jobStepNotificationStates.Count);
+
+			// Job Step notification for 'Spawn Platform1 Test'
+			ValidateJobStepInformation(jobStepNotificationStates[0], parentJob.Id, parentJob.TemplateId, ChannelId, SpawnJobTestGroup, "Platform1", PendingBadge);
+
+			// Job Step notification for 'Spawn Platform2 Test'
+			ValidateJobStepInformation(jobStepNotificationStates[1], parentJob.Id, parentJob.TemplateId, ChannelId, SpawnJobTestGroup, "Platform2", PendingBadge);
+
+			parentJob = await RunBatchAsync(_mainStreamConfig, parentJob, 1);
+
+			// The spawned job will be created before the job step completes
+			IJob spawnedJob1 = await SpawnNightlyBuildTestJobAsync(_mainStreamConfig, parentJob.Id, parentJob.Batches[1].Steps[0].Id, "Platform1", commitId);
+			parentJob = await AddSpawnedJobToParentStepAsync(_mainStreamConfig, parentJob, 1, 0, spawnedJob1.Id);
+			_testJobNotifications.Add(new JobNotificationState(spawnedJob1.Id, spawnedJob1.TemplateId, ChannelId));
+
+			parentJob = await RunStepAsync(_mainStreamConfig, parentJob, 1, 0, JobStepOutcome.Success); // Spawn Platform1 Tests
+
+			// Spawn the second job
+			IJob spawnedJob2 = await SpawnNightlyBuildTestJobAsync(_mainStreamConfig, parentJob.Id, parentJob.Batches[1].Steps[1].Id, "Platform2", commitId);
+			parentJob = await AddSpawnedJobToParentStepAsync(_mainStreamConfig, parentJob, 1, 1, spawnedJob2.Id);
+			_testJobNotifications.Add(new JobNotificationState(spawnedJob2.Id, spawnedJob2.TemplateId, ChannelId));
+
+			parentJob = await RunStepAsync(_mainStreamConfig, parentJob, 1, 1, JobStepOutcome.Success); // Spawn Platform2 Tests
+
+			await WaitForNotificationAsync(fakeSink, parentJob, parentJob.Batches[1].Steps[1]);
+
+			// Execute the Setup Build steps of all of our spawned jobs
+			spawnedJob1 = await RunBatchAsync(_mainStreamConfig, spawnedJob1, 0);
+			spawnedJob1 = await RunStepAsync(_mainStreamConfig, spawnedJob1, 0, 0, JobStepOutcome.Success); // Setup Build
+
+			spawnedJob2 = await RunBatchAsync(_mainStreamConfig, spawnedJob2, 0);
+			spawnedJob2 = await RunStepAsync(_mainStreamConfig, spawnedJob2, 0, 0, JobStepOutcome.Success); // Setup Build
+
+			await WaitForNotificationAsync(fakeSink, spawnedJob2, spawnedJob2.Batches[0].Steps[0]);
+
+			// The spawned job will append its notification records to the parent job and there will be a total of have 10 job step notification records
+			IJobStepNotificationStateQueryBuilder consolidatedJobStepNotificationQueryBuilder = JobNotificationCollection.CreateJobStepNotificationStateQueryBuilder();
+			consolidatedJobStepNotificationQueryBuilder.AddJobAndParentJobFilter(parentJob.Id);
+			consolidatedJobStepNotificationQueryBuilder.AddTemplateAndParentTemplateFilter(parentJob.TemplateId);
+			consolidatedJobStepNotificationQueryBuilder.AddRecipientFilter(ChannelId);
+			IReadOnlyList<IJobStepNotificationStateRef>? consolidatedJobStepNotificationStates = await jobNotificationCollection.GetJobStepNotificationStatesAsync(consolidatedJobStepNotificationQueryBuilder);
+			Assert.IsNotNull(consolidatedJobStepNotificationStates);
+			Assert.AreEqual(6, consolidatedJobStepNotificationStates.Count);
+
+			// Simulate the cancellation of the parent job
+			parentJob = await AbortJobAsync(parentJob);
+
+			await WaitForNotificationAsync(fakeSink, parentJob, null);
+
+			// The job should be removed from the collection
+			jobNotificationQueryBuilder = JobNotificationCollection.CreateJobNotificationStateQueryBuilder();
+			jobNotificationQueryBuilder.AddJobFilter(parentJob.Id);
+			jobNotificationQueryBuilder.AddTemplateFilter(parentJob.TemplateId);
+			jobNotificationQueryBuilder.AddRecipientFilter(ChannelId);
+			jobNotificationStates = await jobNotificationCollection.GetJobNotificationStatesAsync(jobNotificationQueryBuilder);
+			Assert.IsNull(jobNotificationStates);
+
+			// The parent and spawned job steps should be removed from the collection
+			consolidatedJobStepNotificationQueryBuilder = JobNotificationCollection.CreateJobStepNotificationStateQueryBuilder();
+			consolidatedJobStepNotificationQueryBuilder.AddJobAndParentJobFilter(parentJob.Id);
+			consolidatedJobStepNotificationStates = await jobNotificationCollection.GetJobStepNotificationStatesAsync(consolidatedJobStepNotificationQueryBuilder);
+			Assert.IsNull(consolidatedJobStepNotificationStates);
+		}
+	}
+}
diff --git a/Engine/Source/Programs/Horde/Plugins/Experimental/HordeServer.Experimental/ExperimentalPlugin.cs b/Engine/Source/Programs/Horde/Plugins/Experimental/HordeServer.Experimental/ExperimentalPlugin.cs
--- a/Engine/Source/Programs/Horde/Plugins/Experimental/HordeServer.Experimental/ExperimentalPlugin.cs
+++ b/Engine/Source/Programs/Horde/Plugins/Experimental/HordeServer.Experimental/ExperimentalPlugin.cs
@@ -1,5 +1,6 @@
 // Copyright Epic Games, Inc. All Rights Reserved.
 
+using HordeServer.Experimental.Notifications;
 using HordeServer.Notifications;
 using HordeServer.Notifications.Sinks;
 using HordeServer.Plugins;
@@ -28,9 +29,10 @@ namespace HordeServer
 		}
 
 		/// <inheritdoc/>
-		public void ConfigureServices(IServiceCollection services)
+		public void ConfigureServices(IServiceCollection serviceCollection)
 		{
-			services.AddSingleton<INotificationSink, ExperimentalSlackNotificationSink>();
+			serviceCollection.AddSingleton<IJobNotificationCollection, JobNotificationCollection>();
+			serviceCollection.AddSingleton<INotificationSink, ExperimentalSlackNotificationSink>();
 		}
 	}
 
diff --git a/Engine/Source/Programs/Horde/Plugins/Experimental/HordeServer.Experimental/Notifications/IJobNotification.cs b/Engine/Source/Programs/Horde/Plugins/Experimental/HordeServer.Experimental/Notifications/IJobNotification.cs
new file mode 100644
--- /dev/null
+++ b/Engine/Source/Programs/Horde/Plugins/Experimental/HordeServer.Experimental/Notifications/IJobNotification.cs
@@ -0,0 +1,358 @@
+// Copyright Epic Games, Inc. All Rights Reserved.
+
+using EpicGames.Horde.Jobs;
+using EpicGames.Horde.Jobs.Templates;
+using MongoDB.Bson;
+
+namespace HordeServer.Experimental.Notifications
+{
+	/// <summary>
+	/// A job notification state reference
+	/// </summary>
+	public interface IJobNotificationStateRef
+	{
+		/// <summary>
+		/// The job notification state identifier
+		/// </summary>
+		ObjectId Id { get; }
+
+		/// <summary>
+		/// The job identifier
+		/// </summary>
+		JobId JobId { get; }
+
+		/// <summary>
+		/// The template identifier
+		/// </summary>
+		TemplateId TemplateId { get; }
+
+		/// <summary>
+		/// The recipient of the notification
+		/// </summary>
+		string Recipient { get; }
+
+		/// <summary>
+		/// The channel of the notification
+		/// </summary>
+		string Channel { get; }
+
+		/// <summary>
+		/// The timestamp of the notification
+		/// </summary>
+		string Ts { get; }
+	}
+
+	/// <summary>
+	/// A job step notification state reference
+	/// </summary>
+	public interface IJobStepNotificationStateRef
+	{
+		/// <summary>
+		/// The job notification state identifier
+		/// </summary>
+		ObjectId Id { get; }
+
+		/// <summary>
+		/// The job identifier
+		/// </summary>
+		JobId JobId { get; }
+
+		/// <summary>
+		/// The job step identifier
+		/// </summary>
+		JobStepId JobStepId { get; set; }
+
+		/// <summary>
+		/// The template identifier
+		/// </summary>
+		TemplateId TemplateId { get; }
+
+		/// <summary>
+		/// The recipient of the notification
+		/// </summary>
+		string Recipient { get; }
+
+		/// <summary>
+		/// The associated group of the notification
+		/// </summary>
+		string Group { get; }
+
+		/// <summary>
+		/// The target platform of the notification
+		/// </summary>
+		string TargetPlatform { get; }
+
+		/// <summary>
+		/// The badge of the notification
+		/// </summary>
+		string Badge { get; }
+
+		/// <summary>
+		/// The channel of the notification
+		/// </summary>
+		string Channel { get; }
+
+		/// <summary>
+		/// The timestamp of the threaded notification
+		/// </summary>
+		string? ThreadTs { get; }
+
+		/// <summary>
+		/// The timestamp of the notification
+		/// </summary>
+		string Ts { get; }
+
+		/// <summary>
+		/// The parent job identifier
+		/// </summary>
+		JobId? ParentJobId { get; }
+
+		/// <summary>
+		/// The parent job template identifier
+		/// </summary>
+		TemplateId? ParentJobTemplateId { get; }
+	}
+
+	/// <summary>
+	/// Builder to simplify the query generation for job notification states
+	/// </summary>
+	public interface IJobNotificationStateQueryBuilder
+	{
+		/// <summary>
+		/// Add the filter for a job
+		/// </summary>
+		/// <param name="jobId">Identifier of the job to filter by</param>
+		/// <returns>A reference to the current <see cref="IJobNotificationStateQueryBuilder"/></returns>
+		public IJobNotificationStateQueryBuilder AddJobFilter(JobId jobId);
+
+		/// <summary>
+		/// Adds a filter for a list of jobs
+		/// </summary>
+		/// <param name="jobIds">List of job identifiers to filter by</param>
+		/// <returns>A reference to the current <see cref="IJobNotificationStateQueryBuilder"/></returns>
+		public IJobNotificationStateQueryBuilder AddJobFilters(List<JobId> jobIds);
+
+		/// <summary>
+		/// Add the filter for a template
+		/// </summary>
+		/// <param name="templateId">Identifier of the template to filter by</param>
+		/// <returns>A reference to the current <see cref="IJobNotificationStateQueryBuilder"/></returns>
+		public IJobNotificationStateQueryBuilder AddTemplateFilter(TemplateId templateId);
+
+		/// <summary>
+		/// Adds a filter for a list of templates
+		/// </summary>
+		/// <param name="templateIds">List of template identifiers to filter by</param>
+		/// <returns>A reference to the current <see cref="IJobNotificationStateQueryBuilder"/></returns>
+		public IJobNotificationStateQueryBuilder AddTemplateFilters(List<TemplateId> templateIds);
+
+		/// <summary>
+		/// Add the filter for a specific recipient
+		/// </summary>
+		/// <param name="recipient">Recipient to filter by</param>
+		/// <returns>A reference to the current <see cref="IJobNotificationStateQueryBuilder"/></returns>
+		public IJobNotificationStateQueryBuilder AddRecipientFilter(string recipient);
+
+		/// <summary>
+		/// Add the filter for a specific channel
+		/// </summary>
+		/// <param name="channel">Channel to filter by</param>
+		/// <returns>A reference to the current <see cref="IJobNotificationStateQueryBuilder"/></returns>
+		public IJobNotificationStateQueryBuilder AddChannelFilter(string channel);
+
+		/// <summary>
+		/// Add the filter for a timestamp
+		/// </summary>
+		/// <param name="ts">Timestamp to filter by</param>
+		/// <returns>A reference to the current <see cref="IJobNotificationStateQueryBuilder"/></returns>
+		public IJobNotificationStateQueryBuilder AddTimestampFilter(string ts);
+	}
+
+	/// <summary>
+	/// Builder to simplify the query generation for job step notification states
+	/// </summary>
+	public interface IJobStepNotificationStateQueryBuilder
+	{
+		/// <summary>
+		/// Add the filter for a job
+		/// </summary>
+		/// <param name="jobId">Identifier of the job to filter by</param>
+		/// <returns>A reference to the current <see cref="IJobStepNotificationStateQueryBuilder"/></returns>
+		public IJobStepNotificationStateQueryBuilder AddJobFilter(JobId jobId);
+
+		/// <summary>
+		/// Add the filter for a template
+		/// </summary>
+		/// <param name="templateId">Identifier of the template to filter by</param>
+		/// <returns>A reference to the current <see cref="IJobStepNotificationStateQueryBuilder"/></returns>
+		public IJobStepNotificationStateQueryBuilder AddTemplateFilter(TemplateId templateId);
+
+		/// <summary>
+		/// Add the filter for a specific recipient
+		/// </summary>
+		/// <param name="recipient">Recipient to filter by</param>
+		/// <returns>A reference to the current <see cref="IJobStepNotificationStateQueryBuilder"/></returns>
+		public IJobStepNotificationStateQueryBuilder AddRecipientFilter(string recipient);
+
+		/// <summary>
+		/// Add the filter for a job step
+		/// </summary>
+		/// <param name="jobStepId">Identifier of the job step to filter by</param>
+		/// <returns>A reference to the current <see cref="IJobStepNotificationStateQueryBuilder"/></returns>
+		public IJobStepNotificationStateQueryBuilder AddJobStepFilter(JobStepId jobStepId);
+
+		/// <summary>
+		/// Add the filter for a specific group
+		/// </summary>
+		/// <param name="group">Group to filter by</param>
+		/// <returns>A reference to the current <see cref="IJobStepNotificationStateQueryBuilder"/></returns>
+		public IJobStepNotificationStateQueryBuilder AddGroupFilter(string group);
+
+		/// <summary>
+		/// Add the filter for a specific platform
+		/// </summary>
+		/// <param name="platform">Platform to filter by</param>
+		/// <returns>A reference to the current <see cref="IJobStepNotificationStateQueryBuilder"/></returns>
+		public IJobStepNotificationStateQueryBuilder AddPlatformFilter(string platform);
+
+		/// <summary>
+		/// Add the filter for a specific badge
+		/// </summary>
+		/// <param name="badge">Badge to filter by</param>
+		/// <returns>A reference to the current <see cref="IJobStepNotificationStateQueryBuilder"/></returns>
+		public IJobStepNotificationStateQueryBuilder AddBadgeFilter(string badge);
+
+		/// <summary>
+		/// Add the filter for a specific channel
+		/// </summary>
+		/// <param name="channel">Channel to filter by</param>
+		/// <returns>A reference to the current <see cref="IJobStepNotificationStateQueryBuilder"/></returns>
+		public IJobStepNotificationStateQueryBuilder AddChannelFilter(string channel);
+
+		/// <summary>
+		/// Add the filter for a threaded timestamp
+		/// </summary>
+		/// <param name="threadTs">Thread timestamp to filter by</param>
+		/// <returns>A reference to the current <see cref="IJobStepNotificationStateQueryBuilder"/></returns>
+		public IJobStepNotificationStateQueryBuilder AddThreadTimestampFilter(string threadTs);
+
+		/// <summary>
+		/// Add the filter for a timestamp
+		/// </summary>
+		/// <param name="ts">Timestamp to filter by</param>
+		/// <returns>A reference to the current <see cref="IJobStepNotificationStateQueryBuilder"/></returns>
+		public IJobStepNotificationStateQueryBuilder AddTimestampFilter(string ts);
+
+		/// <summary>
+		/// Add the filter for a parent job
+		/// </summary>
+		/// <param name="jobId">Identifier of the job to filter by</param>
+		/// <returns>A reference to the current <see cref="IJobStepNotificationStateQueryBuilder"/></returns>
+		public IJobStepNotificationStateQueryBuilder AddParentJobFilter(JobId jobId);
+
+		/// <summary>
+		/// Add the filter for a parent template
+		/// </summary>
+		/// <param name="templateId">Identifier of the template to filter by</param>
+		/// <returns>A reference to the current <see cref="IJobStepNotificationStateQueryBuilder"/></returns>
+		public IJobStepNotificationStateQueryBuilder AddParentJobTemplateFilter(TemplateId templateId);
+
+		/// <summary>
+		/// Add the filter for either a job or parent job
+		/// </summary>
+		/// <param name="jobId">Identifier of the job to filter by</param>
+		/// <returns>A reference to the current <see cref="IJobStepNotificationStateQueryBuilder"/></returns>
+		public IJobStepNotificationStateQueryBuilder AddJobAndParentJobFilter(JobId jobId);
+
+		/// <summary>
+		/// Add the filter for either a template or parent template
+		/// </summary>
+		/// <param name="templateId">Identifier of the template to filter by</param>
+		/// <returns>A reference to the current <see cref="IJobStepNotificationStateQueryBuilder"/></returns>
+		public IJobStepNotificationStateQueryBuilder AddTemplateAndParentTemplateFilter(TemplateId templateId);
+	}
+
+	/// <summary>
+	/// Collection of job notification documents
+	/// </summary>
+	public interface IJobNotificationCollection
+	{
+		/// <summary>
+		/// Fetches a single entity of <see cref="IJobNotificationStateRef"/>
+		/// </summary>
+		/// <param name="jobId">Identifier of the job</param>
+		/// <param name="templateId">Identifier of the stream's template</param>
+		/// <param name="recipient">Information regarding a specific recipient</param>
+		/// <param name="cancellationToken">Cancellation token for the operation</param>
+		/// <returns>A <see cref="IJobNotificationStateRef"/> if records were found, otherwise returns null for error checking</returns>
+		/// <remarks>Only a single <see cref="IJobNotificationStateRef"/> should exist for the aggregate of the job identifier, template identifier, and recipient should ever exist. Will report an error if another record had been found.</remarks>
+		Task<IJobNotificationStateRef?> GetJobNotificationStateAsync(JobId jobId, TemplateId templateId, string recipient, CancellationToken cancellationToken = default);
+
+		/// <summary>
+		/// Fetches a read only list of <see cref="IJobNotificationStateRef"/>
+		/// </summary>
+		/// <param name="builder">The <see cref="IJobNotificationStateQueryBuilder"/> used to generate a query to filter against</param>
+		/// <param name="cancellationToken">Cancellation token for the operation</param>
+		/// <returns>A read only list of <see cref="IJobNotificationStateRef"/> if records were found, otherwise returns null for error checking</returns>
+		Task<IReadOnlyList<IJobNotificationStateRef>?> GetJobNotificationStatesAsync(IJobNotificationStateQueryBuilder builder, CancellationToken cancellationToken = default);
+
+		/// <summary>
+		/// Add or update the <see cref="IJobNotificationStateRef"/> to the collection
+		/// </summary>
+		/// <param name="jobId">Identifier of the job</param>
+		/// <param name="templateId">Identifier of the stream's template</param>
+		/// <param name="recipient">Information regarding a specific recipient</param>
+		/// <param name="channel">Information regarding a specific channel</param>
+		/// <param name="timestamp">Timestamp for the notification</param>
+		/// <param name="cancellationToken">Cancellation token for the operation</param>
+		/// <returns>A <see cref="IJobNotificationStateRef"/> object</returns>
+		Task<IJobNotificationStateRef> AddOrUpdateJobNotificationStateAsync(JobId jobId, TemplateId templateId, string recipient, string channel, string timestamp, CancellationToken cancellationToken = default);
+
+		/// <summary>
+		/// Deletes a <see cref="IJobNotificationStateRef"/> from the collection
+		/// </summary>
+		/// <param name="jobId">Identifier of the job</param>
+		/// <param name="templateId">Identifier of the stream's template</param>
+		/// <param name="recipient">Information regarding a specific recipient</param>
+		/// <param name="cancellationToken">Cancellation token for the operation</param>
+		/// <returns>The number of records which were deleted</returns>
+		Task<long> DeleteJobNotificationStatesAsync(JobId jobId, TemplateId templateId, string recipient, CancellationToken cancellationToken = default);
+
+		/// <summary>
+		/// Fetches a read only list of <see cref="IJobStepNotificationStateRef"/>
+		/// </summary>
+		/// <param name="builder">The <see cref="IJobStepNotificationStateQueryBuilder"/> used to generate a query to filter against</param>
+		/// <param name="cancellationToken">Cancellation token for the operation</param>
+		/// <returns>A read only list of <see cref="IJobStepNotificationStateRef"/> if records were found, otherwise returns null for error checking</returns>
+		Task<IReadOnlyList<IJobStepNotificationStateRef>?> GetJobStepNotificationStatesAsync(IJobStepNotificationStateQueryBuilder builder, CancellationToken cancellationToken = default);
+
+		/// <summary>
+		/// Add or update the <see cref="IJobStepNotificationStateRef"/> to the collection
+		/// </summary>
+		/// <param name="jobId">Identifier of the job</param>
+		/// <param name="templateId">Identifier of the stream's template</param>
+		/// <param name="recipient">Information regarding a specific recipient</param>
+		/// <param name="jobStepId">Identifier of the job step</param>
+		/// <param name="group">Group associated with the notification</param>
+		/// <param name="platform">Target platform associated with the notification</param>
+		/// <param name="badge">Slack emoji associated with the notification</param>
+		/// <param name="channel">Information regarding a specific channel</param>
+		/// <param name="ts">Timestamp for the notification</param>
+		/// <param name="threadTs">Optional Threaded timestamp for the notification</param>
+		/// <param name="parentJobId">Optional identifier of the parent job</param>
+		/// <param name="parentJobTemplateId">Optional identifier of the parent stream's template</param>
+		/// <param name="cancellationToken">Cancellation token for the operation</param>
+		/// <returns>A <see cref="IJobStepNotificationStateRef"/> object</returns>
+		Task<IJobStepNotificationStateRef> AddOrUpdateJobStepNotificationStateAsync(JobId jobId, TemplateId templateId, string recipient, JobStepId jobStepId, string group, string platform, string badge, string channel, string ts, string? threadTs = null, JobId? parentJobId = null, TemplateId? parentJobTemplateId = null, CancellationToken cancellationToken = default);
+
+		/// <summary>
+		/// Deletes <see cref="IJobStepNotificationStateRef"/> from the collection
+		/// </summary>
+		/// <param name="jobId">Identifier of the job</param>
+		/// <param name="templateId">Identifier of the stream's template</param>
+		/// <param name="recipient">Information regarding a specific recipient</param>
+		/// <param name="cancellationToken">Cancellation token for the operation</param>
+		/// <returns>The number of records which were deleted</returns>
+		Task<long> DeleteJobStepNotificationStatesAsync(JobId jobId, TemplateId templateId, string recipient, CancellationToken cancellationToken = default);
+	}
+}
diff --git a/Engine/Source/Programs/Horde/Plugins/Experimental/HordeServer.Experimental/Notifications/JobNotificationCollection.cs b/Engine/Source/Programs/Horde/Plugins/Experimental/HordeServer.Experimental/Notifications/JobNotificationCollection.cs
new file mode 100644
--- /dev/null
+++ b/Engine/Source/Programs/Horde/Plugins/Experimental/HordeServer.Experimental/Notifications/JobNotificationCollection.cs
@@ -0,0 +1,471 @@
+// Copyright Epic Games, Inc. All Rights Reserved.
+
+using EpicGames.Horde.Jobs;
+using EpicGames.Horde.Jobs.Templates;
+using EpicGames.Slack;
+using HordeServer.Server;
+using Microsoft.Extensions.Logging;
+using MongoDB.Bson;
+using MongoDB.Bson.Serialization.Attributes;
+using MongoDB.Driver;
+
+namespace HordeServer.Experimental.Notifications
+{
+	/// <summary>
+	/// Wrapper around the job notifications collection in a Mongo DB
+	/// </summary>
+	public class JobNotificationCollection : IJobNotificationCollection
+	{
+		/// <summary>
+		/// MongoDB document for job notifications
+		/// </summary>
+		class JobNotificationStateDocument : IJobNotificationStateRef
+		{
+			[BsonId]
+			public ObjectId Id { get; set; }
+
+			[BsonElement("jid")]
+			public JobId JobId { get; set; }
+
+			[BsonElement("tid")]
+			public TemplateId TemplateId { get; set; }
+
+			[BsonElement("rcp")]
+			public string Recipient { get; set; } = String.Empty;
+
+			[BsonElement("ch")]
+			public string Channel { get; set; } = String.Empty;
+
+			[BsonElement("ts")]
+			public string Ts { get; set; } = String.Empty;
+
+			[BsonIgnore]
+			public SlackMessageId MessageId => new SlackMessageId(Channel, null, Ts);
+
+			private JobNotificationStateDocument()
+			{
+				Id = ObjectId.GenerateNewId();
+			}
+
+			public JobNotificationStateDocument(JobId jobId, TemplateId templateId, string recipient, string channel, string ts)
+			{
+				Id = ObjectId.GenerateNewId();
+				JobId = jobId;
+				TemplateId = templateId;
+				Recipient = recipient;
+				Channel = channel;
+				Ts = ts;
+			}
+		}
+
+		/// <summary>
+		/// MongoDB document for job step notifications
+		/// </summary>
+		class JobStepNotificationStateDocument : IJobStepNotificationStateRef
+		{
+			[BsonId]
+			public ObjectId Id { get; set; }
+
+			[BsonElement("jid")]
+			public JobId JobId { get; set; }
+
+			[BsonElement("tid")]
+			public TemplateId TemplateId { get; set; }
+
+			[BsonElement("rcp")]
+			public string Recipient { get; set; } = String.Empty;
+
+			[BsonElement("sid")]
+			public JobStepId JobStepId { get; set; }
+
+			[BsonElement("grp")]
+			public string Group { get; set; } = String.Empty;
+
+			[BsonElement("plt")]
+			public string TargetPlatform { get; set; } = String.Empty;
+
+			[BsonElement("bdg")]
+			public string Badge { get; set; } = String.Empty;
+
+			[BsonElement("ch")]
+			public string Channel { get; set; } = String.Empty;
+
+			[BsonElement("tts")]
+			public string? ThreadTs { get; set; }
+
+			[BsonElement("ts")]
+			public string Ts { get; set; } = String.Empty;
+
+			[BsonIgnoreIfNull, BsonElement("pjid")]
+			public JobId? ParentJobId { get; set; } = null;
+
+			[BsonIgnoreIfNull, BsonElement("ptid")]
+			public TemplateId? ParentJobTemplateId { get; set; } = null;
+
+			[BsonIgnore]
+			public SlackMessageId MessageId => new SlackMessageId(Channel, ThreadTs, Ts);
+
+			private JobStepNotificationStateDocument()
+			{
+				Id = ObjectId.GenerateNewId();
+			}
+
+			public JobStepNotificationStateDocument(JobId jobId, TemplateId templateId, string recipient, JobStepId jobStepId, string group, string targetPlatform, string badge, string channel, string ts, string? tts = null, JobId? parentJobId = null, TemplateId? parentTemplateId = null)
+			{
+				Id = ObjectId.GenerateNewId();
+				JobId = jobId;
+				TemplateId = templateId;
+				Recipient = recipient;
+				JobStepId = jobStepId;
+				Group = group;
+				TargetPlatform = targetPlatform;
+				Badge = badge;
+				Channel = channel;
+				Ts = ts;
+				ThreadTs = tts;
+				ParentJobId = parentJobId;
+				ParentJobTemplateId = parentTemplateId;
+			}
+		}
+
+		/// <inheritdoc/>
+		class JobNotificationStateQueryBuilder : IJobNotificationStateQueryBuilder
+		{
+			public FilterDefinition<JobNotificationStateDocument> Filter { get; private set; }
+
+			readonly FilterDefinitionBuilder<JobNotificationStateDocument> _filterBuilder;
+
+			public JobNotificationStateQueryBuilder()
+			{
+				_filterBuilder = Builders<JobNotificationStateDocument>.Filter;
+				Filter = _filterBuilder.Empty;
+			}
+
+			public IJobNotificationStateQueryBuilder AddJobFilter(JobId jobId)
+			{
+				Filter &= _filterBuilder.Eq(x => x.JobId, jobId);
+				return this;
+			}
+
+			public IJobNotificationStateQueryBuilder AddJobFilters(List<JobId> jobIds)
+			{
+				Filter &= _filterBuilder.In(x => x.JobId, jobIds);
+				return this;
+			}
+
+			public IJobNotificationStateQueryBuilder AddTemplateFilter(TemplateId templateId)
+			{
+				Filter &= _filterBuilder.Eq(x => x.TemplateId, templateId);
+				return this;
+			}
+
+			public IJobNotificationStateQueryBuilder AddTemplateFilters(List<TemplateId> templateId)
+			{
+				Filter &= _filterBuilder.In(x => x.TemplateId, templateId);
+				return this;
+			}
+
+			public IJobNotificationStateQueryBuilder AddRecipientFilter(string recipient)
+			{
+				Filter &= _filterBuilder.Eq(x => x.Recipient, recipient);
+				return this;
+			}
+
+			public IJobNotificationStateQueryBuilder AddChannelFilter(string channel)
+			{
+				Filter &= _filterBuilder.Eq(x => x.Channel, channel);
+				return this;
+			}
+
+			public IJobNotificationStateQueryBuilder AddTimestampFilter(string ts)
+			{
+				Filter &= _filterBuilder.Eq(x => x.Ts, ts);
+				return this;
+			}
+		}
+
+		/// <inheritdoc/>
+		class JobStepNotificationStateQueryBuilder : IJobStepNotificationStateQueryBuilder
+		{
+			public FilterDefinition<JobStepNotificationStateDocument> Filter { get; private set; }
+
+			readonly FilterDefinitionBuilder<JobStepNotificationStateDocument> _filterBuilder;
+
+			public JobStepNotificationStateQueryBuilder()
+			{
+				_filterBuilder = Builders<JobStepNotificationStateDocument>.Filter;
+				Filter = _filterBuilder.Empty;
+			}
+
+			public IJobStepNotificationStateQueryBuilder AddJobFilter(JobId jobId)
+			{
+				Filter &= _filterBuilder.Eq(x => x.JobId, jobId);
+				return this;
+			}
+
+			public IJobStepNotificationStateQueryBuilder AddTemplateFilter(TemplateId templateId)
+			{
+				Filter &= _filterBuilder.Eq(x => x.TemplateId, templateId);
+				return this;
+			}
+
+			public IJobStepNotificationStateQueryBuilder AddRecipientFilter(string recipient)
+			{
+				Filter &= _filterBuilder.Eq(x => x.Recipient, recipient);
+				return this;
+			}
+
+			public IJobStepNotificationStateQueryBuilder AddJobStepFilter(JobStepId jobStepId)
+			{
+				Filter &= _filterBuilder.Eq(x => x.JobStepId, jobStepId);
+				return this;
+			}
+
+			public IJobStepNotificationStateQueryBuilder AddGroupFilter(string group)
+			{
+				Filter &= _filterBuilder.Eq(x => x.Group, group);
+				return this;
+			}
+
+			public IJobStepNotificationStateQueryBuilder AddPlatformFilter(string platform)
+			{
+				Filter &= _filterBuilder.Eq(x => x.TargetPlatform, platform);
+				return this;
+			}
+
+			public IJobStepNotificationStateQueryBuilder AddBadgeFilter(string badge)
+			{
+				Filter &= _filterBuilder.Eq(x => x.Badge, badge);
+				return this;
+			}
+
+			public IJobStepNotificationStateQueryBuilder AddChannelFilter(string channel)
+			{
+				Filter &= _filterBuilder.Eq(x => x.Channel, channel);
+				return this;
+			}
+
+			public IJobStepNotificationStateQueryBuilder AddThreadTimestampFilter(string threadTs)
+			{
+				Filter &= _filterBuilder.Eq(x => x.ThreadTs, threadTs);
+				return this;
+			}
+
+			public IJobStepNotificationStateQueryBuilder AddTimestampFilter(string ts)
+			{
+				Filter &= _filterBuilder.Eq(x => x.Ts, ts);
+				return this;
+			}
+
+			public IJobStepNotificationStateQueryBuilder AddParentJobFilter(JobId jobId)
+			{
+				Filter &= _filterBuilder.Eq(x => x.ParentJobId, jobId);
+				return this;
+			}
+
+			public IJobStepNotificationStateQueryBuilder AddParentJobTemplateFilter(TemplateId templateId)
+			{
+				Filter &= _filterBuilder.Eq(x => x.ParentJobTemplateId, templateId);
+				return this;
+			}
+
+			public IJobStepNotificationStateQueryBuilder AddJobAndParentJobFilter(JobId jobId)
+			{
+				Filter &= _filterBuilder.Where(x => x.JobId == jobId || x.ParentJobId == jobId);
+				return this;
+			}
+
+			public IJobStepNotificationStateQueryBuilder AddTemplateAndParentTemplateFilter(TemplateId templateId)
+			{
+				Filter &= _filterBuilder.Where(x => x.TemplateId == templateId || x.ParentJobTemplateId == templateId);
+				return this;
+			}
+		}
+
+		readonly IMongoCollection<JobNotificationStateDocument> _jobNotificationStates;
+		readonly IMongoCollection<JobStepNotificationStateDocument> _jobStepNotificationStates;
+		readonly ILogger _logger;
+
+		/// <summary>
+		/// Constructor
+		/// </summary>
+		public JobNotificationCollection(IMongoService mongoService, ILogger<JobNotificationCollection> logger)
+		{
+			_logger = logger;
+
+			List<MongoIndex<JobNotificationStateDocument>> jobNotificationIndexes = new List<MongoIndex<JobNotificationStateDocument>>();
+			jobNotificationIndexes.Add(keys => keys.Ascending(x => x.JobId).Ascending(x => x.TemplateId).Ascending(x => x.Recipient), unique: true);
+			_jobNotificationStates = mongoService.GetCollection<JobNotificationStateDocument>("SlackExperimentalJob", jobNotificationIndexes);
+
+			List<MongoIndex<JobStepNotificationStateDocument>> jobNotificationStepIndexes = new List<MongoIndex<JobStepNotificationStateDocument>>();
+			jobNotificationStepIndexes.Add(keys => keys.Ascending(x => x.JobId).Ascending(x => x.TemplateId).Ascending(x => x.Recipient).Ascending(x => x.Group).Ascending(x => x.JobStepId), unique: true);
+			_jobStepNotificationStates = mongoService.GetCollection<JobStepNotificationStateDocument>("SlackExperimentalJobStep", jobNotificationStepIndexes);
+		}
+
+		#region Job Notification
+
+		/// <summary>
+		/// Creates a <see cref="JobNotificationStateQueryBuilder"/> for creating a <see cref="JobNotificationStateDocument"/> query filter
+		/// </summary>
+		/// <returns>A <see cref="IJobNotificationStateQueryBuilder"/> builder</returns>
+		public static IJobNotificationStateQueryBuilder CreateJobNotificationStateQueryBuilder()
+		{
+			return new JobNotificationStateQueryBuilder();
+		}
+
+		/// <inheritdoc/>
+		public async Task<IReadOnlyList<IJobNotificationStateRef>?> GetJobNotificationStatesAsync(IJobNotificationStateQueryBuilder builder, CancellationToken cancellationToken = default)
+		{
+			JobNotificationStateQueryBuilder? jobNotificationStateQueryBuilder = builder as JobNotificationStateQueryBuilder;
+			if (jobNotificationStateQueryBuilder is null)
+			{
+				return null;
+			}
+
+			List<JobNotificationStateDocument> states = await _jobNotificationStates.Find(jobNotificationStateQueryBuilder.Filter).ToListAsync(cancellationToken);
+			return (states.Count != 0) ? states : null;
+		}
+
+		/// <inheritdoc/>
+		public async Task<IJobNotificationStateRef?> GetJobNotificationStateAsync(JobId jobId, TemplateId templateId, string recipient, CancellationToken cancellationToken = default)
+		{
+			JobNotificationStateQueryBuilder builder = new JobNotificationStateQueryBuilder();
+			builder.AddJobFilter(jobId);
+			builder.AddTemplateFilter(templateId);
+			builder.AddRecipientFilter(recipient);
+			IReadOnlyList<IJobNotificationStateRef>? states = await GetJobNotificationStatesAsync(builder, cancellationToken);
+			if (states is null || states.Count == 0)
+			{
+				return null;
+			}
+
+			if (states.Count != 1)
+			{
+				// Log error that multiple records were found and could cause issues
+				_logger.LogInformation("Found {Count} records when fetching for {JobId}:{TemplateId}:{Recipient}. Using first record which could cause issues.", states.Count, jobId, templateId, recipient);
+			}
+
+			return states[0];
+		}
+
+		/// <inheritdoc/>
+		public async Task<IJobNotificationStateRef> AddOrUpdateJobNotificationStateAsync(JobId jobId, TemplateId templateId, string recipient, string channel, string timestamp, CancellationToken cancellationToken = default)
+		{
+			JobNotificationStateQueryBuilder builder = new JobNotificationStateQueryBuilder();
+			builder.AddJobFilter(jobId);
+			builder.AddTemplateFilter(templateId);
+			builder.AddRecipientFilter(recipient);
+
+			JobNotificationStateDocument? jobNotificationState = await _jobNotificationStates.Find(builder.Filter).FirstOrDefaultAsync(cancellationToken);
+			if (jobNotificationState == null)
+			{
+				jobNotificationState = new JobNotificationStateDocument(jobId, templateId, recipient, channel, timestamp);
+				await _jobNotificationStates.InsertOneAsync(jobNotificationState, null, cancellationToken);
+				_logger.LogInformation("Posted message {StateId} (jobId: {JobId}, templateId: {TemplateId}, recipient: {Recipient}, messageId: {MessageId})",
+										jobNotificationState.Id, jobNotificationState.JobId, jobNotificationState.TemplateId, jobNotificationState.Recipient, jobNotificationState.MessageId);
+			}
+			else
+			{
+				UpdateDefinitionBuilder<JobNotificationStateDocument> updateBuilder = Builders<JobNotificationStateDocument>.Update;
+				UpdateDefinition<JobNotificationStateDocument> updates = updateBuilder.Set(x => x.Channel, channel).Set(x => x.Ts, timestamp);
+				jobNotificationState = await _jobNotificationStates.FindOneAndUpdateAsync(x => x.Id == jobNotificationState.Id, updates, null, cancellationToken);
+				_logger.LogInformation("Updated message {StateId} (jobId: {JobId}, templateId: {TemplateId}, recipient: {Recipient}, messageId: {MessageId})",
+										jobNotificationState.Id, jobNotificationState.JobId, jobNotificationState.TemplateId, jobNotificationState.Recipient, jobNotificationState.MessageId);
+			}
+
+			return jobNotificationState;
+		}
+
+		/// <inheritdoc/>
+		public async Task<long> DeleteJobNotificationStatesAsync(JobId jobId, TemplateId templateId, string recipient, CancellationToken cancellationToken = default)
+		{
+			JobNotificationStateQueryBuilder builder = new JobNotificationStateQueryBuilder();
+			builder.AddJobFilter(jobId);
+			builder.AddTemplateFilter(templateId);
+			builder.AddRecipientFilter(recipient);
+
+			DeleteResult result = await _jobNotificationStates.DeleteManyAsync(builder.Filter, cancellationToken);
+
+			_logger.LogInformation("Removed {MessageCount} JobNotificationState(s) for jobId: {JobId}, templateId: {TemplateId}, and recipient: {Recipient}", result.DeletedCount, jobId, templateId, recipient);
+
+			return result.DeletedCount;
+		}
+
+		#endregion Job Notification
+
+		#region Job Step Notification
+
+		/// <summary>
+		/// Creates a <see cref="JobStepNotificationStateQueryBuilder"/> for creating a <see cref="JobStepNotificationStateDocument"/> query filter
+		/// </summary>
+		/// <returns>A <see cref="IJobStepNotificationStateQueryBuilder"/> builder</returns>
+		public static IJobStepNotificationStateQueryBuilder CreateJobStepNotificationStateQueryBuilder()
+		{
+			return new JobStepNotificationStateQueryBuilder();
+		}
+
+		/// <inheritdoc/>
+		public async Task<IReadOnlyList<IJobStepNotificationStateRef>?> GetJobStepNotificationStatesAsync(IJobStepNotificationStateQueryBuilder builder, CancellationToken cancellationToken = default)
+		{
+			JobStepNotificationStateQueryBuilder? jobStepNotificationStateQueryBuilder = builder as JobStepNotificationStateQueryBuilder;
+			if (jobStepNotificationStateQueryBuilder is null)
+			{
+				return null;
+			}
+
+			// Sort the results for easy processing
+			// When retrieving steps for both parent and spawned jobs, sorting ascending by PartentJobId will have the parent job steps on top of the list
+			SortDefinition<JobStepNotificationStateDocument> sort = Builders<JobStepNotificationStateDocument>.Sort.Ascending(x => x.ParentJobId).Ascending(x => x.Recipient).Ascending(x => x.Group).Ascending(x => x.TargetPlatform);
+			List<JobStepNotificationStateDocument> states = await _jobStepNotificationStates.Find(jobStepNotificationStateQueryBuilder.Filter).Sort(sort).ToListAsync(cancellationToken);
+			return (states.Count != 0) ? states : null;
+		}
+
+		/// <inheritdoc/>
+		public async Task<IJobStepNotificationStateRef> AddOrUpdateJobStepNotificationStateAsync(JobId jobId, TemplateId templateId, string recipient, JobStepId jobStepId, string group, string platform, string badge, string channel, string ts, string? threadTs = null, JobId? parentJobId = null, TemplateId? parentJobTemplateId = null, CancellationToken cancellationToken = default)
+		{
+			JobStepNotificationStateQueryBuilder filterBuilder = new JobStepNotificationStateQueryBuilder();
+			filterBuilder.AddJobFilter(jobId);
+			filterBuilder.AddTemplateFilter(templateId);
+			filterBuilder.AddRecipientFilter(recipient);
+			filterBuilder.AddJobStepFilter(jobStepId);
+			filterBuilder.AddGroupFilter(group);
+			filterBuilder.AddPlatformFilter(platform);
+
+			JobStepNotificationStateDocument? jobStepNotificationState = await _jobStepNotificationStates.Find(filterBuilder.Filter).FirstOrDefaultAsync(cancellationToken);
+			if (jobStepNotificationState == null)
+			{
+				jobStepNotificationState = new JobStepNotificationStateDocument(jobId, templateId, recipient, jobStepId, group, platform, badge, channel, ts, threadTs, parentJobId, parentJobTemplateId);
+				await _jobStepNotificationStates.InsertOneAsync(jobStepNotificationState, null, cancellationToken);
+				_logger.LogInformation("Posted message {StateId} (jobId: {JobId}, templateId: {TemplateId}, recipient: {Recipient}, jobStepId: {JobStepId}, group: {Group}, platform: {Platform}, badge: {Badge}, parentJobId: {ParentJobId}, parentJobTemplateId: {ParentJobTemplateId}, messageId: {MessageId})",
+										jobStepNotificationState.Id, jobStepNotificationState.JobId, jobStepNotificationState.TemplateId, jobStepNotificationState.Recipient, jobStepNotificationState.JobStepId, jobStepNotificationState.Group, jobStepNotificationState.TargetPlatform, jobStepNotificationState.Badge, jobStepNotificationState.ParentJobId, jobStepNotificationState.ParentJobTemplateId, jobStepNotificationState.MessageId);
+			}
+			else
+			{
+				UpdateDefinitionBuilder<JobStepNotificationStateDocument> updateBuilder = Builders<JobStepNotificationStateDocument>.Update;
+				UpdateDefinition<JobStepNotificationStateDocument> updates = updateBuilder.Set(x => x.Badge, badge).Set(x => x.Channel, channel).Set(x => x.ThreadTs, threadTs).Set(x => x.Ts, ts);
+				jobStepNotificationState = await _jobStepNotificationStates.FindOneAndUpdateAsync(x => x.Id == jobStepNotificationState.Id, updates, null, cancellationToken);
+				_logger.LogInformation("Updated message {StateId} (jobId: {JobId}, templateId: {TemplateId}, recipient: {Recipient}, jobStepId: {JobStepId}, group: {Group}, platform: {Platform}, badge: {Badge}, parentJobId: {ParentJobId}, parentJobTemplateId: {ParentJobTemplateId}, messageId: {MessageId})",
+										jobStepNotificationState.Id, jobStepNotificationState.JobId, jobStepNotificationState.TemplateId, jobStepNotificationState.Recipient, jobStepNotificationState.JobStepId, jobStepNotificationState.Group, jobStepNotificationState.TargetPlatform, jobStepNotificationState.Badge, jobStepNotificationState.ParentJobId, jobStepNotificationState.ParentJobTemplateId, jobStepNotificationState.MessageId);
+			}
+
+			return jobStepNotificationState;
+		}
+
+		/// <inheritdoc/>
+		public async Task<long> DeleteJobStepNotificationStatesAsync(JobId jobId, TemplateId templateId, string recipient, CancellationToken cancellationToken = default)
+		{
+			JobStepNotificationStateQueryBuilder filterBuilder = new JobStepNotificationStateQueryBuilder();
+			filterBuilder.AddRecipientFilter(recipient);
+			filterBuilder.AddJobAndParentJobFilter(jobId);
+			filterBuilder.AddTemplateAndParentTemplateFilter(templateId);
+
+			DeleteResult result = await _jobStepNotificationStates.DeleteManyAsync(filterBuilder.Filter, cancellationToken);
+
+			_logger.LogInformation("Removed {MessageCount} JobStepNotificationState(s) for jobId: {JobId}, templateId: {TemplateId}, and recipient: {Recipient}", result.DeletedCount, jobId, templateId, recipient);
+
+			return result.DeletedCount;
+		}
+
+		#endregion Job Step Notification
+	}
+}
diff --git a/Engine/Source/Programs/Horde/Plugins/Experimental/HordeServer.Experimental/Notifications/Sinks/ExperimentalSlackNotificationSink.cs b/Engine/Source/Programs/Horde/Plugins/Experimental/HordeServer.Experimental/Notifications/Sinks/ExperimentalSlackNotificationSink.cs
--- a/Engine/Source/Programs/Horde/Plugins/Experimental/HordeServer.Experimental/Notifications/Sinks/ExperimentalSlackNotificationSink.cs
+++ b/Engine/Source/Programs/Horde/Plugins/Experimental/HordeServer.Experimental/Notifications/Sinks/ExperimentalSlackNotificationSink.cs
@@ -1,26 +1,18 @@
 // Copyright Epic Games, Inc. All Rights Reserved.
 
-using EpicGames.Core;
 using EpicGames.Horde.Jobs;
 using EpicGames.Horde.Jobs.Graphs;
-using EpicGames.Horde.Jobs.Templates;
-using EpicGames.Horde.Logs;
-using EpicGames.Slack;
-using EpicGames.Slack.Blocks;
-using EpicGames.Slack.Elements;
 using HordeServer.Agents;
 using HordeServer.Configuration;
 using HordeServer.Devices;
+using HordeServer.Experimental.Notifications;
 using HordeServer.Issues;
+using HordeServer.Jobs;
 using HordeServer.Logs;
-using HordeServer.Server;
 using HordeServer.Streams;
 using HordeServer.Users;
 using Microsoft.Extensions.Logging;
 using Microsoft.Extensions.Options;
-using MongoDB.Bson;
-using MongoDB.Bson.Serialization.Attributes;
-using MongoDB.Driver;
 
 namespace HordeServer.Notifications.Sinks
 {
@@ -29,170 +21,31 @@ namespace HordeServer.Notifications.Sinks
 	/// </summary>
 	public sealed class ExperimentalSlackNotificationSink : INotificationSink, IDisposable
 	{
-		/// <summary>
-		/// MongoDB document for job notifications
-		/// </summary>
-		class JobNotificationStateDocument
-		{
-			[BsonId]
-			public ObjectId Id { get; set; }
-
-			[BsonElement("jid")]
-			public JobId JobId { get; set; }
-
-			[BsonElement("tid")]
-			public TemplateId TemplateId { get; set; }
-
-			[BsonElement("rcp")]
-			public string Recipient { get; set; } = String.Empty;
-
-			[BsonElement("ch")]
-			public string Channel { get; set; } = String.Empty;
-
-			[BsonElement("ts")]
-			public string Ts { get; set; } = String.Empty;
-
-			[BsonIgnore]
-			public SlackMessageId MessageId => new SlackMessageId(Channel, null, Ts);
-		}
-
-		/// <summary>
-		/// MongoDB document for job step notifications
-		/// </summary>
-		class JobStepNotificationStateDocument
-		{
-			[BsonId]
-			public ObjectId Id { get; set; }
-
-			[BsonElement("jid")]
-			public JobId JobId { get; set; }
-
-			[BsonElement("tid")]
-			public TemplateId TemplateId { get; set; }
-
-			[BsonElement("rcp")]
-			public string Recipient { get; set; } = String.Empty;
-
-			[BsonElement("sid")]
-			public JobStepId JobStepId { get; set; }
-
-			[BsonElement("grp")]
-			public string Group { get; set; } = String.Empty;
-
-			[BsonElement("plt")]
-			public string TargetPlatform { get; set; } = String.Empty;
-
-			[BsonElement("bdg")]
-			public string Badge { get; set; } = String.Empty;
-
-			[BsonElement("ch")]
-			public string Channel { get; set; } = String.Empty;
-
-			[BsonElement("tts")]
-			public string? ThreadTs { get; set; }
-
-			[BsonElement("ts")]
-			public string Ts { get; set; } = String.Empty;
-
-			[BsonIgnore]
-			public SlackMessageId MessageId => new SlackMessageId(Channel, ThreadTs, Ts);
-		}
-
-		/// <summary>
-		/// Structure for generic use of log reporting
-		/// </summary>
-		/// <remarks>Cancellation of jobs and job steps may not always include log information. In these cases we want to still provide a message to the user for the reason.</remarks>
-		struct LogData
-		{
-			public LogEventSeverity Severity { get; }
-			public string Message { get; }
-			public LogId? LogId { get; } = null;
-			public LogData(LogEventSeverity severity, string message, LogId? logId = null)
-			{
-				Severity = severity;
-				Message = message;
-				LogId = logId;
-			}
-		}
-
-		readonly IOptionsMonitor<BuildConfig> _buildConfig;
-		readonly IOptionsMonitor<ExperimentalConfig> _experimentalConfig;
-		readonly IMongoCollection<JobNotificationStateDocument> _jobNotificationStates;
-		readonly IMongoCollection<JobStepNotificationStateDocument> _jobStepNotificationStates;
-		readonly IUserCollection _userCollection;
-		readonly IServerInfo _serverInfo;
-		readonly ILogger _logger;
-
-		readonly HttpClient _httpClient;
-		readonly SlackClient _slackClient;
-
-		/// <summary>
-		/// Slack emoji used to emphasize state
-		/// </summary>
-		const string PendingBadge = ":large_blue_square:";
-		const string WarningBadge = ":large_yellow_square:";
-		const string FailureBadge = ":large_red_square:";
-		const string SkippedBadge = ":black_square:";
-		const string PassingBadge = ":large_green_square:";
-
-		/// <summary>
-		/// Slack emoji used to emphasize an aborted job
-		/// </summary>
-		const string AlertBadge = ":rotating_light:";
-
-		/// <summary>
-		/// Limit of logs to be provided in a Slack notification
-		/// </summary>
-		const int MaxEventMessages = 5;
-
-		/// <summary>
-		/// Length limit for log event to avoid character limit from Slack
-		/// </summary>
-		const int MaxLogEventLength = 2000;
-
-		/// <summary>
-		/// Block limit per Slack Context before an exception is thrown
-		/// </summary>
-		/// <remarks>See https://api.slack.com/reference/block-kit/blocks#context for detailed information</remarks>
-		const int MaxSlackBlocks = 10;
+		readonly SlackNotificationProcessor _slackNotificationProcessor;
 
 		/// <summary>
 		/// Constructor
 		/// </summary>
-		/// <param name="mongoService">MongoDB service</param>
 		/// <param name="buildServerConfig">Options from the build plugin's server config</param>
 		/// <param name="buildConfig">Monitored options from the build plugin's global config</param>
 		/// <param name="experimentalConfig">Monitored options from the experimental plugin's global config</param>
+		/// <param name="jobNotificationCollection">Collection of job and job step notifications</param>
+		/// <param name="jobCollection">Collection of jobs</param>
 		/// <param name="userCollection">Collection of users</param>
 		/// <param name="serverInfo">Information of the current server instance</param>
 		/// <param name="logger">Logger for output</param>
-		public ExperimentalSlackNotificationSink(IMongoService mongoService,
-			IOptions<BuildServerConfig> buildServerConfig,
+		public ExperimentalSlackNotificationSink(IOptions<BuildServerConfig> buildServerConfig,
 			IOptionsMonitor<BuildConfig> buildConfig,
 			IOptionsMonitor<ExperimentalConfig> experimentalConfig,
+			IJobNotificationCollection jobNotificationCollection,
+			IJobCollection jobCollection,
 			IUserCollection userCollection,
 			IServerInfo serverInfo,
 			ILogger<ExperimentalSlackNotificationSink> logger)
 		{
-			_buildConfig = buildConfig;
-			_experimentalConfig = experimentalConfig;
-			_jobNotificationStates = mongoService.GetCollection<JobNotificationStateDocument>("SlackExperimentalJob", keys => keys.Ascending(x => x.JobId).Ascending(x => x.TemplateId).Ascending(x => x.Recipient), unique: true);
-			_jobStepNotificationStates = mongoService.GetCollection<JobStepNotificationStateDocument>("SlackExperimentalJobStep", keys => keys.Ascending(x => x.JobId).Ascending(x => x.TemplateId).Ascending(x => x.Recipient).Ascending(x => x.Group).Ascending(x => x.JobStepId), unique: true);
-			_userCollection = userCollection;
-			_serverInfo = serverInfo;
-			_logger = logger;
-
-			_httpClient = new HttpClient();
-			_httpClient.DefaultRequestHeaders.Add("Authorization", $"Bearer {buildServerConfig.Value.SlackToken ?? ""}");
-			_slackClient = new SlackClient(_httpClient, _logger);
+			_slackNotificationProcessor = new SlackNotificationProcessor(buildServerConfig, buildConfig, experimentalConfig, jobNotificationCollection, jobCollection, userCollection, serverInfo, logger);
 		}
 
-		/// <inheritdoc/>
-		public void Dispose()
-			=> _httpClient.Dispose();
-
-		#region INotificationSink Methods
-
 		/// <inheritdoc/>
 		public Task NotifyConfigUpdateAsync(ConfigUpdateInfo info, CancellationToken cancellationToken)
 			=> Task.CompletedTask;
@@ -212,47 +65,7 @@ namespace HordeServer.Notifications.Sinks
 		/// <inheritdoc/>
 		public async Task NotifyJobCompleteAsync(IJob job, IGraph graph, LabelOutcome outcome, CancellationToken cancellationToken)
 		{
-			// Only handle the job complete notification if the job was canceled
-			// Otherwise we will leave it to the steps to handle when all notifications have been sent out to avoid conditions where the job complete notification gets received before the job step complete notification
-			if (job.AbortedByUserId is null)
-			{
-				return;
-			}
-
-			// Send the updated summary after the steps has been updated
-			List<JobNotificationStateDocument>? jobNotificationStates = await GetJobNotificationStatesAsync(job.Id, cancellationToken);
-			if (jobNotificationStates is null)
-			{
-				_logger.LogInformation("No job notifications found for job id '{JobId}'.", job.Id);
-				return;
-			}
-
-			string hordeJobURL = new Uri(_serverInfo.DashboardUrl, $"job/{job.Id}").ToString();
-
-			string branch = "Unknown Branch";
-			if (_buildConfig.CurrentValue.TryGetStream(job.StreamId, out StreamConfig? config))
-			{
-				branch = config.Name;
-			}
-
-			string abortedByUserName = "Unknown User";
-			IUser? user = await _userCollection.GetUserAsync(job.AbortedByUserId.Value, cancellationToken);
-			if (user is not null)
-			{
-				abortedByUserName = user.Name;
-			}
-
-			foreach (JobNotificationStateDocument jobNotificationState in jobNotificationStates)
-			{
-				// Update our steps to reflect the canceled job
-				_ = await UpdatePendingJobStepNotificationsAsync(job, jobNotificationState, cancellationToken);
-
-				await SendSlackSummaryMessageAsync(job, branch, hordeJobURL, jobNotificationState.Recipient, abortedByUserName, cancellationToken);
-
-				// No further notifications will be sent out and can be cleaned up
-				_ = await DeleteJobStepNotificationStatesAsync(jobNotificationState.JobId, jobNotificationState.TemplateId, jobNotificationState.Recipient, cancellationToken);
-				_ = await DeleteJobNotificationStatesAsync(jobNotificationState.JobId, jobNotificationState.TemplateId, jobNotificationState.Recipient, cancellationToken);
-			}
+			await _slackNotificationProcessor.ProcessJobCompleteAsync(job, cancellationToken);
 		}
 
 		/// <inheritdoc/>
@@ -266,24 +79,7 @@ namespace HordeServer.Notifications.Sinks
 		/// <inheritdoc/>
 		public async Task NotifyJobStepAbortedAsync(IEnumerable<IUser>? usersToNotify, IJob job, IJobStepBatch batch, IJobStep step, INode node, List<ILogEventData> jobStepEventData, CancellationToken cancellationToken)
 		{
-			List<LogData> logEvents = new List<LogData>();
-			string cancellationDetails = "Step was canceled.";
-			if (step.AbortedByUserId is not null)
-			{
-				IUser? user = await _userCollection.GetUserAsync(step.AbortedByUserId.Value, cancellationToken);
-				if (user is not null)
-				{
-					cancellationDetails = $"Step was canceled by {user.Name}.";
-				}
-			}
-
-			if (!String.IsNullOrEmpty(step.CancellationReason))
-			{
-				cancellationDetails += $" Reason for cancellation: '{step.CancellationReason}'.";
-			}
-			logEvents.Add(new LogData(LogEventSeverity.Error, cancellationDetails));
-
-			await ProcessJobStepCompleteNotificationAsync(job, step, logEvents, cancellationToken);
+			await _slackNotificationProcessor.ProcessJobStepAbortedAsync(job, step, cancellationToken);
 		}
 
 		/// <inheritdoc/>
@@ -295,8 +91,7 @@ namespace HordeServer.Notifications.Sinks
 				return;
 			}
 
-			List<LogData> logEvents = ExtractLogInformation(step, jobStepEventData);
-			await ProcessJobStepCompleteNotificationAsync(job, step, logEvents, cancellationToken);
+			await _slackNotificationProcessor.ProcessJobStepCompleteAsync(job, step, jobStepEventData, cancellationToken);
 		}
 
 		/// <inheritdoc/>
@@ -315,1019 +110,8 @@ namespace HordeServer.Notifications.Sinks
 		public Task SendIssueReportAsync(IssueReportGroup report, CancellationToken cancellationToken)
 			=> Task.CompletedTask;
 
-		#endregion INotificationSink Methods
-
-		#region Job Processing
-
-		/// <summary>
-		/// Search the job for steps to be included in the notifications based on the experimental plugin's global config
-		/// </summary>
-		/// <param name="job">Information about the job</param>
-		/// <param name="config">Configuration to fetch notification information</param>
-		/// <param name="cancellationToken">Cancellation token for the operation</param>
-		async Task ProcessJobAsync(IJob job, JobNotificationConfig config, CancellationToken cancellationToken = default)
-		{
-			// Populate our job information needed for the notification
-			try
-			{
-				foreach (IJobStepBatch batch in job.Batches)
-				{
-					foreach (IJobStep step in batch.Steps)
-					{
-						if (config.GetStepDetails(step.Name, out string? group, out string? platform) && !String.IsNullOrEmpty(group) && !String.IsNullOrEmpty(platform))
-						{
-							string badge = GetBadgeForStep(step);
-							foreach (string recipient in config.Channels)
-							{
-								await AddOrUpdateJobStepNotificationStateAsync(job.Id, job.TemplateId, recipient, step.Id, group, platform, badge, null, cancellationToken);
-							}
-						}
-					}
-				}
-			}
-			catch (Exception ex)
-			{
-				_logger.LogInformation(ex, "Failed to extract step details for '{Job}' in stream '{StreamId}' and for template '{TemplateId}' ({Error})", job.Id, job.StreamId, job.TemplateId, ex.Message);
-			}
-		}
-
-		/// <summary>
-		/// Prepares a job summary and job step thread to be sent to Slack based on the experimental plugin's global configuration
-		/// </summary>
-		/// <param name="job">Information about the job</param>
-		/// <param name="step">Information about the step</param>
-		/// <param name="jobStepEventData">List of the step's logged events</param>
-		/// <param name="cancellationToken">Cancellation token for the operation</param>
-		async Task ProcessJobStepCompleteNotificationAsync(IJob job, IJobStep step, List<LogData> jobStepEventData, CancellationToken cancellationToken = default)
-		{
-			ExperimentalConfig currentExperimentalConfig = _experimentalConfig.CurrentValue;
-
-			HashSet<JobNotificationConfig>? jobNotificationConfigs;
-			if (!currentExperimentalConfig.TryGetJobNofications(job.StreamId, job.TemplateId, out jobNotificationConfigs))
-			{
-				_logger.LogInformation("No config found for '{JobName}' in stream '{StreamId}' and for template '{TemplateId}'", job.Name, job.StreamId, job.TemplateId);
-				return;
-			}
-
-			string hordeJobURL = new Uri(_serverInfo.DashboardUrl, $"job/{job.Id}").ToString();
-
-			string branch = "Unknown Branch";
-			if (_buildConfig.CurrentValue.TryGetStream(job.StreamId, out StreamConfig? streamConfig))
-			{
-				branch = streamConfig.Name;
-			}
-
-			// Check if the step completed was used to generate the details of the job so that the initial notification can be sent out
-			if (step.Name == IJob.SetupNodeName)
-			{
-				// Check if the Setup step failed to notify that the job can't be executed
-				bool wasSetupSuccess = (step.Outcome == JobStepOutcome.Success || step.Outcome == JobStepOutcome.Warnings);
-				if (!wasSetupSuccess)
-				{
-					await SendSetupFailedSlackMessageAsync(job, step, branch, jobStepEventData, jobNotificationConfigs, cancellationToken);
-					return;
-				}
-
-				foreach (JobNotificationConfig config in jobNotificationConfigs)
-				{
-					await ProcessJobAsync(job, config, cancellationToken);
-					foreach (string recipient in config.Channels)
-					{
-						await SendSlackSummaryMessageAsync(job, branch, hordeJobURL, recipient, null, cancellationToken);
-					}
-				}
-				return;
-			}
-
-			// Send a message for our current step
-			bool needsSummaryUpdate = await UpdateAndSendSlackNotificationsForStepAsync(job.Id, step, jobStepEventData, null, cancellationToken);
-
-			// Fetch all dependent steps if there was an error and send notifications with this step as the cause for the dependent steps being skipped
-			if (step.Outcome != JobStepOutcome.Success && step.Outcome != JobStepOutcome.Warnings)
-			{
-				HashSet<JobStepId>? dependentSteps = await FindDependentJobStepIdsAsync(job, step.Id, cancellationToken);
-				if (dependentSteps is not null)
-				{
-					foreach (JobStepId stepId in dependentSteps)
-					{
-						// Send a message for our dependent step
-						if (job.TryGetStep(stepId, out IJobStep? dependentStep))
-						{
-							needsSummaryUpdate |= await UpdateAndSendSlackNotificationsForStepAsync(job.Id, dependentStep, jobStepEventData, step, cancellationToken);
-						}
-					}
-				}
-			}
-
-			// Check if the job is marked as being completed and cleanup any steps that we fell out of sync with
-			List<JobNotificationStateDocument>? jobNotificationStates = await GetJobNotificationStatesAsync(job.Id, cancellationToken);
-			if (jobNotificationStates is not null && job.GetState() == JobState.Complete)
-			{
-				foreach (JobNotificationStateDocument jobNotificationState in jobNotificationStates)
-				{
-					needsSummaryUpdate |= await UpdatePendingJobStepNotificationsAsync(job, jobNotificationState, cancellationToken);
-				}
-			}
-
-			// Check to see if any of our steps were updated which requires the Slack summary to be updated
-			if (!needsSummaryUpdate)
-			{
-				return;
-			}
-
-			// Send the updated summary after the steps has been updated
-			if (jobNotificationStates is null)
-			{
-				_logger.LogInformation("No job notifications found for job id '{JobId}' even though job step '{StepId}' was updated.", job.Id, step.Id);
-				return;
-			}
-
-			foreach (JobNotificationStateDocument jobNotificationState in jobNotificationStates)
-			{
-				await SendSlackSummaryMessageAsync(job, branch, hordeJobURL, jobNotificationState.Recipient, null, cancellationToken);
-
-				// Check to see if all notifications for this job were sent to then delete all documents related to this job as no other notifications will be sent
-				await TryNotificationCleanupAsync(jobNotificationState, cancellationToken);
-			}
-		}
-
-		/// <summary>
-		/// Finds all dependent steps of the <see cref="JobStepNotificationStateDocument"/>
-		/// </summary>
-		/// <param name="job">Information about the job</param>
-		/// <param name="underlyingStepId">Identifier of the possible job step dependency</param>
-		/// <param name="cancellationToken">Cancellation token for the operation</param>
-		/// <returns></returns>
-		async Task<HashSet<JobStepId>?> FindDependentJobStepIdsAsync(IJob job, JobStepId underlyingStepId, CancellationToken cancellationToken = default)
-		{
-			// Fetch all of our job steps sending notifications
-			FilterDefinition<JobStepNotificationStateDocument> filter = Builders<JobStepNotificationStateDocument>.Filter.Eq(x => x.JobId, job.Id);
-			List<JobStepNotificationStateDocument>? jobStepNotificationStateDocuments = await GetJobStepNotificationStatesAsync(filter, cancellationToken);
-			if (jobStepNotificationStateDocuments is null)
-			{
-				return null;
-			}
-
-			HashSet<JobStepId> dependencyJobStepNotifications = new HashSet<JobStepId>();
-			HashSet<JobStepId> cachedSteps = new HashSet<JobStepId>();
-			foreach (JobStepNotificationStateDocument jobStepNotificationState in jobStepNotificationStateDocuments)
-			{
-				if (cachedSteps.Contains(jobStepNotificationState.JobStepId))
-				{
-					continue;
-				}
-				cachedSteps.Add(jobStepNotificationState.JobStepId);
-
-				List<JobStepId> dependencySteps = new List<JobStepId>()
-				{
-					jobStepNotificationState.JobStepId
-				};
-
-				HashSet<JobStepId> visitedDependencySteps = new HashSet<JobStepId>()
-				{
-					jobStepNotificationState.JobStepId
-				};
-
-				bool isMatchingDependency = false;
-				while (dependencySteps.Count > 0 && !isMatchingDependency)
-				{
-					JobStepId jobStepId = dependencySteps[0];
-					dependencySteps.RemoveAt(0);
-
-					IJobStep? currentStep;
-					if (!job.TryGetStep(jobStepId, out currentStep))
-					{
-						continue;
-					}
-
-					// Iterate through all of the step's dependencies
-					// Add the step Id of the notification state if the dependencies match
-					foreach (JobStepId dependencyStepId in currentStep.InputDependencies)
-					{
-						if (underlyingStepId == dependencyStepId)
-						{
-							isMatchingDependency = true;
-							dependencyJobStepNotifications.Add(jobStepNotificationState.JobStepId);
-							break;
-						}
-						else if (!visitedDependencySteps.Contains(dependencyStepId))
-						{
-							dependencySteps.Add(dependencyStepId);
-							visitedDependencySteps.Add(dependencyStepId);
-						}
-					}
-
-					// Check to see if we already found a potential match before continuing searching through dependencies
-					if (isMatchingDependency)
-					{
-						continue;
-					}
-
-					foreach (JobStepId dependencyStepId in currentStep.OrderDependencies)
-					{
-						if (underlyingStepId == dependencyStepId)
-						{
-							isMatchingDependency = true;
-							dependencyJobStepNotifications.Add(jobStepNotificationState.JobStepId);
-							break;
-						}
-						else if (!visitedDependencySteps.Contains(dependencyStepId))
-						{
-							dependencySteps.Add(dependencyStepId);
-							visitedDependencySteps.Add(dependencyStepId);
-						}
-					}
-				}
-			}
-
-			return dependencyJobStepNotifications;
-		}
-
-		/// <summary>
-		/// Syncs up the job information with our current recorded steps
-		/// </summary>
-		/// <param name="job">Information about the job</param>
-		/// <param name="jobNotificationState">Job Notification State to check against</param>
-		/// <param name="cancellationToken">Cancellation token for the operation</param>
-		/// <returns></returns>
-		async Task<bool> UpdatePendingJobStepNotificationsAsync(IJob job, JobNotificationStateDocument jobNotificationState, CancellationToken cancellationToken = default)
-		{
-			bool wasUpdated = false;
-			// Fetch all of our job steps that are still pending to determine if we still have notifications to send
-			FilterDefinition<JobStepNotificationStateDocument> filter = Builders<JobStepNotificationStateDocument>.Filter.Eq(x => x.JobId, jobNotificationState.JobId) &
-																			Builders<JobStepNotificationStateDocument>.Filter.Eq(x => x.TemplateId, jobNotificationState.TemplateId) &
-																			Builders<JobStepNotificationStateDocument>.Filter.Eq(x => x.Recipient, jobNotificationState.Recipient) &
-																			Builders<JobStepNotificationStateDocument>.Filter.Eq(x => x.Badge, PendingBadge);
-
-			List<JobStepNotificationStateDocument>? jobStepNotificationStates = await GetJobStepNotificationStatesAsync(filter, cancellationToken);
-			if (jobStepNotificationStates is null)
-			{
-				return false;
-			}
-
-			foreach (JobStepNotificationStateDocument jobStepNotificationState in jobStepNotificationStates)
-			{
-				IJobStep? step = null;
-				if (!job.TryGetStep(jobStepNotificationState.JobStepId, out step))
-				{
-					continue;
-				}
-
-				string badge = GetBadgeForStep(step);
-				if (badge != jobStepNotificationState.Badge)
-				{
-					_ = await AddOrUpdateJobStepNotificationStateAsync(jobStepNotificationState.JobId, jobStepNotificationState.TemplateId, jobStepNotificationState.Recipient, jobStepNotificationState.JobStepId,
-																		jobStepNotificationState.Group, jobStepNotificationState.TargetPlatform, badge, null, cancellationToken);
-					wasUpdated = true;
-				}
-			}
-
-			return wasUpdated;
-		}
-
-		/// <summary>
-		/// Clean up notification states once all notifications have been sent
-		/// </summary>
-		/// <param name="jobNotificationState">Job Notification State to check against</param>
-		/// <param name="cancellationToken">Cancellation token for the operation</param>
-		/// <returns></returns>
-		async Task TryNotificationCleanupAsync(JobNotificationStateDocument jobNotificationState, CancellationToken cancellationToken = default)
-		{
-			// Fetch all of our job steps that are still pending to determine if we still have notifications to send
-			FilterDefinition<JobStepNotificationStateDocument> filter = Builders<JobStepNotificationStateDocument>.Filter.Eq(x => x.JobId, jobNotificationState.JobId) &
-																			Builders<JobStepNotificationStateDocument>.Filter.Eq(x => x.TemplateId, jobNotificationState.TemplateId) &
-																			Builders<JobStepNotificationStateDocument>.Filter.Eq(x => x.Recipient, jobNotificationState.Recipient) &
-																			Builders<JobStepNotificationStateDocument>.Filter.Eq(x => x.Badge, PendingBadge);
-
-			List<JobStepNotificationStateDocument>? jobStepNotificationStates = await GetJobStepNotificationStatesAsync(filter, cancellationToken);
-			if (jobStepNotificationStates is null)
-			{
-				_ = await DeleteJobStepNotificationStatesAsync(jobNotificationState.JobId, jobNotificationState.TemplateId, jobNotificationState.Recipient, cancellationToken);
-				_ = await DeleteJobNotificationStatesAsync(jobNotificationState.JobId, jobNotificationState.TemplateId, jobNotificationState.Recipient, cancellationToken);
-			}
-		}
-
-		/// <summary>
-		/// Determines the badge based on the job step's state and outcome
-		/// </summary>
-		/// <param name="step">Information about the step</param>
-		/// <returns>String representation of the badge</returns>
-		static string GetBadgeForStep(IJobStep step)
-		{
-			// Check the state
-			switch (step.State)
-			{
-				case JobStepState.Waiting:
-				case JobStepState.Ready:
-				case JobStepState.Running:
-					{
-						// Step is either queued or currently being executed
-						return PendingBadge;
-					}
-				case JobStepState.Skipped:
-					{
-						return SkippedBadge;
-					}
-				case JobStepState.Aborted:
-					{
-						// Treat aborted as failure conditions
-						return FailureBadge;
-					}
-				case JobStepState.Completed:
-					{
-						// Continue to check the outcome of the job
-						break;
-					}
-				default:
-					{
-						// Log an error "Job Step in an Unspecified State"
-						return String.Empty;
-					}
-			}
-
-			// Check the outcome
-			switch (step.Outcome)
-			{
-				case JobStepOutcome.Failure:
-					{
-						return FailureBadge;
-					}
-				case JobStepOutcome.Warnings:
-					{
-						return WarningBadge;
-					}
-				case JobStepOutcome.Success:
-					{
-						return PassingBadge;
-					}
-				default:
-					{
-						// Log an error "Job Step in an Unspecified Outcome"
-						return String.Empty;
-					}
-			}
-		}
-
-		/// <summary>
-		/// Parses list of <see cref="ILogEventData"/> into the generic <see cref="LogData"/> structure
-		/// </summary>
-		/// <param name="step">Information about the step</param>
-		/// <param name="jobStepEventData">List of the step's logged events</param>
-		/// <returns>List of parsed logged events to be consumed by the Slack notifications</returns>
-		static List<LogData> ExtractLogInformation(IJobStep step, List<ILogEventData> jobStepEventData)
-		{
-			List<LogData> logData = new List<LogData>();
-			foreach (ILogEventData eventData in jobStepEventData)
-			{
-				logData.Add(new LogData(eventData.Severity, eventData.Message, step.LogId));
-			}
-			return logData;
-		}
-
-		#endregion Job Processing
-
-		#region MongoDB
-
-		/// <summary>
-		/// Add or update the <see cref="JobNotificationStateDocument"/> to the collection
-		/// </summary>
-		/// <param name="jobId">Identifier of the job</param>
-		/// <param name="templateId">Identifier of the stream's template</param>
-		/// <param name="recipient">Identifier of the channel or user in Slack or user's email address</param>
-		/// <param name="messageId">Slack message identifer of the sent notification</param>
-		/// <param name="cancellationToken">Cancellation token for the operation</param>
-		/// <returns>A tuple of the <see cref="JobNotificationStateDocument"/> and a bool used to determine if the record is considered to be new</returns>
-		async Task<(JobNotificationStateDocument, bool)> AddOrUpdateJobNotificationStateAsync(JobId jobId, TemplateId templateId, string recipient, SlackMessageId messageId, CancellationToken cancellationToken = default)
-		{
-			ObjectId newId = ObjectId.GenerateNewId();
-
-			FilterDefinition<JobNotificationStateDocument> filter = Builders<JobNotificationStateDocument>.Filter.Eq(x => x.JobId, jobId) &
-																		Builders<JobNotificationStateDocument>.Filter.Eq(x => x.TemplateId, templateId) &
-																		Builders<JobNotificationStateDocument>.Filter.Eq(x => x.Recipient, recipient);
-			UpdateDefinition<JobNotificationStateDocument> update = Builders<JobNotificationStateDocument>.Update.SetOnInsert(x => x.Id, newId).Set(x => x.Channel, messageId.Channel).Set(x => x.Ts, messageId.Ts);
-
-			JobNotificationStateDocument state = await _jobNotificationStates.FindOneAndUpdateAsync(filter, update, new FindOneAndUpdateOptions<JobNotificationStateDocument> { IsUpsert = true, ReturnDocument = ReturnDocument.After }, cancellationToken);
-			if (state.Id == newId)
-			{
-				_logger.LogInformation("Posted message {StateId} (jobId: {JobId}, templateId: {TemplateId}, recipient: {Recipient}, messageId: {MessageId})", state.Id, state.JobId, state.TemplateId, state.Recipient, state.MessageId);
-			}
-			else
-			{
-				_logger.LogInformation("Updated message {StateId} (jobId: {JobId}, templateId: {TemplateId}, recipient: {Recipient}, messageId: {MessageId})", state.Id, state.JobId, state.TemplateId, state.Recipient, state.MessageId);
-			}
-
-			return (state, state.Id == newId);
-		}
-
-		/// <summary>
-		/// Fetches a list of <see cref="JobNotificationStateDocument"/>
-		/// </summary>
-		/// <param name="jobId">Identifier of the job</param>
-		/// <param name="cancellationToken">Cancellation token for the operation</param>
-		/// <returns>A list of <see cref="JobNotificationStateDocument"/> if records were found, otherwise returns null for error checking</returns>
-		async Task<List<JobNotificationStateDocument>?> GetJobNotificationStatesAsync(JobId jobId, CancellationToken cancellationToken = default)
-		{
-			FilterDefinition<JobNotificationStateDocument> filter = Builders<JobNotificationStateDocument>.Filter.Eq(x => x.JobId, jobId);
-			List<JobNotificationStateDocument> states = await _jobNotificationStates.Find(filter).ToListAsync(cancellationToken);
-			if (states.Count == 0)
-			{
-				return null;
-			}
-
-			return states;
-		}
-
-		/// <summary>
-		/// Fetches a single entity of <see cref="JobNotificationStateDocument"/>
-		/// </summary>
-		/// <param name="jobId">Identifier of the job</param>
-		/// <param name="templateId">Identifier of the stream's template</param>
-		/// <param name="recipient">Identifier of the channel or user in Slack or user's email address</param>
-		/// <param name="cancellationToken">Cancellation token for the operation</param>
-		/// <returns>A <see cref="JobNotificationStateDocument"/> if records were found, otherwise returns null for error checking</returns>
-		/// <remarks>Only a single <see cref="JobNotificationStateDocument"/> should exist for the aggregate of the job identifier, template identifier, and recipient should ever exist. Will report an error if another record had been found.</remarks>
-		async Task<JobNotificationStateDocument?> GetJobNotificationStateAsync(JobId jobId, TemplateId templateId, string recipient, CancellationToken cancellationToken = default)
-		{
-			FilterDefinition<JobNotificationStateDocument> filter = Builders<JobNotificationStateDocument>.Filter.Eq(x => x.JobId, jobId) &
-																		Builders<JobNotificationStateDocument>.Filter.Eq(x => x.TemplateId, templateId) &
-																		Builders<JobNotificationStateDocument>.Filter.Eq(x => x.Recipient, recipient);
-			List<JobNotificationStateDocument> states = await _jobNotificationStates.Find(filter).ToListAsync(cancellationToken);
-			if (states.Count == 0)
-			{
-				return null;
-			}
-
-			if (states.Count != 1)
-			{
-				// Log error that multiple were found and could cause issues
-				_logger.LogInformation("Found {Count} records when fetching for {JobId}:{TemplateId}:{Recipient}. Using first record which could cause issues.", states.Count, jobId, templateId, recipient);
-			}
-
-			return states[0];
-		}
-
-		/// <summary>
-		/// Deletes <see cref="JobNotificationStateDocument"/> from the collection
-		/// </summary>
-		/// <param name="jobId">Identifier of the job</param>
-		/// <param name="templateId">Identifier of the stream's template</param>
-		/// <param name="recipient">Identifier of the channel or user in Slack or user's email address</param>
-		/// <param name="cancellationToken">Cancellation token for the operation</param>
-		/// <returns>True if records were deleted, otherwise false</returns>
-		async Task<bool> DeleteJobNotificationStatesAsync(JobId jobId, TemplateId templateId, string recipient, CancellationToken cancellationToken = default)
-		{
-			FilterDefinition<JobNotificationStateDocument> filter = Builders<JobNotificationStateDocument>.Filter.Eq(x => x.JobId, jobId) &
-																		Builders<JobNotificationStateDocument>.Filter.Eq(x => x.TemplateId, templateId) &
-																		Builders<JobNotificationStateDocument>.Filter.Eq(x => x.Recipient, recipient);
-			DeleteResult result = await _jobNotificationStates.DeleteManyAsync(filter, cancellationToken);
-
-			_logger.LogInformation("Removed {MessageCount} JobNotificationState(s) for jobId: {JobId}, templateId: {TemplateId}, and recipient: {Recipient}", result.DeletedCount, jobId, templateId, recipient);
-
-			return result.DeletedCount > 0;
-		}
-
-		/// <summary>
-		/// Add or update the <see cref="JobStepNotificationStateDocument"/> to the collection
-		/// </summary>
-		/// <param name="jobId">Identifier of the job</param>
-		/// <param name="templateId">Identifier of the stream's template</param>
-		/// <param name="recipient">Identifier of the channel or user in Slack or user's email address</param>
-		/// <param name="stepId">Identifier of the job step</param>
-		/// <param name="group">Group associated with the notification</param>
-		/// <param name="platform">Target platform associated with the notification</param>
-		/// <param name="badge">Slack emoji associated with the notification</param>
-		/// <param name="messageId">Optional Slack message identifer if the notification has been sent out</param>
-		/// <param name="cancellationToken">Cancellation token for the operation</param>
-		/// <returns>A tuple of the <see cref="JobStepNotificationStateDocument"/> and a bool used to determine if the record is considered to be new</returns>
-		async Task<(JobStepNotificationStateDocument, bool)> AddOrUpdateJobStepNotificationStateAsync(JobId jobId, TemplateId templateId, string recipient, JobStepId stepId, string group, string platform, string badge, SlackMessageId? messageId, CancellationToken cancellationToken = default)
-		{
-			ObjectId newId = ObjectId.GenerateNewId();
-
-			FilterDefinition<JobStepNotificationStateDocument> filter = Builders<JobStepNotificationStateDocument>.Filter.Eq(x => x.JobId, jobId) &
-																		Builders<JobStepNotificationStateDocument>.Filter.Eq(x => x.TemplateId, templateId) &
-																		Builders<JobStepNotificationStateDocument>.Filter.Eq(x => x.Recipient, recipient) &
-																		Builders<JobStepNotificationStateDocument>.Filter.Eq(x => x.JobStepId, stepId) &
-																		Builders<JobStepNotificationStateDocument>.Filter.Eq(x => x.Group, group) &
-																		Builders<JobStepNotificationStateDocument>.Filter.Eq(x => x.TargetPlatform, platform);
-			UpdateDefinition<JobStepNotificationStateDocument> update = Builders<JobStepNotificationStateDocument>.Update.SetOnInsert(x => x.Id, newId).Set(x => x.Badge, badge);
-
-			if (messageId is not null)
-			{
-				update = update.Set(x => x.Channel, messageId.Channel).Set(x => x.ThreadTs, messageId.ThreadTs).Set(x => x.Ts, messageId.Ts);
-			}
-
-			JobStepNotificationStateDocument state = await _jobStepNotificationStates.FindOneAndUpdateAsync(filter, update, new FindOneAndUpdateOptions<JobStepNotificationStateDocument> { IsUpsert = true, ReturnDocument = ReturnDocument.After }, cancellationToken);
-			if (state.Id == newId)
-			{
-				_logger.LogInformation("Posted message {StateId} (jobId: {JobId}, templateId: {TemplateId}, recipient: {Recipient}, jobStepId: {JobStepId}, group: {Group}, platform: {Platform}, badge: {Badge}, messageId: {MessageId})", state.Id, state.JobId, state.TemplateId, state.Recipient, state.JobStepId, state.Group, state.TargetPlatform, state.Badge, state.MessageId);
-			}
-			else
-			{
-				_logger.LogInformation("Updated message {StateId} (jobId: {JobId}, templateId: {TemplateId}, recipient: {Recipient}, jobStepId: {JobStepId}, group: {Group}, platform: {Platform}, badge: {Badge}, messageId: {MessageId})", state.Id, state.JobId, state.TemplateId, state.Recipient, state.JobStepId, state.Group, state.TargetPlatform, state.Badge, state.MessageId);
-			}
-
-			return (state, state.Id == newId);
-		}
-
-		/// <summary>
-		/// Fetches a list of <see cref="JobStepNotificationStateDocument"/>
-		/// </summary>
-		/// <param name="filter">Filter used to query against the collection</param>
-		/// <param name="cancellationToken">Cancellation token for the operation</param>
-		/// <returns>A list of <see cref="JobStepNotificationStateDocument"/> if records were found, otherwise returns null for error checking</returns>
-		async Task<List<JobStepNotificationStateDocument>?> GetJobStepNotificationStatesAsync(FilterDefinition<JobStepNotificationStateDocument> filter, CancellationToken cancellationToken = default)
-		{
-			List<JobStepNotificationStateDocument> states = await _jobStepNotificationStates.Find(filter).SortBy(x => x.JobId).SortBy(x => x.TemplateId).SortBy(x => x.Recipient).SortBy(x => x.Group).ToListAsync(cancellationToken);
-			if (states.Count == 0)
-			{
-				return null;
-			}
-
-			return states;
-		}
-
-		/// <summary>
-		/// Deletes <see cref="JobStepNotificationStateDocument"/> from the collection
-		/// </summary>
-		/// <param name="jobId">Identifier of the job</param>
-		/// <param name="templateId">Identifier of the stream's template</param>
-		/// <param name="recipient">Identifier of the channel or user in Slack or user's email address</param>
-		/// <param name="cancellationToken">Cancellation token for the operation</param>
-		/// <returns>True if records were deleted, otherwise false</returns>
-		async Task<bool> DeleteJobStepNotificationStatesAsync(JobId jobId, TemplateId templateId, string recipient, CancellationToken cancellationToken = default)
-		{
-			FilterDefinition<JobStepNotificationStateDocument> filter = Builders<JobStepNotificationStateDocument>.Filter.Eq(x => x.JobId, jobId) &
-																		Builders<JobStepNotificationStateDocument>.Filter.Eq(x => x.TemplateId, templateId) &
-																		Builders<JobStepNotificationStateDocument>.Filter.Eq(x => x.Recipient, recipient);
-			DeleteResult result = await _jobStepNotificationStates.DeleteManyAsync(filter, cancellationToken);
-
-			_logger.LogInformation("Removed {MessageCount} JobStepNotificationState(s) for jobId: {JobId}, templateId: {TemplateId}, and recipient: {Recipient}", result.DeletedCount, jobId, templateId, recipient);
-
-			return result.DeletedCount > 0;
-		}
-
-		#endregion MongoDB
-
-		#region Slack
-
-		/// <summary>
-		/// Create the job's summary message
-		/// </summary>
-		/// <param name="job">Information about the job</param>
-		/// <param name="branch">Name of the branch from the stream</param>
-		/// <param name="hordeJobURL">URL of the job in Horde</param>
-		/// <param name="recipient">Identifier of the channel or user in Slack or user's email address</param>
-		/// <param name="abortedByUserName">Optional username provided when cancelling a job</param>
-		/// <param name="cancellationToken">Cancellation token for the operation</param>
-		/// <returns>Generated <see cref="SlackMessage"/> if successful, otherwise returns null</returns>
-		async Task<SlackMessage?> GenerateSlackSummaryMessageAsync(IJob job, string branch, string hordeJobURL, string recipient, string? abortedByUserName = null, CancellationToken cancellationToken = default)
-		{
-			// Fetch all of our job steps needed to send out the summary
-			FilterDefinition<JobStepNotificationStateDocument> filter = Builders<JobStepNotificationStateDocument>.Filter.Eq(x => x.JobId, job.Id) &
-																		Builders<JobStepNotificationStateDocument>.Filter.Eq(x => x.TemplateId, job.TemplateId) &
-																		Builders<JobStepNotificationStateDocument>.Filter.Eq(x => x.Recipient, recipient);
-			List<JobStepNotificationStateDocument>? jobStepNotificationStateDocuments = await GetJobStepNotificationStatesAsync(filter, cancellationToken);
-			if (jobStepNotificationStateDocuments is null)
-			{
-				return null;
-			}
-
-			return CreateSlackSummaryNotification(job, branch, hordeJobURL, jobStepNotificationStateDocuments, abortedByUserName);
-		}
-
-		/// <summary>
-		/// Send a notification on a job setup failure
-		/// </summary>
-		/// <param name="job">Information about the job</param>
-		/// <param name="step">Information about the step</param>
-		/// <param name="branch">Name of the branch from the stream</param>
-		/// <param name="jobStepEventData">List of the step's logged events</param>
-		/// <param name="jobNotificationConfigs">Collection of configurations tied to this job</param>
-		/// <param name="cancellationToken">Cancellation token for the operation</param>
-		/// <returns></returns>
-		async Task SendSetupFailedSlackMessageAsync(IJob job, IJobStep step, string branch, List<LogData> jobStepEventData, HashSet<JobNotificationConfig> jobNotificationConfigs, CancellationToken cancellationToken = default)
-		{
-			SlackMessage message = new SlackMessage();
-
-			string jobUrlLink = new Uri($"{_serverInfo.DashboardUrl}job/{job.Id}").ToString();
-			string jobStepLink = new Uri($"{jobUrlLink}?step={step.Id}").ToString();
-
-			// Create our notification text
-			string notificationText = GenerateSlackMessageTitle(job, branch);
-			message.Text = $"{notificationText} Details...";
-
-			// Create a Block for the Slack Title
-			string messageTitle = GenerateSlackMessageTitle(job, branch, jobUrlLink);
-			ContextBlock jobContext = new ContextBlock();
-			jobContext.Elements.Add(new TextObject(messageTitle));
-			message.Blocks.Add(jobContext);
-
-			ContextBlock infoBlock = new ContextBlock();
-			infoBlock.Elements.Add(new TextObject($"<{jobStepLink}|*{step.Name}*> encountered errors"));
-			message.Blocks.Add(infoBlock);
-
-			AddLogDataContext(message, jobStepEventData, step.LogId);
-
-			foreach (JobNotificationConfig config in jobNotificationConfigs)
-			{
-				foreach (string recipient in config.Channels)
-				{
-					// This is our first time sending the summary grid, add both the job and the job steps to the collection
-					string channel = recipient;
-					SlackUser? userInfo = await _slackClient.FindUserByEmailAsync(recipient, cancellationToken);
-					if (userInfo is not null && userInfo.Id is not null)
-					{
-						channel = userInfo.Id;
-					}
-
-					try
-					{
-						_ = await _slackClient.PostMessageAsync(channel, message, cancellationToken);
-					}
-					catch (Exception ex)
-					{
-						_logger.LogInformation(ex, "Failed to send Slack setup failure message for job '{JobId}' and the template '{Template}' ({Error})", job.Id, job.TemplateId, ex.Message);
-					}
-				}
-			}
-
-			_logger.LogInformation("Finished sending failing setup notification for job '{JobId}'", job.Id);
-		}
-
-		/// <summary>
-		/// Sends a Slack notification for the specified job
-		/// </summary>
-		/// <param name="job">Information about the job</param>
-		/// <param name="branch">Name of the branch from the stream</param>
-		/// <param name="hordeJobURL">Optional URL of the job in Horde</param>
-		/// <param name="recipient">Identifier of the channel or user in Slack or user's email address</param>
-		/// <param name="abortedByUserName">Optional username provided when cancelling a job</param>
-		/// <param name="cancellationToken">Cancellation token for the operation</param>
-		/// <returns>True if a notification needed to be sent, otherwise returns false</returns>
-		async Task SendSlackSummaryMessageAsync(IJob job, string branch, string hordeJobURL, string recipient, string? abortedByUserName = null, CancellationToken cancellationToken = default)
-		{
-			SlackMessage? message = await GenerateSlackSummaryMessageAsync(job, branch, hordeJobURL, recipient, abortedByUserName, cancellationToken);
-			if (message is null)
-			{
-				_logger.LogInformation("Unable to generate Job summary notification for job '{JobId}' and the template '{TemplateId}'", job.Id, job.TemplateId);
-				return;
-			}
-
-			JobNotificationStateDocument? jobNotificationState = await GetJobNotificationStateAsync(job.Id, job.TemplateId, recipient, cancellationToken);
-			if (jobNotificationState is not null)
-			{
-				try
-				{
-					_ = _slackClient.UpdateMessageAsync(jobNotificationState.MessageId, message, cancellationToken);
-				}
-				catch (Exception ex)
-				{
-					_logger.LogInformation(ex, "Failed to update Slack summary message for job '{JobId}' and the template '{Template}' ({Error})", jobNotificationState.JobId, jobNotificationState.TemplateId, ex.Message);
-				}
-			}
-			else
-			{
-				// This is our first time sending the summary grid, add both the job and the job steps to the collection
-				string channel = recipient;
-				SlackUser? userInfo = await _slackClient.FindUserByEmailAsync(recipient, cancellationToken);
-				if (userInfo is not null && userInfo.Id is not null)
-				{
-					channel = userInfo.Id;
-				}
-
-				try
-				{
-					SlackMessageId messageId = await _slackClient.PostMessageAsync(channel, message, cancellationToken);
-					_ = await AddOrUpdateJobNotificationStateAsync(job.Id, job.TemplateId, recipient, messageId, cancellationToken);
-				}
-				catch (Exception ex)
-				{
-					_logger.LogInformation(ex, "Failed to send Slack summary message for job '{JobId}' and the template '{Template}' ({Error})", job.Id, job.TemplateId, ex.Message);
-				}
-			}
-		}
-
-		/// <summary>
-		/// Sends a Slack notification to a new or existing thread
-		/// </summary>
-		/// <param name="jobMessageId"><see cref="SlackMessageId"/> of the summary message</param>
-		/// <param name="jobStepNotificationState">Information about the job step's notification state</param>
-		/// <param name="threadMessage"><see cref="SlackMessage"/> to be sent</param>
-		/// <param name="cancellationToken">Cancellation token for the operation</param>
-		/// <returns>The <see cref="SlackMessageId"/> of either the newly sent notification or from the <see cref="JobStepNotificationStateDocument"/></returns>
-		async Task<SlackMessageId?> SendSlackThreadMessageAsync(SlackMessageId jobMessageId, JobStepNotificationStateDocument jobStepNotificationState, SlackMessage threadMessage, CancellationToken cancellationToken = default)
-		{
-			// Check to see if a message has already been sent out to avoid sending duplicate notifications
-			if (!String.IsNullOrEmpty(jobStepNotificationState.MessageId.ThreadTs))
-			{
-				return jobStepNotificationState.MessageId;
-			}
-
-			try
-			{
-				SlackMessageId threadMessageId = await _slackClient.PostMessageToThreadAsync(jobMessageId, threadMessage, cancellationToken);
-
-				_logger.LogInformation("Finished sending notification for job step '{JobStepId}' with Job Id '{JobId}' for template '{Template}'", jobStepNotificationState.JobStepId, jobStepNotificationState.JobId, jobStepNotificationState.TemplateId);
-				return threadMessageId;
-			}
-			catch (Exception ex)
-			{
-				_logger.LogInformation(ex, "Failed to send Slack message for job step '{JobStepId}' with Job Id '{JobId}' for template '{Template}' ({Error})", jobStepNotificationState.JobStepId, jobStepNotificationState.JobId, jobStepNotificationState.TemplateId, ex.Message);
-			}
-
-			return null;
-		}
-
-		/// <summary>
-		/// Sends a Slack notification for the specified job step
-		/// </summary>
-		/// <param name="jobId">Identifier of the job</param>
-		/// <param name="step">Information about the step</param>
-		/// <param name="jobStepEventData">List of the step's logged events</param>
-		/// <param name="underlyingStep">Optional information about the step's dependent step</param>
-		/// <param name="cancellationToken">Cancellation token for the operation</param>
-		/// <returns>True if a job step notification needed to be updated, otherwise returns false</returns>
-		async Task<bool> UpdateAndSendSlackNotificationsForStepAsync(JobId jobId, IJobStep step, List<LogData> jobStepEventData, IJobStep? underlyingStep = null, CancellationToken cancellationToken = default)
-		{
-			// Update our job step if it's being monitored and send an updated summary and post to the thread if needed
-			FilterDefinition<JobStepNotificationStateDocument> filter = Builders<JobStepNotificationStateDocument>.Filter.Eq(x => x.JobId, jobId) &
-																		Builders<JobStepNotificationStateDocument>.Filter.Eq(x => x.JobStepId, step.Id);
-			List<JobStepNotificationStateDocument>? jobStepNotificationStates = await GetJobStepNotificationStatesAsync(filter, cancellationToken);
-			if (jobStepNotificationStates is null || jobStepNotificationStates.Count == 0)
-			{
-				_logger.LogInformation("No notification to be sent for step '{JobStepId}' in job '{JobId}'", step.Id, jobId);
-				return false;
-			}
-
-			string badge = GetBadgeForStep(step);
-			SlackMessage? threadMessage = null;
-
-			// Notify on any problems of the Job Step
-			if (step.Outcome != JobStepOutcome.Success)
-			{
-				threadMessage = GenerateSlackThreadMessage(step, badge, jobStepEventData, underlyingStep);
-			}
-
-			foreach (JobStepNotificationStateDocument jobStepNotificationState in jobStepNotificationStates)
-			{
-				SlackMessageId? threadMessageId = null;
-				if (threadMessage is not null)
-				{
-					// Fetch our initial job notification to get the Message Id for threads
-					JobNotificationStateDocument? jobNotificationState = await GetJobNotificationStateAsync(jobId, jobStepNotificationState.TemplateId, jobStepNotificationState.Recipient, cancellationToken);
-					if (jobNotificationState is null)
-					{
-						_logger.LogInformation("No message id found for job '{JobId}' with template '{TemplateId}' for thread to be posted.", jobId, jobStepNotificationState.TemplateId);
-						continue;
-					}
-					threadMessageId = await SendSlackThreadMessageAsync(jobNotificationState.MessageId, jobStepNotificationState, threadMessage, cancellationToken);
-				}
-				_ = await AddOrUpdateJobStepNotificationStateAsync(jobStepNotificationState.JobId, jobStepNotificationState.TemplateId, jobStepNotificationState.Recipient, jobStepNotificationState.JobStepId, jobStepNotificationState.Group, jobStepNotificationState.TargetPlatform, badge, threadMessageId, cancellationToken);
-			}
-
-			return true;
-		}
-
-		/// <summary>
-		/// Adds a context to the Slack message with log details
-		/// </summary>
-		/// <param name="slackBlockContainer">Slack block container</param>
-		/// <param name="logData">List of our log information</param>
-		/// <param name="logId">Optional identifier of the log</param>
-		void AddLogDataContext(ISlackBlockContainer slackBlockContainer, List<LogData> logData, LogId? logId = null)
-		{
-			// Don't clutter the message with all of the events and do a preliminary check to see if the user needs to be made aware that events have been truncated
-			if (logData.Count > MaxEventMessages)
-			{
-				ContextBlock limitContext = new ContextBlock();
-				limitContext.Elements.Add(new TextObject($":warning: *Step has multiple events reported. Will only show first 5.* :warning:"));
-				slackBlockContainer.Blocks.Add(limitContext);
-			}
-
-			int logEventCount = Math.Min(logData.Count, MaxEventMessages);
-			int maxLogEventLength = logEventCount <= 0 ? MaxLogEventLength : MaxLogEventLength / logEventCount;
-			foreach (LogData log in logData.Take(MaxEventMessages))
-			{
-				SectionBlock logSection = new SectionBlock(new TextObject(SanitizeQuoteText(log.Message, maxLogEventLength)));
-				slackBlockContainer.Blocks.Add(logSection);
-			}
-
-			if (logId is not null)
-			{
-				string logUrl = new Uri(_serverInfo.DashboardUrl, $"log/{logId}").ToString();
-
-				ContextBlock detailsContext = new ContextBlock();
-				detailsContext.Elements.Add(new TextObject($"<{logUrl}|View Job Step Log>"));
-				slackBlockContainer.Blocks.Add(detailsContext);
-			}
-		}
-
-		/// <summary>
-		/// 
-		/// </summary>
-		/// <param name="step">Information about the step</param>
-		/// <param name="badge">Slack emoji associated with the notification</param>
-		/// <param name="logData">List of our log information</param>
-		/// <param name="underlyingStep">Optional information about the step's dependent step</param>
-		/// <returns>Generated <see cref="SlackMessage"/> if successful, otherwise returns null</returns>
-		SlackMessage? GenerateSlackThreadMessage(IJobStep step, string badge, List<LogData> logData, IJobStep? underlyingStep = null)
-		{
-			if (logData.Count == 0)
-			{
-				return null;
-			}
-
-			List<LogData> data = logData.Where(x => x.Severity == LogEventSeverity.Error).ToList();
-			if (!data.Any())
-			{
-				data = logData.Where(x => x.Severity == LogEventSeverity.Warning).ToList();
-			}
-
-			string jobStepLink = new Uri($"{_serverInfo.DashboardUrl}job/{step.Job.Id}?step={step.Id}").ToString();
-
-			SlackMessage alertMessage = new SlackMessage();
-
-			ContextBlock alertStepBlock = new ContextBlock();
-			alertStepBlock.Elements.Add(new PlainTextObject(badge, true));
-			alertStepBlock.Elements.Add(new TextObject($"<{jobStepLink}|*{step.Name}*>"));
-			alertStepBlock.Elements.Add(new PlainTextObject(badge, true));
-			alertMessage.Blocks.Add(alertStepBlock);
-
-			// If an underlying step was provided then the error occurred further up the job and impacted the current step
-			if (underlyingStep is not null)
-			{
-				ContextBlock dependencyContext = new ContextBlock();
-				dependencyContext.Elements.Add(new TextObject("*Step was marked as skipped due to errors in the following dependency*"));
-				alertMessage.Blocks.Add(dependencyContext);
-
-				// Slack Attachments are marked as being legacy and is recommended to be avoided
-				// The reason why Attachments are used is because this information is secondary to describe why a test had skipped
-				// We also wanted to differentiate the dependencies from the actual monitored node by indenting the dependency information
-				SlackAttachment attachment = new SlackAttachment();
-				attachment.Color = "#666666";
-
-				string dependencyJobStepLink = new Uri($"{_serverInfo.DashboardUrl}job/{step.Job.Id}?step={underlyingStep.Id}").ToString();
-
-				attachment.AddContext(new TextObject($"<{dependencyJobStepLink}|*{underlyingStep.Name}*>"));
-				AddLogDataContext(attachment, data, underlyingStep.LogId);
-				alertMessage.Attachments.Add(attachment);
-			}
-			else
-			{
-				AddLogDataContext(alertMessage, data, step.LogId);
-			}
-
-			return alertMessage;
-		}
-
-		/// <summary>
-		/// Generates the title for the summary messages
-		/// </summary>
-		/// <param name="job">Information about the job</param>
-		/// <param name="branch">Name of the branch from the stream</param>
-		/// <param name="hordeJobURL">Optional URL of the job in Horde</param>
-		/// <returns>Title message using Slack's markdown format</returns>
-		static string GenerateSlackMessageTitle(IJob job, string branch, string? hordeJobURL = null)
-		{
-			string changelistDetails = $"CL {job.CommitId}";
-			if (job.PreflightCommitId is not null)
-			{
-				changelistDetails = $"Preflight CL {job.PreflightCommitId} against {changelistDetails}";
-			}
-
-			if (!String.IsNullOrEmpty(hordeJobURL))
-			{
-				return $"<{hordeJobURL}|*{branch} - {changelistDetails} - {job.Name}*>";
-			}
-
-			return $"{branch} - {changelistDetails} - {job.Name}";
-		}
-
-		/// <summary>
-		/// Sanitizes the text for a Slack quote message
-		/// </summary>
-		/// <param name="text">Text to sanitize</param>
-		/// <param name="maxLength">Maximum allowed length for the message</param>
-		static string SanitizeQuoteText(string text, int maxLength = MaxLogEventLength)
-		{
-			if (text.Length > maxLength)
-			{
-				text = text.Substring(0, maxLength - 4) + "...";
-			}
-			return $"```{text}```";
-		}
-
-		/// <summary>
-		/// Adds a context to the Slack message with cancellation details
-		/// </summary>
-		/// <param name="message"><see cref="SlackMessage"/> to add the context to</param>
-		/// <param name="text">Information to provide in the context</param>
-		/// <param name="cancellationReason">Optional message explaining the cancellation</param>
-		static void AddCancelledJobSlackContext(SlackMessage message, string text, string? cancellationReason = null)
-		{
-			ContextBlock context = new ContextBlock();
-			context.Elements.Add(new TextObject(text));
-			message.Blocks.Add(context);
-
-			if (!String.IsNullOrEmpty(cancellationReason))
-			{
-				ContextBlock cancellationHeader = new ContextBlock();
-				cancellationHeader.Elements.Add(new TextObject($"*Cancellation Reason*"));
-				message.Blocks.Add(cancellationHeader);
-
-				ContextBlock cancellationContext = new ContextBlock();
-				cancellationContext.Elements.Add(new TextObject(cancellationReason));
-				message.Blocks.Add(cancellationContext);
-			}
-		}
-
-		/// <summary>
-		/// Create the job's summary message
-		/// </summary>
-		/// <param name="job">Information about the job</param>
-		/// <param name="branch">Name of the branch from the stream</param>
-		/// <param name="hordeJobURL">URL of the job in Horde</param>
-		/// <param name="jobStepNotificationStateDocuments">List of <see cref="JobStepNotificationStateDocument"/></param>
-		/// <param name="abortedByUserName">Optional username provided when cancelling a job</param>
-		/// <returns>Generated <see cref="SlackMessage"/></returns>
-		static SlackMessage CreateSlackSummaryNotification(IJob job, string branch, string hordeJobURL, List<JobStepNotificationStateDocument> jobStepNotificationStateDocuments, string? abortedByUserName = null)
-		{
-			DividerBlock divider = new DividerBlock();
-			SlackMessage message = new SlackMessage();
-
-			// Create our notification text
-			string notificationText = GenerateSlackMessageTitle(job, branch);
-			message.Text = $"{notificationText} Details...";
-
-			// Create a Block for the Slack Title
-			string messageTitle = GenerateSlackMessageTitle(job, branch, hordeJobURL);
-			ContextBlock jobContext = new ContextBlock();
-			jobContext.Elements.Add(new TextObject(messageTitle));
-			message.Blocks.Add(jobContext);
-
-			// If the job was manually aborted, then provide a message that catches the recipients eyes that the full information may not be available
-			if (!String.IsNullOrEmpty(abortedByUserName))
-			{
-				AddCancelledJobSlackContext(message, $"*{AlertBadge} Job was canceled by {abortedByUserName} and may be missing information {AlertBadge}*", job.CancellationReason);
-				message.Blocks.Add(divider);
-			}
-
-			int itemsInGroup = 0;
-			string currentGroup = String.Empty;
-			ContextBlock? header = null;
-			ContextBlock? stepContext = null;
-			foreach (JobStepNotificationStateDocument state in jobStepNotificationStateDocuments)
-			{
-				// Nothing to report on if the step is no longer in the job
-				if (!job.TryGetStep(state.JobStepId, out IJobStep? step))
-				{
-					continue;
-				}
-
-				if (currentGroup != state.Group || header is null)
-				{
-					// Check to see if we have an existing step context that needs to be added
-					if (stepContext is not null)
-					{
-						message.Blocks.Add(stepContext);
-						message.Blocks.Add(divider);
-						stepContext = null;
-					}
-
-					currentGroup = state.Group;
-					header = new ContextBlock();
-					header.Elements.Add(new TextObject($"*{currentGroup}*"));
-					message.Blocks.Add(header);
-				}
-
-				string jobStepLink = $"{hordeJobURL}?step={state.JobStepId}";
-
-				// Slack can only allow for 10 elements per block (https://api.slack.com/reference/block-kit/blocks#context) and we want to track when to create new ContextBlocks
-				bool bNeedsNewContext = (itemsInGroup % MaxSlackBlocks == 0 || stepContext is null);
-				if (bNeedsNewContext)
-				{
-					if (stepContext is not null)
-					{
-						message.Blocks.Add(stepContext);
-					}
-					stepContext = new ContextBlock();
-				}
-				stepContext?.Elements.Add(new TextObject($"<{jobStepLink}|{state.Badge} {state.TargetPlatform}>"));
-
-				// The StepContext should always be valid here, but to disable the warning, we provide a default so that we only have 1 block per context.
-				itemsInGroup = stepContext?.Elements.Count ?? MaxSlackBlocks;
-			}
-
-			// Check to see if we have an existing step context that needs to be added
-			if (stepContext is not null)
-			{
-				message.Blocks.Add(stepContext);
-				message.Blocks.Add(divider);
-			}
-
-			return message;
-		}
-
-		#endregion Slack
+		/// <inheritdoc/>
+		public void Dispose()
+			=> _slackNotificationProcessor.Dispose();
 	}
 }
diff --git a/Engine/Source/Programs/Horde/Plugins/Experimental/HordeServer.Experimental/Notifications/SlackNotificationProcessor.cs b/Engine/Source/Programs/Horde/Plugins/Experimental/HordeServer.Experimental/Notifications/SlackNotificationProcessor.cs
new file mode 100644
--- /dev/null
+++ b/Engine/Source/Programs/Horde/Plugins/Experimental/HordeServer.Experimental/Notifications/SlackNotificationProcessor.cs
@@ -0,0 +1,1520 @@
+// Copyright Epic Games, Inc. All Rights Reserved.
+
+using EpicGames.Horde.Jobs;
+using EpicGames.Horde.Jobs.Templates;
+using EpicGames.Horde.Logs;
+using EpicGames.Horde.Streams;
+using EpicGames.Slack;
+using EpicGames.Slack.Blocks;
+using EpicGames.Slack.Elements;
+using HordeServer.Jobs;
+using HordeServer.Logs;
+using HordeServer.Streams;
+using HordeServer.Users;
+using Microsoft.Extensions.Logging;
+using Microsoft.Extensions.Options;
+
+namespace HordeServer.Experimental.Notifications
+{
+	/// <summary>
+	/// Processes notifications to be sent through the Slack API
+	/// </summary>
+	public sealed class SlackNotificationProcessor : IDisposable
+	{
+		/// <summary>
+		/// Structure for generic use of log reporting
+		/// </summary>
+		/// <remarks>Cancellation of jobs and job steps may not always include log information. In these cases we want to still provide a message to the user for the reason.</remarks>
+		internal struct LogData
+		{
+			public LogEventSeverity Severity { get; }
+			public string Message { get; }
+			public LogId? LogId { get; } = null;
+			public LogData(LogEventSeverity severity, string message, LogId? logId = null)
+			{
+				Severity = severity;
+				Message = message;
+				LogId = logId;
+			}
+		}
+
+		readonly IOptionsMonitor<BuildConfig> _buildConfig;
+		readonly IOptionsMonitor<ExperimentalConfig> _experimentalConfig;
+		readonly IJobNotificationCollection _jobNotificationCollection;
+		readonly IJobCollection _jobCollection;
+		readonly IUserCollection _userCollection;
+		readonly IServerInfo _serverInfo;
+		readonly ILogger _logger;
+
+		readonly HttpClient? _httpClient;
+		readonly SlackClient? _slackClient;
+
+		/// <summary>
+		/// Slack emoji used to emphasize state
+		/// </summary>
+		/// <remarks>All emoji were verified against what is available in Slack by default</remarks>
+		const string PendingBadge = ":large_blue_square:";
+		const string WarningBadge = ":large_yellow_square:";
+		const string FailureBadge = ":large_red_square:";
+		const string SkippedBadge = ":black_square:";
+		const string PassingBadge = ":large_green_square:";
+
+		/// <summary>
+		/// Slack emoji used to emphasize an aborted job
+		/// </summary>
+		const string AlertBadge = ":rotating_light:";
+
+		/// <summary>
+		/// Limit of logs to be provided in a Slack notification
+		/// </summary>
+		const int MaxEventMessages = 5;
+
+		/// <summary>
+		/// Length limit to avoid character limit on each individual block container from Slack
+		/// </summary>
+		const int MaxBlockCharacterLength = 2048;
+
+		/// <summary>
+		/// Block limit per Slack Context before an exception is thrown
+		/// </summary>
+		/// <remarks>See https://api.slack.com/reference/block-kit/blocks#context for detailed information</remarks>
+		const int MaxSlackBlocks = 10;
+
+		/// <summary>
+		/// Constructor
+		/// </summary>
+		/// <param name="buildServerConfig">Options from the build plugin's server config</param>
+		/// <param name="buildConfig">Monitored options from the build plugin's global config</param>
+		/// <param name="experimentalConfig">Monitored options from the experimental plugin's global config</param>
+		/// <param name="jobNotificationCollection">Collection of job and job step notifications</param>
+		/// <param name="jobCollection">Collection of jobs</param>
+		/// <param name="userCollection">Collection of users</param>
+		/// <param name="serverInfo">Information of the current server instance</param>
+		/// <param name="logger">Logger for output</param>
+		public SlackNotificationProcessor(IOptions<BuildServerConfig> buildServerConfig,
+			IOptionsMonitor<BuildConfig> buildConfig,
+			IOptionsMonitor<ExperimentalConfig> experimentalConfig,
+			IJobNotificationCollection jobNotificationCollection,
+			IJobCollection jobCollection,
+			IUserCollection userCollection,
+			IServerInfo serverInfo,
+			ILogger logger)
+		{
+			_buildConfig = buildConfig;
+			_experimentalConfig = experimentalConfig;
+			_jobNotificationCollection = jobNotificationCollection;
+			_jobCollection = jobCollection;
+			_userCollection = userCollection;
+			_serverInfo = serverInfo;
+			_logger = logger;
+
+			if (!String.IsNullOrEmpty(buildServerConfig.Value.SlackToken))
+			{
+				_httpClient = new HttpClient();
+				_httpClient.DefaultRequestHeaders.Add("Authorization", $"Bearer {buildServerConfig.Value.SlackToken}");
+				_slackClient = new SlackClient(_httpClient, _logger);
+			}
+		}
+
+		/// <inheritdoc/>
+		public void Dispose()
+			=> _httpClient?.Dispose();
+
+		/// <summary>
+		/// Processes a Job complete notification
+		/// </summary>
+		/// <param name="job">Information about the job</param>
+		/// <param name="cancellationToken">Cancellation token for the operation</param>
+		/// <remarks>Currently this should only be called if a job is canceled. Completed jobs should still receive a job step to be processed.</remarks>
+		public async Task ProcessJobCompleteAsync(IJob job, CancellationToken cancellationToken = default)
+		{
+			// Only handle the job complete notification if the job was canceled
+			// Otherwise we will leave it to the steps to handle when all notifications have been sent out to avoid conditions where the job complete notification gets received before the job step complete notification
+			if (job.AbortedByUserId is null)
+			{
+				return;
+			}
+
+			bool wasSpawnJobProcessed = await TryProcessSpawnJobCompleteNotificationAsync(job, cancellationToken);
+			if (wasSpawnJobProcessed)
+			{
+				_logger.LogInformation("Job complete notification for '{JobId}' processed for parent job '{ParentJobId}'", job.Id, job.ParentJobId);
+			}
+
+			IJobNotificationStateQueryBuilder builder = JobNotificationCollection.CreateJobNotificationStateQueryBuilder();
+			builder.AddJobFilter(job.Id);
+			builder.AddTemplateFilter(job.TemplateId);
+			IReadOnlyList<IJobNotificationStateRef>? jobNotificationStates = await _jobNotificationCollection.GetJobNotificationStatesAsync(builder, cancellationToken);
+			if (jobNotificationStates is null)
+			{
+				_logger.LogInformation("No job notifications found for job id '{JobId}'.", job.Id);
+				return;
+			}
+
+			string abortedByUserName = "Unknown User";
+			IUser? user = await _userCollection.GetUserAsync(job.AbortedByUserId.Value, cancellationToken);
+			if (user is not null)
+			{
+				abortedByUserName = user.Name;
+			}
+
+			string branch = GetBranchFromStream(_buildConfig.CurrentValue, job.StreamId);
+			foreach (IJobNotificationStateRef jobNotificationState in jobNotificationStates)
+			{
+				// Update our steps to reflect the canceled job
+				_ = await UpdatePendingJobStepNotificationsAsync(jobNotificationState.JobId, jobNotificationState.TemplateId, jobNotificationState.Recipient, cancellationToken);
+
+				await SendSlackSummaryMessageAsync(job, branch, jobNotificationState.Recipient, jobNotificationState, null, abortedByUserName, cancellationToken);
+
+				// No further notifications will be sent out and can be cleaned up
+				_ = await _jobNotificationCollection.DeleteJobStepNotificationStatesAsync(jobNotificationState.JobId, jobNotificationState.TemplateId, jobNotificationState.Recipient, cancellationToken);
+				_ = await _jobNotificationCollection.DeleteJobNotificationStatesAsync(jobNotificationState.JobId, jobNotificationState.TemplateId, jobNotificationState.Recipient, cancellationToken);
+			}
+		}
+
+		/// <summary>
+		/// Processes a notification for a job step that has been canceled
+		/// </summary>
+		/// <param name="job">Information about the job</param>
+		/// <param name="step">Information about the step</param>
+		/// <param name="cancellationToken">Cancellation token for the operation</param>
+		/// <remarks>A job step that has been canceled will be handled the same way as if the job step has failures.</remarks>
+		public async Task ProcessJobStepAbortedAsync(IJob job, IJobStep step, CancellationToken cancellationToken = default)
+		{
+			List<LogData> logEvents = new List<LogData>();
+			string cancellationDetails = "Step was canceled.";
+			if (step.AbortedByUserId is not null)
+			{
+				IUser? user = await _userCollection.GetUserAsync(step.AbortedByUserId.Value, cancellationToken);
+				if (user is not null)
+				{
+					cancellationDetails = $"Step was canceled by {user.Name}.";
+				}
+			}
+
+			if (!String.IsNullOrEmpty(step.CancellationReason))
+			{
+				cancellationDetails += $" Reason for cancellation: '{step.CancellationReason}'.";
+			}
+			logEvents.Add(new LogData(LogEventSeverity.Error, cancellationDetails));
+			await ProcessJobStepCompleteNotificationAsync(job, step, logEvents, cancellationToken);
+		}
+
+		/// <summary>
+		/// Process a job step that has been completed
+		/// </summary>
+		/// <param name="job">Information about the job</param>
+		/// <param name="step">Information about the step</param>
+		/// <param name="jobStepEventData">List of the step's logged events</param>
+		/// <param name="cancellationToken">Cancellation token for the operation</param>
+		/// <returns></returns>
+		public async Task ProcessJobStepCompleteAsync(IJob job, IJobStep step, List<ILogEventData> jobStepEventData, CancellationToken cancellationToken = default)
+		{
+			List<LogData> logEvents = ExtractLogInformation(step, jobStepEventData);
+			await ProcessJobStepCompleteNotificationAsync(job, step, logEvents, cancellationToken);
+		}
+
+		#region Job Processing
+
+		/// <summary>
+		/// Search the job for steps to be included in the notifications based on the experimental plugin's global config
+		/// </summary>
+		/// <param name="job">Information about the job</param>
+		/// <param name="config">Configuration to fetch notification information</param>
+		/// <param name="parentJob">Optional information regarding a parent job</param>
+		/// <param name="cancellationToken">Cancellation token for the operation</param>
+		async Task ProcessJobAsync(IJob job, JobNotificationConfig config, IJob? parentJob = null, CancellationToken cancellationToken = default)
+		{
+			// Populate our job information needed for the notification
+			try
+			{
+				foreach (IJobStepBatch batch in job.Batches)
+				{
+					foreach (IJobStep step in batch.Steps)
+					{
+						if (config.GetStepDetails(step.Name, out string? group, out string? platform) && !String.IsNullOrEmpty(group) && !String.IsNullOrEmpty(platform))
+						{
+							string badge = GetBadgeForStep(step);
+							foreach (string recipient in config.Channels)
+							{
+								await _jobNotificationCollection.AddOrUpdateJobStepNotificationStateAsync(job.Id, job.TemplateId, recipient, step.Id, group, platform, badge, String.Empty, String.Empty, null, parentJob?.Id, parentJob?.TemplateId, cancellationToken);
+							}
+						}
+					}
+				}
+			}
+			catch (Exception ex)
+			{
+				_logger.LogInformation(ex, "Failed to extract step details for '{Job}' in stream '{StreamId}' and for template '{TemplateId}' ({Error})", job.Id, job.StreamId, job.TemplateId, ex.Message);
+			}
+		}
+
+		/// <summary>
+		/// Try to process the job complete notification as a spawned job
+		/// </summary>
+		/// <param name="job">Information about the job</param>
+		/// <param name="cancellationToken">Cancellation token for the operation</param>
+		/// <returns>True is the job was found to be a valid spawned job, otherwise, will return false.</returns>
+		async Task<bool> TryProcessSpawnJobCompleteNotificationAsync(IJob job, CancellationToken cancellationToken = default)
+		{
+			// Fetch the parent information if available
+			if (job.ParentJobId is null || job.ParentJobStepId is null)
+			{
+				return false;
+			}
+
+			IJob? parentJob = await _jobCollection.GetAsync(job.ParentJobId!.Value, cancellationToken);
+			if (parentJob is null)
+			{
+				_logger.LogInformation("Unable to process job complete notification as parent job was not found for job '{JobId}' and parent job '{ParentJob}'", job.Id, job.ParentJobId);
+				return false;
+			}
+
+			IJobNotificationStateQueryBuilder builder = JobNotificationCollection.CreateJobNotificationStateQueryBuilder();
+			builder.AddJobFilter(parentJob.Id);
+			builder.AddTemplateFilter(parentJob.TemplateId);
+			IReadOnlyList<IJobNotificationStateRef>? jobNotificationStates = await _jobNotificationCollection.GetJobNotificationStatesAsync(builder, cancellationToken);
+			if (jobNotificationStates is null)
+			{
+				_logger.LogInformation("No active notifications found for parent job '{ParentJob}' with parent job template '{ParentJobTemplateId}' in spawned job '{JobId}' with spawned job template '{TemplateId}'", parentJob.Id, parentJob.TemplateId, job.Id, job.TemplateId);
+				return false;
+			}
+
+			string branch = GetBranchFromStream(_buildConfig.CurrentValue, parentJob.StreamId);
+			foreach (IJobNotificationStateRef jobNotificationState in jobNotificationStates)
+			{
+				// No further notifications will be sent out and can be cleaned up
+				_ = await _jobNotificationCollection.DeleteJobStepNotificationStatesAsync(job.Id, job.TemplateId, jobNotificationState.Recipient, cancellationToken);
+
+				// Send out an updated summary to include the current job has been canceled
+				await SendSlackSummaryMessageAsync(parentJob, branch, jobNotificationState.Recipient, jobNotificationState, null, null, cancellationToken);
+			}
+
+			return true;
+		}
+
+		/// <summary>
+		/// Try to process the job setup notification as a spawned job
+		/// </summary>
+		/// <param name="job">Information about the job</param>
+		/// <param name="step">Information about the step</param>
+		/// <param name="jobStepEventData">List of the step's logged events</param>
+		/// <param name="cancellationToken">Cancellation token for the operation</param>
+		/// <returns>True is the job was found to be a valid spawned job, otherwise, will return false.</returns>
+		async Task<bool> TryProcessSpawnJobSetupNotificationAsync(IJob job, IJobStep step, List<LogData> jobStepEventData, CancellationToken cancellationToken = default)
+		{
+			// Fetch the parent information if available
+			if (job.ParentJobId is null || job.ParentJobStepId is null)
+			{
+				return false;
+			}
+
+			// Fetch all of our notifications for the parent job and parent job step to see if we need to consolidate jobs
+			IJobStepNotificationStateQueryBuilder jobStepNotificationBuilder = JobNotificationCollection.CreateJobStepNotificationStateQueryBuilder();
+			jobStepNotificationBuilder.AddJobFilter(job.ParentJobId!.Value);
+			jobStepNotificationBuilder.AddJobStepFilter(job.ParentJobStepId!.Value);
+			IReadOnlyList<IJobStepNotificationStateRef>? jobStepNotificationStates = await _jobNotificationCollection.GetJobStepNotificationStatesAsync(jobStepNotificationBuilder, cancellationToken);
+			if (jobStepNotificationStates is null || jobStepNotificationStates.Count == 0)
+			{
+				_logger.LogInformation("No active notifications found for parent step '{ParentJobStep}' in parent job '{ParentJob}' for spawned job step '{JobStepId}' in spawned job '{JobId}'", job.ParentJobStepId, job.ParentJobId, step.Id, job.Id);
+				return false;
+			}
+
+			IJob? parentJob = await _jobCollection.GetAsync(job.ParentJobId!.Value, cancellationToken);
+			if (parentJob is null)
+			{
+				_logger.LogInformation("Unable to process spawned setup as parent job was not found for job '{JobId}' and parent job '{ParentJob}'", job.Id, job.ParentJobId);
+				return false;
+			}
+
+			IJobNotificationStateQueryBuilder jobNotificationBuilder = JobNotificationCollection.CreateJobNotificationStateQueryBuilder();
+			jobNotificationBuilder.AddJobFilter(parentJob.Id);
+			jobNotificationBuilder.AddTemplateFilter(parentJob.TemplateId);
+			IReadOnlyList<IJobNotificationStateRef>? jobNotificationStates = await _jobNotificationCollection.GetJobNotificationStatesAsync(jobNotificationBuilder, cancellationToken);
+			if (jobNotificationStates is null)
+			{
+				_logger.LogInformation("No job notifications found for parent job id '{JobId}'.", job.ParentJobId);
+				return false;
+			}
+
+			// From this point forward the current job/step will be handled as a spawned job
+
+			string branch = GetBranchFromStream(_buildConfig.CurrentValue, parentJob.StreamId);
+
+			// Check if the Setup step failed to notify that the job can't be executed
+			bool wasSetupSuccess = (step.Outcome == JobStepOutcome.Success || step.Outcome == JobStepOutcome.Warnings);
+			if (!wasSetupSuccess)
+			{
+				// Send the failed setup job to the thread and update the summary
+				foreach (IJobNotificationStateRef jobNotificationState in jobNotificationStates)
+				{
+					await SendSetupFailedSlackMessageToThreadAsync(job, step, jobStepEventData, jobNotificationState.Channel, jobNotificationState.Ts, cancellationToken);
+					await SendSlackSummaryMessageAsync(job, branch, jobNotificationState.Recipient, jobNotificationState, parentJob, null, cancellationToken);
+
+					// Check to see if all notifications for this job were sent then delete all documents related to this job as no other notifications will be sent
+					await TryNotificationCleanupAsync(parentJob, jobNotificationState.Recipient, cancellationToken);
+				}
+				return true;
+			}
+
+			// Fetch our notification configurations for the parent job to append our current job's steps
+			ExperimentalConfig currentExperimentalConfig = _experimentalConfig.CurrentValue;
+
+			HashSet<JobNotificationConfig>? jobNotificationConfigs;
+			if (!currentExperimentalConfig.TryGetJobNofications(parentJob.StreamId, parentJob.TemplateId, out jobNotificationConfigs))
+			{
+				_logger.LogInformation("Parent configurations not found for parent step '{ParentJobStep}' in parent job '{ParentJob}' for '{JobId}' in stream '{StreamId}' and for template '{TemplateId}'", job.ParentJobStepId, job.ParentJobId, job.Id, parentJob.StreamId, parentJob.TemplateId);
+				return true;
+			}
+
+			foreach (JobNotificationConfig config in jobNotificationConfigs)
+			{
+				await ProcessJobAsync(job, config, parentJob, cancellationToken);
+			}
+
+			// Send the updated summary after the steps has been updated
+			foreach (IJobNotificationStateRef jobNotificationState in jobNotificationStates)
+			{
+				await SendSlackSummaryMessageAsync(job, branch, jobNotificationState.Recipient, jobNotificationState, parentJob, null, cancellationToken);
+
+				// Check to see if all notifications for this job were sent then delete all documents related to this job as no other notifications will be sent
+				await TryNotificationCleanupAsync(parentJob, jobNotificationState.Recipient, cancellationToken);
+			}
+
+			return true;
+		}
+
+		/// <summary>
+		/// Process the job setup notification
+		/// </summary>
+		/// <param name="job">Information about the job</param>
+		/// <param name="step">Information about the step</param>
+		/// <param name="jobStepEventData">List of the step's logged events</param>
+		/// <param name="cancellationToken">Cancellation token for the operation</param>
+		async Task ProcessJobSetupNotificationAsync(IJob job, IJobStep step, List<LogData> jobStepEventData, CancellationToken cancellationToken = default)
+		{
+			// Fetch our notification configurations for the parent job to append our current job's steps
+			ExperimentalConfig currentExperimentalConfig = _experimentalConfig.CurrentValue;
+
+			HashSet<JobNotificationConfig>? jobNotificationConfigs;
+			if (!currentExperimentalConfig.TryGetJobNofications(job.StreamId, job.TemplateId, out jobNotificationConfigs))
+			{
+				_logger.LogInformation("No configurations found for '{JobId}' in stream '{StreamId}' and for template '{TemplateId}'", job.Id, job.StreamId, job.TemplateId);
+				return;
+			}
+
+			string branch = GetBranchFromStream(_buildConfig.CurrentValue, job.StreamId);
+
+			// Check if the Setup step failed to notify that the job can't be executed
+			bool wasSetupSuccess = (step.Outcome == JobStepOutcome.Success || step.Outcome == JobStepOutcome.Warnings);
+			if (!wasSetupSuccess)
+			{
+				await SendSetupFailedSlackMessageAsync(job, step, branch, jobStepEventData, jobNotificationConfigs, cancellationToken);
+				return;
+			}
+
+			foreach (JobNotificationConfig config in jobNotificationConfigs)
+			{
+				await ProcessJobAsync(job, config, null, cancellationToken);
+				foreach (string recipient in config.Channels)
+				{
+					await SendSlackSummaryMessageAsync(job, branch, recipient, null, null, null, cancellationToken);
+				}
+			}
+		}
+
+		/// <summary>
+		/// Process the job step completed notification
+		/// </summary>
+		/// <param name="job">Information about the job</param>
+		/// <param name="step">Information about the step</param>
+		/// <param name="jobStepEventData">List of the step's logged events</param>
+		/// <param name="cancellationToken">Cancellation token for the operation</param>
+		async Task ProcessJobStepCompleteNotificationAsync(IJob job, IJobStep step, List<LogData> jobStepEventData, CancellationToken cancellationToken = default)
+		{
+			// Perform a check to see if the completed step was the Setup step for the job, as this will use the configurations to populate the notification information
+			if (step.Name == IJob.SetupNodeName)
+			{
+				// Check to see if we're a spawned job which requires the job to be merged into a parent job
+				bool wasSpawnJobProcessed = await TryProcessSpawnJobSetupNotificationAsync(job, step, jobStepEventData, cancellationToken);
+				if (wasSpawnJobProcessed)
+				{
+					_logger.LogInformation("Job '{JobId}' processed for parent job '{ParentJobId}'", job.Id, job.ParentJobId);
+				}
+				else
+				{
+					await ProcessJobSetupNotificationAsync(job, step, jobStepEventData, cancellationToken);
+				}
+				return;
+			}
+
+			// From here we already have a summary notification sent and we need to focus on updating our notification
+
+			// Send a message for our current step
+			bool needsSummaryUpdate = await UpdateAndSendSlackNotificationsForStepAsync(job, step, jobStepEventData, null, cancellationToken);
+
+			// Fetch all dependent steps if there was an error and send notifications with this step as the cause for the dependent steps being skipped
+			if (step.Outcome != JobStepOutcome.Success && step.Outcome != JobStepOutcome.Warnings)
+			{
+				HashSet<JobStepId>? dependentSteps = await FindDependentJobStepIdsAsync(job, step.Id, cancellationToken);
+				if (dependentSteps is not null)
+				{
+					foreach (JobStepId stepId in dependentSteps)
+					{
+						// Send a message for our dependent step
+						if (job.TryGetStep(stepId, out IJobStep? dependentStep))
+						{
+							needsSummaryUpdate |= await UpdateAndSendSlackNotificationsForStepAsync(job, dependentStep, jobStepEventData, step, cancellationToken);
+						}
+					}
+				}
+			}
+
+			// Get all of the job information such as identifier and template identifier
+			List<JobId> jobIds = new List<JobId>()
+			{
+				job.Id
+			};
+
+			List<TemplateId> templateIds = new List<TemplateId>()
+			{
+				job.TemplateId
+			};
+
+			// Check if the job has a parent to include in the list when filtering for job notifications
+			IJob? parentJob = null;
+			if (job.ParentJobId is not null)
+			{
+				parentJob = await _jobCollection.GetAsync(job.ParentJobId!.Value, cancellationToken);
+				if (parentJob is not null)
+				{
+					jobIds.Add(parentJob.Id);
+					templateIds.Add(parentJob.TemplateId);
+				}
+			}
+
+			IJobNotificationStateQueryBuilder builder = JobNotificationCollection.CreateJobNotificationStateQueryBuilder();
+			builder.AddJobFilters(jobIds);
+			builder.AddTemplateFilters(templateIds);
+			IReadOnlyList<IJobNotificationStateRef>? jobNotificationStates = await _jobNotificationCollection.GetJobNotificationStatesAsync(builder, cancellationToken);
+			if (jobNotificationStates is not null && job.GetState() == JobState.Complete)
+			{
+				foreach (IJobNotificationStateRef jobNotificationState in jobNotificationStates)
+				{
+					needsSummaryUpdate |= await UpdatePendingJobStepNotificationsAsync(jobNotificationState.JobId, jobNotificationState.TemplateId, jobNotificationState.Recipient, cancellationToken);
+				}
+			}
+
+			// Check to see if any of the steps were updated which requires the Slack summary to be updated
+			if (!needsSummaryUpdate)
+			{
+				return;
+			}
+
+			// Send the updated summary after the steps has been updated
+			if (jobNotificationStates is null)
+			{
+				_logger.LogInformation("No job notifications found for job id '{JobId}' even though job step '{StepId}' was updated.", job.Id, step.Id);
+				return;
+			}
+
+			string branch = GetBranchFromStream(_buildConfig.CurrentValue, job.StreamId);
+			foreach (IJobNotificationStateRef jobNotificationState in jobNotificationStates)
+			{
+				// Check to see if the current job is the main job or if it was spawned from the parent
+				bool isSpawnedJob = parentJob is not null && jobNotificationState.JobId == parentJob.Id;
+				IJob notificationJob = isSpawnedJob ? parentJob! : job;
+
+				await SendSlackSummaryMessageAsync(job, branch, jobNotificationState.Recipient, jobNotificationState, isSpawnedJob ? parentJob : null, null, cancellationToken);
+
+				// Check to see if all notifications for this job were sent then delete all documents related to this job as no other notifications will be sent
+				await TryNotificationCleanupAsync(notificationJob, jobNotificationState.Recipient, cancellationToken);
+			}
+		}
+
+		/// <summary>
+		/// Finds all dependent steps of the <see cref="IJobStepNotificationStateRef"/>
+		/// </summary>
+		/// <param name="job">Information about the job</param>
+		/// <param name="underlyingStepId">Identifier of the possible job step dependency</param>
+		/// <param name="cancellationToken">Cancellation token for the operation</param>
+		/// <returns>HashSet of <see cref="JobStepId"/> if dependent steps have been found, otherwise returns null</returns>
+		async Task<HashSet<JobStepId>?> FindDependentJobStepIdsAsync(IJob job, JobStepId underlyingStepId, CancellationToken cancellationToken = default)
+		{
+			// Fetch all of our job steps sending notifications
+			IJobStepNotificationStateQueryBuilder builder = JobNotificationCollection.CreateJobStepNotificationStateQueryBuilder();
+			builder.AddJobFilter(job.Id);
+			builder.AddTemplateFilter(job.TemplateId);
+			IReadOnlyList<IJobStepNotificationStateRef>? jobStepNotificationStateDocuments = await _jobNotificationCollection.GetJobStepNotificationStatesAsync(builder, cancellationToken);
+			if (jobStepNotificationStateDocuments is null)
+			{
+				return null;
+			}
+
+			HashSet<JobStepId> dependencyJobStepNotifications = new HashSet<JobStepId>();
+			HashSet<JobStepId> cachedSteps = new HashSet<JobStepId>();
+			foreach (IJobStepNotificationStateRef jobStepNotificationState in jobStepNotificationStateDocuments)
+			{
+				if (cachedSteps.Contains(jobStepNotificationState.JobStepId))
+				{
+					continue;
+				}
+				cachedSteps.Add(jobStepNotificationState.JobStepId);
+
+				List<JobStepId> dependencySteps = new List<JobStepId>()
+				{
+					jobStepNotificationState.JobStepId
+				};
+
+				HashSet<JobStepId> visitedDependencySteps = new HashSet<JobStepId>()
+				{
+					jobStepNotificationState.JobStepId
+				};
+
+				bool isMatchingDependency = false;
+				while (dependencySteps.Count > 0 && !isMatchingDependency)
+				{
+					JobStepId jobStepId = dependencySteps[0];
+					dependencySteps.RemoveAt(0);
+
+					IJobStep? currentStep;
+					if (!job.TryGetStep(jobStepId, out currentStep))
+					{
+						continue;
+					}
+
+					// Iterate through all of the step's dependencies
+					// Add the step Id of the notification state if the dependencies match
+					foreach (JobStepId dependencyStepId in currentStep.InputDependencies)
+					{
+						if (underlyingStepId == dependencyStepId)
+						{
+							isMatchingDependency = true;
+							dependencyJobStepNotifications.Add(jobStepNotificationState.JobStepId);
+							break;
+						}
+						else if (!visitedDependencySteps.Contains(dependencyStepId))
+						{
+							dependencySteps.Add(dependencyStepId);
+							visitedDependencySteps.Add(dependencyStepId);
+						}
+					}
+
+					// Check to see if we already found a potential match before continuing searching through dependencies
+					if (isMatchingDependency)
+					{
+						continue;
+					}
+
+					foreach (JobStepId dependencyStepId in currentStep.OrderDependencies)
+					{
+						if (underlyingStepId == dependencyStepId)
+						{
+							isMatchingDependency = true;
+							dependencyJobStepNotifications.Add(jobStepNotificationState.JobStepId);
+							break;
+						}
+						else if (!visitedDependencySteps.Contains(dependencyStepId))
+						{
+							dependencySteps.Add(dependencyStepId);
+							visitedDependencySteps.Add(dependencyStepId);
+						}
+					}
+				}
+			}
+
+			return dependencyJobStepNotifications;
+		}
+
+		/// <summary>
+		/// Syncs up the main job information with our current recorded steps
+		/// </summary>
+		/// <param name="jobId">Identifier of the job</param>
+		/// <param name="templateId">Identifier of the template</param>
+		/// <param name="recipient">Identifier of the channel or user in Slack or user's email address</param>
+		/// <param name="cancellationToken">Cancellation token for the operation</param>
+		/// <returns>True if the information has been found and updated, otherwise false</returns>
+		async Task<bool> UpdatePendingJobStepNotificationsAsync(JobId jobId, TemplateId templateId, string recipient, CancellationToken cancellationToken = default)
+		{
+			// Fetch all of our job steps that are still pending to determine if we still have notifications to send
+			// Because we need to fetch not only the main job step documents as well as any spawned job step documents
+			// We need to include the parent job and template identifiers to get all of our records
+			IJobStepNotificationStateQueryBuilder builder = JobNotificationCollection.CreateJobStepNotificationStateQueryBuilder();
+			builder.AddJobAndParentJobFilter(jobId);
+			builder.AddTemplateAndParentTemplateFilter(templateId);
+			builder.AddRecipientFilter(recipient);
+			builder.AddBadgeFilter(PendingBadge);
+			IReadOnlyList<IJobStepNotificationStateRef>? jobStepNotificationStates = await _jobNotificationCollection.GetJobStepNotificationStatesAsync(builder, cancellationToken);
+			if (jobStepNotificationStates is null)
+			{
+				return false;
+			}
+
+			bool wasUpdated = false;
+			foreach (IJobStepNotificationStateRef jobStepNotificationState in jobStepNotificationStates)
+			{
+				IJob? job = await _jobCollection.GetAsync(jobStepNotificationState.JobId, cancellationToken);
+				if (job is null)
+				{
+					continue;
+				}
+
+				IJobStep? step = null;
+				if (!job.TryGetStep(jobStepNotificationState.JobStepId, out step))
+				{
+					continue;
+				}
+
+				string badge = GetBadgeForStep(step);
+				if (badge != jobStepNotificationState.Badge)
+				{
+					_ = await _jobNotificationCollection.AddOrUpdateJobStepNotificationStateAsync(jobStepNotificationState.JobId, jobStepNotificationState.TemplateId, jobStepNotificationState.Recipient, jobStepNotificationState.JobStepId,
+																									jobStepNotificationState.Group, jobStepNotificationState.TargetPlatform, badge, jobStepNotificationState.Channel, jobStepNotificationState.Ts,
+																									jobStepNotificationState.ThreadTs, jobStepNotificationState.ParentJobId, jobStepNotificationState.ParentJobTemplateId, cancellationToken);
+					wasUpdated = true;
+				}
+			}
+
+			return wasUpdated;
+		}
+
+		/// <summary>
+		/// Clean up notification states once all notifications have been sent
+		/// </summary>
+		/// <param name="job">Information about the job</param>
+		/// <param name="recipient">Identifier of the channel or user in Slack or user's email address</param>
+		/// <param name="cancellationToken">Cancellation token for the operation</param>
+		async Task TryNotificationCleanupAsync(IJob job, string recipient, CancellationToken cancellationToken = default)
+		{
+			// Fetch all of our job steps that are still pending to determine if we still have notifications to send
+			// Because we need to fetch not only the main job step documents as well as any spawned job step documents
+			// We need to include the parent job and template identifiers to get all of our records
+			IJobStepNotificationStateQueryBuilder builder = JobNotificationCollection.CreateJobStepNotificationStateQueryBuilder();
+			builder.AddJobAndParentJobFilter(job.Id);
+			builder.AddTemplateAndParentTemplateFilter(job.TemplateId);
+			builder.AddRecipientFilter(recipient);
+			IReadOnlyList<IJobStepNotificationStateRef>? jobStepNotificationStates = await _jobNotificationCollection.GetJobStepNotificationStatesAsync(builder, cancellationToken);
+			if (jobStepNotificationStates is null)
+			{
+				// Nothing to clean up
+				return;
+			}
+
+			// Because we could have multiple jobs providing notification information, we need to check and make sure all jobs have been accounted for
+			bool allNotificationsSent = true;
+			foreach (IJobStepNotificationStateRef jobStepNotificationState in jobStepNotificationStates)
+			{
+				if (jobStepNotificationState.Badge == PendingBadge)
+				{
+					allNotificationsSent = false;
+					break;
+				}
+
+				// Check if the current step spawned another job that may still need to run it's Setup node
+				if (job.TryGetStep(jobStepNotificationState.JobStepId, out IJobStep? step) && step.SpawnedJobs is not null)
+				{
+					foreach (JobId spawnedJobId in step.SpawnedJobs)
+					{
+						IJob? spawnedJob = await _jobCollection.GetAsync(spawnedJobId, cancellationToken);
+						IJobStep? setupStep = GetSetupStepFromJob(spawnedJob);
+						if (setupStep is not null && setupStep.State == JobStepState.Ready)
+						{
+							allNotificationsSent = false;
+							break;
+						}
+					}
+				}
+			}
+
+			if (allNotificationsSent)
+			{
+				_ = await _jobNotificationCollection.DeleteJobStepNotificationStatesAsync(job.Id, job.TemplateId, recipient, cancellationToken);
+				_ = await _jobNotificationCollection.DeleteJobNotificationStatesAsync(job.Id, job.TemplateId, recipient, cancellationToken);
+			}
+		}
+
+		/// <summary>
+		/// Fetch the branch name from the stream
+		/// </summary>
+		/// <param name="buildConfig">Build plugin's global config</param>
+		/// <param name="streamId">Identifier of the stream</param>
+		/// <returns>Name of the corresponding branch or 'Unknown Branch' if no stream was found</returns>
+		static string GetBranchFromStream(BuildConfig buildConfig, StreamId streamId)
+		{
+			if (buildConfig.TryGetStream(streamId, out StreamConfig? streamConfig))
+			{
+				return streamConfig.Name;
+			}
+
+			return "Unknown Branch";
+		}
+
+		/// <summary>
+		/// Determines the badge based on the job step's state and outcome
+		/// </summary>
+		/// <param name="step">Information about the step</param>
+		/// <returns>Emoji representation of the job step's current state and outcome</returns>
+		static string GetBadgeForStep(IJobStep step)
+		{
+			// Check the state
+			switch (step.State)
+			{
+				case JobStepState.Waiting:
+				case JobStepState.Ready:
+				case JobStepState.Running:
+					{
+						// Step is either queued or currently being executed
+						return PendingBadge;
+					}
+				case JobStepState.Skipped:
+					{
+						return SkippedBadge;
+					}
+				case JobStepState.Aborted:
+					{
+						// Treat aborted as failure conditions
+						return FailureBadge;
+					}
+				case JobStepState.Completed:
+					{
+						// Continue to check the outcome of the job
+						break;
+					}
+				default:
+					{
+						return String.Empty;
+					}
+			}
+
+			// Check the outcome
+			switch (step.Outcome)
+			{
+				case JobStepOutcome.Failure:
+					{
+						return FailureBadge;
+					}
+				case JobStepOutcome.Warnings:
+					{
+						return WarningBadge;
+					}
+				case JobStepOutcome.Success:
+					{
+						return PassingBadge;
+					}
+				default:
+					{
+						return String.Empty;
+					}
+			}
+		}
+
+		/// <summary>
+		/// Fetch the setup job step
+		/// </summary>
+		/// <param name="job">Information about the job</param>
+		/// <returns>A <see cref="IJobStep"/> corressponding to the Setup</returns>
+		static IJobStep? GetSetupStepFromJob(IJob? job)
+		{
+			if (job is null)
+			{
+				return null;
+			}
+
+			foreach (IJobStepBatch batch in job.Batches)
+			{
+				foreach (IJobStep step in batch.Steps)
+				{
+					if (step.Name == IJob.SetupNodeName)
+					{
+						return step;
+					}
+				}
+			}
+
+			return null;
+		}
+
+		/// <summary>
+		/// Parses list of <see cref="ILogEventData"/> into the generic <see cref="LogData"/> structure
+		/// </summary>
+		/// <param name="step">Information about the step</param>
+		/// <param name="jobStepEventData">List of the step's logged events</param>
+		/// <returns>List of parsed logged events to be consumed by the Slack notifications</returns>
+		static List<LogData> ExtractLogInformation(IJobStep step, List<ILogEventData> jobStepEventData)
+		{
+			List<LogData> logData = new List<LogData>();
+			foreach (ILogEventData eventData in jobStepEventData)
+			{
+				logData.Add(new LogData(eventData.Severity, eventData.Message, step.LogId));
+			}
+			return logData;
+		}
+
+		#endregion Job Processing
+
+		#region Slack
+
+		/// <summary>
+		/// Create the job's summary message
+		/// </summary>
+		/// <param name="job">Information about the job</param>
+		/// <param name="branch">Name of the branch from the stream</param>
+		/// <param name="recipient">Identifier of the channel or user in Slack or user's email address</param>
+		/// <param name="parentJob">Optional information regarding a parent job</param>
+		/// <param name="abortedByUserName">Optional username provided when cancelling a job</param>
+		/// <param name="cancellationToken">Cancellation token for the operation</param>
+		/// <returns>Generated <see cref="SlackMessage"/> if successful, otherwise returns null</returns>
+		async Task<SlackMessage?> GenerateSlackSummaryMessageAsync(IJob job, string branch, string recipient, IJob? parentJob = null, string? abortedByUserName = null, CancellationToken cancellationToken = default)
+		{
+			IReadOnlyList<IJobStepNotificationStateRef>? jobStepNotificationStateDocuments;
+
+			// Fetch all of our notification information of any parent/spawned jobs
+			if (parentJob is not null)
+			{
+				IJobStepNotificationStateQueryBuilder parentQueryBuilder = JobNotificationCollection.CreateJobStepNotificationStateQueryBuilder();
+				parentQueryBuilder.AddJobAndParentJobFilter(parentJob.Id);
+				parentQueryBuilder.AddTemplateAndParentTemplateFilter(parentJob.TemplateId);
+				parentQueryBuilder.AddRecipientFilter(recipient);
+				jobStepNotificationStateDocuments = await _jobNotificationCollection.GetJobStepNotificationStatesAsync(parentQueryBuilder, cancellationToken);
+				if (jobStepNotificationStateDocuments is not null)
+				{
+					return await CreateSlackSummaryNotificationAsync(parentJob, branch, jobStepNotificationStateDocuments, abortedByUserName, cancellationToken);
+				}
+			}
+
+			// If we still have no notifications then fetch against our main job
+			IJobStepNotificationStateQueryBuilder builder = JobNotificationCollection.CreateJobStepNotificationStateQueryBuilder();
+			builder.AddJobAndParentJobFilter(job.Id);
+			builder.AddTemplateAndParentTemplateFilter(job.TemplateId);
+			builder.AddRecipientFilter(recipient);
+			jobStepNotificationStateDocuments = await _jobNotificationCollection.GetJobStepNotificationStatesAsync(builder, cancellationToken);
+			if (jobStepNotificationStateDocuments is not null)
+			{
+				return await CreateSlackSummaryNotificationAsync(job, branch, jobStepNotificationStateDocuments, abortedByUserName, cancellationToken);
+			}
+
+			return null;
+		}
+
+		/// <summary>
+		/// Send a notification on a job setup failure
+		/// </summary>
+		/// <param name="job">Information about the job</param>
+		/// <param name="step">Information about the step</param>
+		/// <param name="branch">Name of the branch from the stream</param>
+		/// <param name="jobStepEventData">List of the step's logged events</param>
+		/// <param name="jobNotificationConfigs">Collection of configurations tied to this job</param>
+		/// <param name="cancellationToken">Cancellation token for the operation</param>
+		async Task SendSetupFailedSlackMessageAsync(IJob job, IJobStep step, string branch, List<LogData> jobStepEventData, HashSet<JobNotificationConfig> jobNotificationConfigs, CancellationToken cancellationToken = default)
+		{
+			if (_slackClient is null)
+			{
+				_logger.LogInformation("Cannot send setup failed message. No Slack client set.");
+				return;
+			}
+
+			SlackMessage message = new SlackMessage();
+
+			string jobUrlLink = new Uri($"{_serverInfo.DashboardUrl}job/{job.Id}").ToString();
+			string jobStepLink = new Uri($"{jobUrlLink}?step={step.Id}").ToString();
+
+			// Create our notification text
+			string notificationText = GenerateSlackMessageTitle(job, branch);
+			message.Text = $"{notificationText} Details...";
+
+			// Create a Block for the Slack Title
+			string messageTitle = GenerateSlackMessageTitle(job, branch, jobUrlLink);
+			ContextBlock jobContext = new ContextBlock();
+			jobContext.Elements.Add(new TextObject(messageTitle));
+			message.Blocks.Add(jobContext);
+
+			ContextBlock infoBlock = new ContextBlock();
+			infoBlock.Elements.Add(new TextObject($"<{jobStepLink}|*{step.Name}*> encountered errors"));
+			message.Blocks.Add(infoBlock);
+
+			AddLogDataContext(message, jobStepEventData, step.LogId);
+
+			foreach (JobNotificationConfig config in jobNotificationConfigs)
+			{
+				foreach (string recipient in config.Channels)
+				{
+					// This is our first time sending the summary grid, add both the job and the job steps to the collection
+					string channel = recipient;
+					SlackUser? userInfo = await _slackClient.FindUserByEmailAsync(recipient, cancellationToken);
+					if (userInfo is not null && userInfo.Id is not null)
+					{
+						channel = userInfo.Id;
+					}
+					
+					try
+					{
+						_ = await _slackClient.PostMessageAsync(channel, message, cancellationToken);
+					}
+					catch (Exception ex)
+					{
+						_logger.LogInformation(ex, "Failed to send Slack setup failure message for job '{JobId}' and the template '{Template}' ({Error})", job.Id, job.TemplateId, ex.Message);
+					}
+				}
+			}
+
+			_logger.LogInformation("Finished sending failing setup notification for job '{JobId}'", job.Id);
+		}
+
+		/// <summary>
+		/// Send a threaded message on a job setup failure
+		/// </summary>
+		/// <param name="job">Information about the job</param>
+		/// <param name="step">Information about the step</param>
+		/// <param name="jobStepEventData">List of the step's logged events</param>
+		/// <param name="channel">Identifier of the Slack channel for the threaded message</param>
+		/// <param name="timestamp">Timestamp of the initial notification for the threaded message</param>
+		/// <param name="cancellationToken">Cancellation token for the operation</param>
+		async Task SendSetupFailedSlackMessageToThreadAsync(IJob job, IJobStep step, List<LogData> jobStepEventData, string channel, string timestamp, CancellationToken cancellationToken = default)
+		{
+			if (_slackClient is null)
+			{
+				_logger.LogInformation("Cannot send setup failed thread message. No Slack client set.");
+				return;
+			}
+
+			string badge = GetBadgeForStep(step);
+			string headerOverride = $"{job.Name} - {step.Name}";
+			SlackMessage? threadMessage = GenerateSlackThreadMessage(step, badge, jobStepEventData, headerOverride);
+			if (threadMessage is null)
+			{
+				_logger.LogInformation("Failed to generate thread message for the failed job step '{JobStepId}' with Job Id '{JobId}'", step.Id, job.Id);
+				return;
+			}
+
+			try
+			{
+				SlackMessageId jobMessageId = new SlackMessageId(channel, null, timestamp);
+				_ = await _slackClient.PostMessageToThreadAsync(jobMessageId, threadMessage, cancellationToken);
+			
+				_logger.LogInformation("Finished sending threaded notification for the failed job step '{JobStepId}' with Job Id '{JobId}'", step.Id, job.Id);
+			}
+			catch (Exception ex)
+			{
+				_logger.LogInformation(ex, "Failed to send Slack threaded message for the failed job step '{JobStepId}' with Job Id '{JobId}' ({Error})", step.Id, job.Id, ex.Message);
+			}
+
+			_logger.LogInformation("Finished sending failing setup thread notification for job '{JobId}'", job.Id);
+		}
+
+		/// <summary>
+		/// Sends a Slack notification for the specified job
+		/// </summary>
+		/// <param name="job">Information about the job</param>
+		/// <param name="branch">Name of the branch from the stream</param>
+		/// <param name="recipient">Identifier of the channel or user in Slack or user's email address</param>
+		/// <param name="jobNotificationState">Optional information regarding a previously sent job notification</param>
+		/// <param name="parentJob">Optional information regarding a parent job</param>
+		/// <param name="abortedByUserName">Optional username provided when cancelling a job</param>
+		/// <param name="cancellationToken">Cancellation token for the operation</param>
+		async Task SendSlackSummaryMessageAsync(IJob job, string branch, string recipient, IJobNotificationStateRef? jobNotificationState = null, IJob? parentJob = null, string? abortedByUserName = null, CancellationToken cancellationToken = default)
+		{
+			SlackMessage? message = await GenerateSlackSummaryMessageAsync(job, branch, recipient, parentJob, abortedByUserName, cancellationToken);
+			if (message is null)
+			{
+				_logger.LogInformation("Unable to generate Job summary notification for job '{JobId}' and the template '{TemplateId}'", job.Id, job.TemplateId);
+				return;
+			}
+
+			if (jobNotificationState is not null)
+			{
+				if (_slackClient is null)
+				{
+					_logger.LogInformation("Cannot update summary message. No Slack client set.");
+					return;
+				}
+
+				try
+				{
+					SlackMessageId messageId = new SlackMessageId(jobNotificationState.Channel, null, jobNotificationState.Ts);
+					_ = _slackClient.UpdateMessageAsync(messageId, message, cancellationToken);
+				}
+				catch (Exception ex)
+				{
+					_logger.LogInformation(ex, "Failed to update Slack summary message for job '{JobId}' and the template '{Template}' ({Error})", jobNotificationState.JobId, jobNotificationState.TemplateId, ex.Message);
+				}
+				return;
+			}
+			
+			if (_slackClient is null)
+			{
+				_logger.LogInformation("Cannot send summary message. No Slack client set.");
+				_ = await _jobNotificationCollection.AddOrUpdateJobNotificationStateAsync(job.Id, job.TemplateId, recipient, recipient, String.Empty, cancellationToken);
+				return;
+			}
+			
+			// This is our first time sending the summary grid, add both the job and the job steps to the collection
+			string channel = recipient;
+			SlackUser? userInfo = await _slackClient.FindUserByEmailAsync(recipient, cancellationToken);
+			if (userInfo is not null && userInfo.Id is not null)
+			{
+				channel = userInfo.Id;
+			}
+				
+			try
+			{
+				SlackMessageId messageId = await _slackClient.PostMessageAsync(channel, message, cancellationToken);
+				_ = await _jobNotificationCollection.AddOrUpdateJobNotificationStateAsync(job.Id, job.TemplateId, recipient, messageId.Channel, messageId.Ts, cancellationToken);
+			}
+			catch (Exception ex)
+			{
+				_logger.LogInformation(ex, "Failed to send Slack summary message for job '{JobId}' and the template '{Template}' ({Error})", job.Id, job.TemplateId, ex.Message);
+			}
+		}
+
+		/// <summary>
+		/// Sends a Slack message to a new or existing thread
+		/// </summary>
+		/// <param name="jobMessageId"><see cref="SlackMessageId"/> of the summary message</param>
+		/// <param name="jobStepNotificationState">Information about the job step's notification state</param>
+		/// <param name="threadMessage"><see cref="SlackMessage"/> to be sent</param>
+		/// <param name="cancellationToken">Cancellation token for the operation</param>
+		/// <returns>The <see cref="SlackMessageId"/> of either the newly sent notification or from the <see cref="IJobStepNotificationStateRef"/></returns>
+		async Task<SlackMessageId?> SendSlackThreadMessageAsync(SlackMessageId jobMessageId, IJobStepNotificationStateRef jobStepNotificationState, SlackMessage threadMessage, CancellationToken cancellationToken = default)
+		{
+			// Check to see if a message has already been sent out to avoid sending duplicate notifications
+			if (!String.IsNullOrEmpty(jobStepNotificationState.ThreadTs))
+			{
+				return new SlackMessageId(jobStepNotificationState.Channel, jobStepNotificationState.ThreadTs, jobStepNotificationState.Ts);
+			}
+
+			if (_slackClient is null)
+			{
+				_logger.LogInformation("Cannot send thread message. No Slack client set.");
+				return null;
+			}
+
+			try
+			{
+				SlackMessageId threadMessageId = await _slackClient.PostMessageToThreadAsync(jobMessageId, threadMessage, cancellationToken);
+			
+				_logger.LogInformation("Finished sending notification for job step '{JobStepId}' with Job Id '{JobId}' for template '{Template}'", jobStepNotificationState.JobStepId, jobStepNotificationState.JobId, jobStepNotificationState.TemplateId);
+				return threadMessageId;
+			}
+			catch (Exception ex)
+			{
+				_logger.LogInformation(ex, "Failed to send Slack message for job step '{JobStepId}' with Job Id '{JobId}' for template '{Template}' ({Error})", jobStepNotificationState.JobStepId, jobStepNotificationState.JobId, jobStepNotificationState.TemplateId, ex.Message);
+			}
+
+			return null;
+		}
+
+		/// <summary>
+		/// Sends a Slack notification for the specified job step
+		/// </summary>
+		/// <param name="job">Information about the job</param>
+		/// <param name="step">Information about the step</param>
+		/// <param name="jobStepEventData">List of the step's logged events</param>
+		/// <param name="underlyingStep">Optional information about the step's dependent step</param>
+		/// <param name="cancellationToken">Cancellation token for the operation</param>
+		/// <returns>True if a job step notification needed to be updated, otherwise returns false</returns>
+		async Task<bool> UpdateAndSendSlackNotificationsForStepAsync(IJob job, IJobStep step, List<LogData> jobStepEventData, IJobStep? underlyingStep = null, CancellationToken cancellationToken = default)
+		{
+			// Update our job step if it's being monitored and send an updated summary and post to the thread if needed
+			IJobStepNotificationStateQueryBuilder builder = JobNotificationCollection.CreateJobStepNotificationStateQueryBuilder();
+			builder.AddJobFilter(job.Id);
+			builder.AddJobStepFilter(step.Id);
+			IReadOnlyList<IJobStepNotificationStateRef>? jobStepNotificationStates = await _jobNotificationCollection.GetJobStepNotificationStatesAsync(builder, cancellationToken);
+			if (jobStepNotificationStates is null || jobStepNotificationStates.Count == 0)
+			{
+				_logger.LogInformation("No notification to be sent for step '{JobStepId}' in job '{JobId}'", step.Id, job.Id);
+				return false;
+			}
+
+			string badge = GetBadgeForStep(step);
+			foreach (IJobStepNotificationStateRef jobStepNotificationState in jobStepNotificationStates)
+			{
+				string channel = jobStepNotificationState.Channel;
+				string timestamp = jobStepNotificationState.Ts;
+				string? threadTimestamp = null;
+
+				SlackMessage? threadMessage = null;
+
+				// Notify on any problems of the Job Step
+				if (step.Outcome != JobStepOutcome.Success)
+				{
+					string? headerOverride = null;
+					if (jobStepNotificationState.ParentJobId is not null && jobStepNotificationState.ParentJobTemplateId is not null)
+					{
+						// Our current step has a parent job notification associated with it so we need to provide the header override to differentiate the job from the parent
+						headerOverride = $"{job.Name} - {step.Name}";
+					}
+					threadMessage = GenerateSlackThreadMessage(step, badge, jobStepEventData, headerOverride, underlyingStep);
+				}
+
+				if (threadMessage is not null)
+				{
+					// Fetch our initial job notification to get the Message Id for threads
+					JobId notificationJobId = jobStepNotificationState.ParentJobId ?? jobStepNotificationState.JobId;
+					TemplateId notificationTemplateId = jobStepNotificationState.ParentJobTemplateId ?? jobStepNotificationState.TemplateId;
+					IJobNotificationStateRef? jobNotificationState = await _jobNotificationCollection.GetJobNotificationStateAsync(notificationJobId, notificationTemplateId, jobStepNotificationState.Recipient, cancellationToken);
+					if (jobNotificationState is null)
+					{
+						_logger.LogInformation("No message id found for job '{JobId}' with template '{TemplateId}' for thread to be posted.", job.Id, jobStepNotificationState.TemplateId);
+						continue;
+					}
+					SlackMessageId jobMessageId = new SlackMessageId(jobNotificationState.Channel, null, jobNotificationState.Ts);
+					SlackMessageId? threadMessageId = await SendSlackThreadMessageAsync(jobMessageId, jobStepNotificationState, threadMessage, cancellationToken);
+					if (threadMessageId is not null)
+					{
+						channel = threadMessageId.Channel;
+						threadTimestamp = threadMessageId.ThreadTs;
+						timestamp = threadMessageId.Ts;
+					}
+				}
+				_ = await _jobNotificationCollection.AddOrUpdateJobStepNotificationStateAsync(jobStepNotificationState.JobId, jobStepNotificationState.TemplateId, jobStepNotificationState.Recipient, jobStepNotificationState.JobStepId, jobStepNotificationState.Group, jobStepNotificationState.TargetPlatform, badge, channel, timestamp, threadTimestamp, jobStepNotificationState.ParentJobId, jobStepNotificationState.ParentJobTemplateId, cancellationToken);
+			}
+
+			return true;
+		}
+
+		/// <summary>
+		/// Create the job's summary message
+		/// </summary>
+		/// <param name="job">Information about the job</param>
+		/// <param name="branch">Name of the branch from the stream</param>
+		/// <param name="jobStepNotificationStateDocuments">List of <see cref="IJobStepNotificationStateRef"/></param>
+		/// <param name="abortedByUserName">Optional username provided when cancelling a job</param>
+		/// <param name="cancellationToken">Cancellation token for the operation</param>
+		/// <returns>Generated <see cref="SlackMessage"/></returns>
+		async Task<SlackMessage> CreateSlackSummaryNotificationAsync(IJob job, string branch, IReadOnlyList<IJobStepNotificationStateRef> jobStepNotificationStateDocuments, string? abortedByUserName = null, CancellationToken cancellationToken = default)
+		{
+			DividerBlock divider = new DividerBlock();
+			SlackMessage message = new SlackMessage();
+
+			// Create our notification text
+			string notificationText = GenerateSlackMessageTitle(job, branch);
+			message.Text = $"{notificationText} Details...";
+
+			// Create a Block for the Slack Title
+			string hordeJobURL = new Uri(_serverInfo.DashboardUrl, $"job/{job.Id}").ToString();
+			string messageTitle = GenerateSlackMessageTitle(job, branch, hordeJobURL);
+			ContextBlock jobContext = new ContextBlock();
+			jobContext.Elements.Add(new TextObject(messageTitle));
+			message.Blocks.Add(jobContext);
+
+			// If the job was manually aborted, then provide a message that catches the recipients eyes that the full information may not be available
+			if (!String.IsNullOrEmpty(abortedByUserName))
+			{
+				AddCancelledJobSlackContext(message, $"*{AlertBadge} Job was canceled by {abortedByUserName} and may be missing information {AlertBadge}*", job.CancellationReason);
+				message.Blocks.Add(divider);
+			}
+
+			int itemsInGroup = 0;
+			string currentGroup = String.Empty;
+			ContextBlock? header = null;
+			ContextBlock? stepContext = null;
+			List<IJob> pendingSpawnedJobs = new List<IJob>();
+			List<IJob> canceledSpawnedJobs = new List<IJob>();
+			List<IJob> errorSpawnedJobs = new List<IJob>();
+			foreach (IJobStepNotificationStateRef state in jobStepNotificationStateDocuments)
+			{
+				// Nothing to report on if the step is no longer in the job
+				IJob? notificationJob = await _jobCollection.GetAsync(state.JobId, cancellationToken);
+				if (notificationJob is null || !notificationJob.TryGetStep(state.JobStepId, out IJobStep? step))
+				{
+					continue;
+				}
+
+				if (currentGroup != state.Group || header is null)
+				{
+					// Check to see if we have an existing step context that needs to be added
+					if (stepContext is not null)
+					{
+						message.Blocks.Add(stepContext);
+						message.Blocks.Add(divider);
+						stepContext = null;
+					}
+
+					currentGroup = state.Group;
+					header = new ContextBlock();
+					header.Elements.Add(new TextObject($"*{currentGroup}*"));
+					message.Blocks.Add(header);
+				}
+
+				string hordeJobStepURL = new Uri(_serverInfo.DashboardUrl, $"job/{state.JobId}?step={state.JobStepId}").ToString();
+
+				// Slack can only allow for 10 elements per block (https://api.slack.com/reference/block-kit/blocks#context) and we want to track when to create new ContextBlocks
+				bool bNeedsNewContext = (itemsInGroup % MaxSlackBlocks == 0 || stepContext is null);
+				if (bNeedsNewContext)
+				{
+					if (stepContext is not null)
+					{
+						message.Blocks.Add(stepContext);
+					}
+					stepContext = new ContextBlock();
+				}
+				stepContext?.Elements.Add(new TextObject($"<{hordeJobStepURL}|{state.Badge} {state.TargetPlatform}>"));
+
+				// The StepContext should always be valid here, but to disable the warning, we provide a default so that we only have 1 block per context.
+				itemsInGroup = stepContext?.Elements.Count ?? MaxSlackBlocks;
+
+				// Check to see if our step spawned any jobs
+				if (step.SpawnedJobs is null)
+				{
+					continue;
+				}
+
+				foreach (JobId spawnedJobId in step.SpawnedJobs)
+				{
+					// Fetch the job and see if the jobs are queued to run, were canceled, or encountered issues during the setup build
+					IJob? spawnedJob = await _jobCollection.GetAsync(spawnedJobId, cancellationToken);
+					if (spawnedJob is null)
+					{
+						continue;
+					}
+
+					if (spawnedJob.AbortedByUserId is not null)
+					{
+						canceledSpawnedJobs.Add(spawnedJob);
+						continue;
+					}
+
+					IJobStep? setupStep = GetSetupStepFromJob(spawnedJob);
+					if (setupStep is null || setupStep.Outcome == JobStepOutcome.Failure)
+					{
+						errorSpawnedJobs.Add(spawnedJob);
+					}
+					else if (setupStep.State == JobStepState.Aborted)
+					{
+						canceledSpawnedJobs.Add(spawnedJob);
+					}
+					else if (setupStep.State != JobStepState.Completed)
+					{
+						pendingSpawnedJobs.Add(spawnedJob);
+					}
+				}
+			}
+
+			// Check to see if we have an existing step context that needs to be added
+			if (stepContext is not null)
+			{
+				message.Blocks.Add(stepContext);
+				message.Blocks.Add(divider);
+			}
+
+			// Generate our summary to include information regarding the spawned jobs only if the current job wasn't canceled
+			// Reason being is that the parent job will no longer have information to handle updates
+			if (job.AbortedByUserId is null)
+			{
+				TryAddSpawnedJobSlackContext(message, $"*Queued jobs:*", pendingSpawnedJobs);
+				TryAddSpawnedJobSlackContext(message, $"*Canceled jobs:*", canceledSpawnedJobs);
+				TryAddSpawnedJobSlackContext(message, $"*Failed jobs:*", errorSpawnedJobs);
+			}
+
+			return message;
+		}
+
+		/// <summary>
+		/// Adds a context to the Slack message with log details
+		/// </summary>
+		/// <param name="slackBlockContainer">Slack block container</param>
+		/// <param name="logData">List of our log information</param>
+		/// <param name="logId">Optional identifier of the log</param>
+		void AddLogDataContext(ISlackBlockContainer slackBlockContainer, List<LogData> logData, LogId? logId = null)
+		{
+			// Don't clutter the message with all of the events and do a preliminary check to see if the user needs to be made aware that events have been truncated
+			if (logData.Count > MaxEventMessages)
+			{
+				ContextBlock limitContext = new ContextBlock();
+				limitContext.Elements.Add(new TextObject($":warning: *Step has multiple events reported. Will only show first 5.* :warning:"));
+				slackBlockContainer.Blocks.Add(limitContext);
+			}
+
+			foreach (LogData log in logData.Take(MaxEventMessages))
+			{
+				SectionBlock logSection = new SectionBlock(new TextObject(SanitizeQuoteText(log.Message)));
+				slackBlockContainer.Blocks.Add(logSection);
+			}
+
+			if (logId is not null)
+			{
+				string logUrl = new Uri(_serverInfo.DashboardUrl, $"log/{logId}").ToString();
+
+				ContextBlock detailsContext = new ContextBlock();
+				detailsContext.Elements.Add(new TextObject($"<{logUrl}|View Job Step Log>"));
+				slackBlockContainer.Blocks.Add(detailsContext);
+			}
+		}
+
+		/// <summary>
+		/// Create the job step informational threaded message
+		/// </summary>
+		/// <param name="step">Information about the step</param>
+		/// <param name="badge">Slack emoji associated with the notification</param>
+		/// <param name="logData">List of our log information</param>
+		/// <param name="headerOverride">Optional string header to be used within the message</param>
+		/// <param name="underlyingStep">Optional information about the step's dependent step</param>
+		/// <returns>Generated <see cref="SlackMessage"/> if successful, otherwise returns null</returns>
+		SlackMessage? GenerateSlackThreadMessage(IJobStep step, string badge, List<LogData> logData, string? headerOverride = null, IJobStep? underlyingStep = null)
+		{
+			if (logData.Count == 0)
+			{
+				return null;
+			}
+
+			List<LogData> data = logData.Where(x => x.Severity == LogEventSeverity.Error).ToList();
+			if (!data.Any())
+			{
+				data = logData.Where(x => x.Severity == LogEventSeverity.Warning).ToList();
+			}
+
+			string jobStepLink = new Uri($"{_serverInfo.DashboardUrl}job/{step.Job.Id}?step={step.Id}").ToString();
+			string alertInfo = $"{step.Name}";
+			if (!String.IsNullOrEmpty(headerOverride))
+			{
+				alertInfo = headerOverride;
+			}
+			string alertHeader = $"<{jobStepLink}|*{alertInfo}*>";
+
+			SlackMessage alertMessage = new SlackMessage();
+
+			ContextBlock alertStepBlock = new ContextBlock();
+			alertStepBlock.Elements.Add(new PlainTextObject(badge, true));
+			alertStepBlock.Elements.Add(new TextObject(alertHeader));
+			alertStepBlock.Elements.Add(new PlainTextObject(badge, true));
+			alertMessage.Blocks.Add(alertStepBlock);
+
+			if (underlyingStep is null)
+			{
+				AddLogDataContext(alertMessage, data, step.LogId);
+				return alertMessage;
+			}
+
+			// If an underlying step was provided then the error occurred further up the job and impacted the current step
+			ContextBlock dependencyContext = new ContextBlock();
+			dependencyContext.Elements.Add(new TextObject("*Step was marked as skipped due to errors in the following dependency*"));
+			alertMessage.Blocks.Add(dependencyContext);
+
+			// Slack Attachments are marked as being legacy and is recommended to be avoided
+			// The reason why Attachments are used is because this information is secondary to describe why a test had skipped
+			// We also wanted to differentiate the dependencies from the actual monitored node by indenting the dependency information
+			SlackAttachment attachment = new SlackAttachment();
+			attachment.Color = "#666666";
+
+			string dependencyJobStepLink = new Uri($"{_serverInfo.DashboardUrl}job/{step.Job.Id}?step={underlyingStep.Id}").ToString();
+
+			attachment.AddContext(new TextObject($"<{dependencyJobStepLink}|*{underlyingStep.Name}*>"));
+			AddLogDataContext(attachment, data, underlyingStep.LogId);
+			alertMessage.Attachments.Add(attachment);
+
+			return alertMessage;
+		}
+
+		/// <summary>
+		/// Attempts to add a Slack context to the <see cref="SlackMessage"/>
+		/// </summary>
+		/// <param name="message"><see cref="SlackMessage"/> to add the context to</param>
+		/// <param name="header">String header of the context</param>
+		/// <param name="spawnedJobs">List of spawned jobs for this context</param>
+		void TryAddSpawnedJobSlackContext(SlackMessage message, string header, List<IJob> spawnedJobs)
+		{
+			if (!spawnedJobs.Any())
+			{
+				return;
+			}
+
+			ContextBlock? headerContext = new ContextBlock();
+			headerContext.Elements.Add(new TextObject(header));
+			message.Blocks.Add(headerContext);
+
+			// Create a Block for the Slack Title
+			int itemsInGroup = 0;
+			ContextBlock? spawnJobContext = null;
+			foreach (IJob job in spawnedJobs)
+			{
+				string hordeJobURL = new Uri(_serverInfo.DashboardUrl, $"job/{job.Id}").ToString();
+
+				// Slack can only allow for 10 elements per block (https://api.slack.com/reference/block-kit/blocks#context) and we want to track when to create new ContextBlocks
+				bool bNeedsNewContext = (itemsInGroup % MaxSlackBlocks == 0 || spawnJobContext is null);
+				if (bNeedsNewContext)
+				{
+					if (spawnJobContext is not null)
+					{
+						message.Blocks.Add(spawnJobContext);
+					}
+					spawnJobContext = new ContextBlock();
+				}
+				spawnJobContext?.Elements.Add(new TextObject($"<{hordeJobURL}|{job.Name}>"));
+
+				// The StepContext should always be valid here, but to disable the warning, we provide a default so that we only have 1 block per context.
+				itemsInGroup = spawnJobContext?.Elements.Count ?? MaxSlackBlocks;
+			}
+
+			if (spawnJobContext is null)
+			{
+				// Log an error and return
+				return;
+			}
+
+			message.Blocks.Add(spawnJobContext);
+		}
+
+		/// <summary>
+		/// Generates the title for the summary messages
+		/// </summary>
+		/// <param name="job">Information about the job</param>
+		/// <param name="branch">Name of the branch from the stream</param>
+		/// <param name="hordeJobURL">Optional URL of the job in Horde</param>
+		/// <returns>Title message using Slack's markdown format</returns>
+		static string GenerateSlackMessageTitle(IJob job, string branch, string? hordeJobURL = null)
+		{
+			string changelistDetails = $"CL {job.CommitId}";
+			if (job.PreflightCommitId is not null)
+			{
+				changelistDetails = $"Preflight CL {job.PreflightCommitId} against {changelistDetails}";
+			}
+
+			if (!String.IsNullOrEmpty(hordeJobURL))
+			{
+				return $"<{hordeJobURL}|*{branch} - {changelistDetails} - {job.Name}*>";
+			}
+
+			return $"{branch} - {changelistDetails} - {job.Name}";
+		}
+
+		/// <summary>
+		/// Sanitizes the text for a Slack quote message
+		/// </summary>
+		/// <param name="text">Text to sanitize</param>
+		/// <param name="maxLength">Maximum allowed length for the message</param>
+		static string SanitizeQuoteText(string text, int maxLength = MaxBlockCharacterLength)
+		{
+			if (text.Length > maxLength)
+			{
+				text = text.Substring(0, maxLength - 4) + "...";
+			}
+			return $"```{text}```";
+		}
+
+		/// <summary>
+		/// Adds a context to the Slack message with cancellation details
+		/// </summary>
+		/// <param name="message"><see cref="SlackMessage"/> to add the context to</param>
+		/// <param name="text">Information to provide in the context</param>
+		/// <param name="cancellationReason">Optional message explaining the cancellation</param>
+		static void AddCancelledJobSlackContext(SlackMessage message, string text, string? cancellationReason = null)
+		{
+			ContextBlock context = new ContextBlock();
+			context.Elements.Add(new TextObject(text));
+			message.Blocks.Add(context);
+
+			if (!String.IsNullOrEmpty(cancellationReason))
+			{
+				ContextBlock cancellationHeader = new ContextBlock();
+				cancellationHeader.Elements.Add(new TextObject($"*Cancellation Reason*"));
+				message.Blocks.Add(cancellationHeader);
+
+				ContextBlock cancellationContext = new ContextBlock();
+				cancellationContext.Elements.Add(new TextObject(cancellationReason));
+				message.Blocks.Add(cancellationContext);
+			}
+		}
+
+		#endregion Slack
+	}
+}
