Horde: Add Show Telemetry option to view on Conforms

diff --git a/Engine/Source/Programs/Horde/HordeDashboard/src/components/LogSource.tsx b/Engine/Source/Programs/Horde/HordeDashboard/src/components/LogSource.tsx
--- a/Engine/Source/Programs/Horde/HordeDashboard/src/components/LogSource.tsx
+++ b/Engine/Source/Programs/Horde/HordeDashboard/src/components/LogSource.tsx
@@ -624,6 +624,10 @@ class LeaseLogSource extends LogSource {
       this.leaseId = leaseId;
    }
 
+   get agentTelemetry(): AgentTelemetryHandler | undefined {
+      return this._agentTelemety;
+   }
+
    getAgentId(): string | undefined {
       return this.lease?.agentId;
    }
@@ -705,6 +709,8 @@ class LeaseLogSource extends LogSource {
       if (this.active !== active) {
          this.setActive(active);
       }
+
+      this.agentTelemetry?.set(this.lease?.agentId ?? "", new Date(this.startTime as any), this.lease?.finishTime ? new Date(this.lease.finishTime) : undefined);
    }
 
    get crumbs(): BreadcrumbItem[] {
@@ -734,5 +740,6 @@ class LeaseLogSource extends LogSource {
    agent?: AgentData;
    lease?: LeaseData;
 
+   _agentTelemety: AgentTelemetryHandler = new AgentTelemetryHandler();
 
 }
\ No newline at end of file
