Add CV link to Slack message

diff --git a/Engine/Source/Programs/Horde/Docs/Deployment/ServerSettings.md b/Engine/Source/Programs/Horde/Docs/Deployment/ServerSettings.md
--- a/Engine/Source/Programs/Horde/Docs/Deployment/ServerSettings.md
+++ b/Engine/Source/Programs/Horde/Docs/Deployment/ServerSettings.md
@@ -185,7 +185,8 @@ Name | Description
 `perforceConnectionPoolSize` | `integer`<br>Number of pooled perforce connections to keep
 `enableConformTasks` | `boolean`<br>Whether to enable the conform task source.
 `p4SwarmUrl` | `string`<br>Url of P4 Swarm installation
-`robomergeUrl` | `string`<br>Url of Robomergem installation
+`robomergeUrl` | `string`<br>Url of Robomerge installation
+`commitsViewerUrl` | `string`<br>Url of Commits Viewer
 `jiraUsername` | `string`<br>The Jira service account user name
 `jiraApiToken` | `string`<br>The Jira service account API token
 `jiraUrl` | `string`<br>The Uri for the Jira installation
diff --git a/Engine/Source/Programs/Horde/Plugins/Build/HordeServer.Build/BuildResponseFilter.cs b/Engine/Source/Programs/Horde/Plugins/Build/HordeServer.Build/BuildResponseFilter.cs
--- a/Engine/Source/Programs/Horde/Plugins/Build/HordeServer.Build/BuildResponseFilter.cs
+++ b/Engine/Source/Programs/Horde/Plugins/Build/HordeServer.Build/BuildResponseFilter.cs
@@ -37,6 +37,11 @@ namespace HordeServer
 					dashboardConfigResponse.RobomergeUrl = _staticBuildConfig.RobomergeUrl.ToString().TrimEnd('/');
 				}
 
+				if (_staticBuildConfig.CommitsViewerUrl != null)
+				{
+					dashboardConfigResponse.CommitsViewerUrl = _staticBuildConfig.CommitsViewerUrl.ToString().TrimEnd('/');
+				}
+
 				dashboardConfigResponse.DeviceProblemCooldownMinutes = _staticBuildConfig.DeviceProblemCooldownMinutes;
 			}
 		}
diff --git a/Engine/Source/Programs/Horde/Plugins/Build/HordeServer.Build/BuildServerConfig.cs b/Engine/Source/Programs/Horde/Plugins/Build/HordeServer.Build/BuildServerConfig.cs
--- a/Engine/Source/Programs/Horde/Plugins/Build/HordeServer.Build/BuildServerConfig.cs
+++ b/Engine/Source/Programs/Horde/Plugins/Build/HordeServer.Build/BuildServerConfig.cs
@@ -45,6 +45,11 @@ namespace HordeServer
 		/// </summary>
 		public Uri? RobomergeUrl { get; set; }
 
+		/// <summary>
+		/// Url of Commits Viewer
+		/// </summary>
+		public Uri? CommitsViewerUrl { get; set; }
+
 		/// <summary>
 		/// The Jira service account user name
 		/// </summary>
diff --git a/Engine/Source/Programs/Horde/Plugins/Build/HordeServer.Build/Notifications/Sinks/SlackNotificationSink.cs b/Engine/Source/Programs/Horde/Plugins/Build/HordeServer.Build/Notifications/Sinks/SlackNotificationSink.cs
--- a/Engine/Source/Programs/Horde/Plugins/Build/HordeServer.Build/Notifications/Sinks/SlackNotificationSink.cs
+++ b/Engine/Source/Programs/Horde/Plugins/Build/HordeServer.Build/Notifications/Sinks/SlackNotificationSink.cs
@@ -1282,6 +1282,12 @@ namespace HordeServer.Notifications.Sinks
 								suspectMessage += $"<{swarmLink}|View changes in Swarm>";
 							}
 
+							Uri? commitsViewerLink = GetCommitsViewerLinkForSpan(span);
+							if (commitsViewerLink != null)
+							{
+								suspectMessage += $"(<{commitsViewerLink}|View changes in Commits Viewer>)";
+							}
+
 							Uri? robomergeTrackChangesLink = GetRobomergeTrackChangesUrlForSpan(span);
 							if (robomergeTrackChangesLink != null)
 							{
@@ -1508,6 +1514,30 @@ namespace HordeServer.Notifications.Sinks
 			}
 		}
 
+		Uri? GetCommitsViewerLinkForSpan(IIssueSpan span)
+		{
+			if (_settings.CommitsViewerUrl == null || span.LastSuccess == null)
+			{
+				return null;
+			}
+			else
+			{
+				string depot_stream = span.StreamName.TrimStart('/');
+				string[] depot_stream_split = depot_stream.Split(new[] { '/' }, 2);
+
+				// If for some reason the stream is not split correctly for CV, return null
+				if (depot_stream_split.Length != 2)
+				{
+					return null;
+				}
+
+				string query_string = $"?depot={depot_stream_split[0]}&stream={depot_stream_split[1]}" +
+					$"&query_range[CL][start_cl]={span.LastSuccess.CommitId.GetPerforceChange() + 1}&query_range[CL][end_cl]={span.FirstFailure.CommitId}";
+
+				return new Uri(_settings.CommitsViewerUrl, query_string);
+			}
+		}
+
 		async Task PostSingleMessageToThreadAsync(string recipient, string eventId, SlackMessageId threadId, string message, CancellationToken cancellationToken)
 		{
 			(MessageStateDocument state, bool isNew) = await AddOrUpdateMessageStateAsync(recipient, eventId, null, "", null, cancellationToken);
diff --git a/Engine/Source/Programs/Shared/EpicGames.Horde/Dashboard/DashboardMessages.cs b/Engine/Source/Programs/Shared/EpicGames.Horde/Dashboard/DashboardMessages.cs
--- a/Engine/Source/Programs/Shared/EpicGames.Horde/Dashboard/DashboardMessages.cs
+++ b/Engine/Source/Programs/Shared/EpicGames.Horde/Dashboard/DashboardMessages.cs
@@ -35,6 +35,11 @@ namespace EpicGames.Horde.Dashboard
 		/// </summary>
 		public string? RobomergeUrl { get; set; }
 
+		/// <summary>
+		/// Url of Commits Viewer
+		/// </summary>
+		public string? CommitsViewerUrl { get; set; }
+
 		/// <summary>
 		/// Help email address that users can contact with issues
 		/// </summary>
