Fixing duplicated threaded message being sent during a step cancellation

#rnx

diff --git a/Engine/Source/Programs/Horde/Plugins/Experimental/HordeServer.Experimental/Notifications/IJobNotification.cs b/Engine/Source/Programs/Horde/Plugins/Experimental/HordeServer.Experimental/Notifications/IJobNotification.cs
--- a/Engine/Source/Programs/Horde/Plugins/Experimental/HordeServer.Experimental/Notifications/IJobNotification.cs
+++ b/Engine/Source/Programs/Horde/Plugins/Experimental/HordeServer.Experimental/Notifications/IJobNotification.cs
@@ -219,8 +219,9 @@ namespace HordeServer.Experimental.Notifications
 		/// Add the filter for a specific badge
 		/// </summary>
 		/// <param name="badge">Badge to filter by</param>
+		/// <param name="shouldMatch">Whether or not the filter should match the exact badge</param>
 		/// <returns>A reference to the current <see cref="IJobStepNotificationStateQueryBuilder"/></returns>
-		public IJobStepNotificationStateQueryBuilder AddBadgeFilter(string badge);
+		public IJobStepNotificationStateQueryBuilder AddBadgeFilter(string badge, bool shouldMatch);
 
 		/// <summary>
 		/// Add the filter for a specific channel
diff --git a/Engine/Source/Programs/Horde/Plugins/Experimental/HordeServer.Experimental/Notifications/JobNotificationCollection.cs b/Engine/Source/Programs/Horde/Plugins/Experimental/HordeServer.Experimental/Notifications/JobNotificationCollection.cs
--- a/Engine/Source/Programs/Horde/Plugins/Experimental/HordeServer.Experimental/Notifications/JobNotificationCollection.cs
+++ b/Engine/Source/Programs/Horde/Plugins/Experimental/HordeServer.Experimental/Notifications/JobNotificationCollection.cs
@@ -233,9 +233,16 @@ namespace HordeServer.Experimental.Notifications
 				return this;
 			}
 
-			public IJobStepNotificationStateQueryBuilder AddBadgeFilter(string badge)
-			{
-				Filter &= _filterBuilder.Eq(x => x.Badge, badge);
+			public IJobStepNotificationStateQueryBuilder AddBadgeFilter(string badge, bool shouldMatch)
+			{
+				if (shouldMatch)
+				{
+					Filter &= _filterBuilder.Eq(x => x.Badge, badge);
+				}
+				else
+				{
+					Filter &= _filterBuilder.Ne(x => x.Badge, badge);
+				}
 				return this;
 			}
 
diff --git a/Engine/Source/Programs/Horde/Plugins/Experimental/HordeServer.Experimental/Notifications/SlackNotificationProcessor.cs b/Engine/Source/Programs/Horde/Plugins/Experimental/HordeServer.Experimental/Notifications/SlackNotificationProcessor.cs
--- a/Engine/Source/Programs/Horde/Plugins/Experimental/HordeServer.Experimental/Notifications/SlackNotificationProcessor.cs
+++ b/Engine/Source/Programs/Horde/Plugins/Experimental/HordeServer.Experimental/Notifications/SlackNotificationProcessor.cs
@@ -643,7 +643,7 @@ namespace HordeServer.Experimental.Notifications
 			builder.AddJobAndParentJobFilter(jobId);
 			builder.AddTemplateAndParentTemplateFilter(templateId);
 			builder.AddRecipientFilter(recipient);
-			builder.AddBadgeFilter(PendingBadge);
+			builder.AddBadgeFilter(PendingBadge, true);
 			IReadOnlyList<IJobStepNotificationStateRef>? jobStepNotificationStates = await _jobNotificationCollection.GetJobStepNotificationStatesAsync(builder, cancellationToken);
 			if (jobStepNotificationStates is null)
 			{
@@ -1114,10 +1114,13 @@ namespace HordeServer.Experimental.Notifications
 		/// <returns>True if a job step notification needed to be updated, otherwise returns false</returns>
 		async Task<bool> UpdateAndSendSlackNotificationsForStepAsync(IJob job, IJobStep step, List<LogData> jobStepEventData, IJobStep? underlyingStep = null, CancellationToken cancellationToken = default)
 		{
+			string badge = GetBadgeForStep(step);
+
 			// Update our job step if it's being monitored and send an updated summary and post to the thread if needed
 			IJobStepNotificationStateQueryBuilder builder = JobNotificationCollection.CreateJobStepNotificationStateQueryBuilder();
 			builder.AddJobFilter(job.Id);
 			builder.AddJobStepFilter(step.Id);
+			builder.AddBadgeFilter(badge, false);
 			IReadOnlyList<IJobStepNotificationStateRef>? jobStepNotificationStates = await _jobNotificationCollection.GetJobStepNotificationStatesAsync(builder, cancellationToken);
 			if (jobStepNotificationStates is null || jobStepNotificationStates.Count == 0)
 			{
@@ -1125,7 +1128,6 @@ namespace HordeServer.Experimental.Notifications
 				return false;
 			}
 
-			string badge = GetBadgeForStep(step);
 			foreach (IJobStepNotificationStateRef jobStepNotificationState in jobStepNotificationStates)
 			{
 				string channel = jobStepNotificationState.Channel;
