Horde: Improve tracing spans in job scheduler

diff --git a/Engine/Source/Programs/Horde/Plugins/Build/HordeServer.Build/Jobs/Schedules/ScheduleService.cs b/Engine/Source/Programs/Horde/Plugins/Build/HordeServer.Build/Jobs/Schedules/ScheduleService.cs
--- a/Engine/Source/Programs/Horde/Plugins/Build/HordeServer.Build/Jobs/Schedules/ScheduleService.cs
+++ b/Engine/Source/Programs/Horde/Plugins/Build/HordeServer.Build/Jobs/Schedules/ScheduleService.cs
@@ -132,18 +132,25 @@ namespace HordeServer.Jobs.Schedules
 
 		async ValueTask TickAsync(CancellationToken cancellationToken)
 		{
+			using TelemetrySpan span = _tracer.StartActiveSpan($"{nameof(ScheduleService)}.{nameof(TickAsync)}");
 			DateTime utcNow = _clock.UtcNow;
 
 			// Don't start any new jobs during scheduled downtime
-			if (_downtimeService.IsDowntimeActive)
+			bool isDowntimeActive = _downtimeService.IsDowntimeActive;
+			if (isDowntimeActive)
 			{
+				span.SetAttribute("horde.is_downtime", isDowntimeActive);
 				return;
 			}
 
 			// Update the current queue
 			await using (RedisLock sharedLock = new(_redis.GetDatabase(), s_tickLockKey))
 			{
-				if (await sharedLock.AcquireAsync(TimeSpan.FromMinutes(1.0), false))
+				using TelemetrySpan tickLockSpan = _tracer.StartActiveSpan($"{nameof(ScheduleService)}.{nameof(TickAsync)}.TickLock");
+				bool tickLockAcquired = await sharedLock.AcquireAsync(TimeSpan.FromMinutes(1.0), false);
+				span.SetAttribute("horde.lock_acquired", tickLockAcquired);
+
+				if (tickLockAcquired)
 				{
 					await UpdateQueueAsync(utcNow, cancellationToken);
 				}
@@ -162,7 +169,12 @@ namespace HordeServer.Jobs.Schedules
 				// Acquire the lock for this schedule and update it
 				await using (RedisLock sharedLock = new RedisLock<QueueItem>(_redis.GetDatabase(), s_baseLockKey, item))
 				{
-					if (await sharedLock.AcquireAsync(TimeSpan.FromMinutes(1.0)))
+					using TelemetrySpan baseLockSpan = _tracer.StartActiveSpan($"{nameof(ScheduleService)}.{nameof(TickAsync)}.Tick");
+					bool baseLockAcquired = await sharedLock.AcquireAsync(TimeSpan.FromMinutes(1.0));
+					span.SetAttribute("horde.lock_acquired", baseLockAcquired);
+					span.SetAttribute("horde.stream_id", item.StreamId);
+					span.SetAttribute("horde.template_id", item.TemplateId);
+					if (baseLockAcquired)
 					{
 						try
 						{
@@ -216,6 +228,7 @@ namespace HordeServer.Jobs.Schedules
 		/// </summary>
 		public async Task UpdateQueueAsync(DateTime utcNow, CancellationToken cancellationToken)
 		{
+			using TelemetrySpan span = _tracer.StartActiveSpan($"{nameof(ScheduleService)}.{nameof(UpdateQueueAsync)}");
 			List<SortedSetEntry<QueueItem>> queueItems = new List<SortedSetEntry<QueueItem>>();
 
 			BuildConfig buildConfig = _buildConfig.CurrentValue;
@@ -293,8 +306,8 @@ namespace HordeServer.Jobs.Schedules
 			}
 
 			using TelemetrySpan span = _tracer.StartActiveSpan($"{nameof(ScheduleService)}.{nameof(TriggerAsync)}");
-			span.SetAttribute("streamId", stream.Id);
-			span.SetAttribute("templateId", templateId);
+			span.SetAttribute("horde.stream_id", stream.Id);
+			span.SetAttribute("horde.template_id", templateId);
 
 			Stopwatch stopwatch = Stopwatch.StartNew();
 			_logger.LogInformation("Updating schedule for {StreamId} template {TemplateId}", stream.Id, templateId);
@@ -349,6 +362,11 @@ namespace HordeServer.Jobs.Schedules
 		/// <returns>Async task</returns>
 		private async Task TriggerAsync(IStream stream, TemplateId templateId, ITemplateRef templateRef, ISchedule schedule, int numActiveJobs, DateTime utcNow, ILogger scheduleLogger, CancellationToken cancellationToken)
 		{
+			using TelemetrySpan span = _tracer.StartActiveSpan($"{nameof(ScheduleService)}.{nameof(TriggerAsync)}.Full");
+			span.SetAttribute("horde.stream_id", stream.Id);
+			span.SetAttribute("horde.template_id", templateId);
+			span.SetAttribute("horde.active_jobs", numActiveJobs);
+			
 			// Check we're not already at the maximum number of allowed jobs
 			if (schedule.MaxActive != 0 && numActiveJobs >= schedule.MaxActive)
 			{
