Fixing document update and logging

#rnx

diff --git a/Engine/Source/Programs/Horde/Plugins/Experimental/HordeServer.Experimental/Notifications/JobNotificationCollection.cs b/Engine/Source/Programs/Horde/Plugins/Experimental/HordeServer.Experimental/Notifications/JobNotificationCollection.cs
--- a/Engine/Source/Programs/Horde/Plugins/Experimental/HordeServer.Experimental/Notifications/JobNotificationCollection.cs
+++ b/Engine/Source/Programs/Horde/Plugins/Experimental/HordeServer.Experimental/Notifications/JobNotificationCollection.cs
@@ -375,7 +375,7 @@ namespace HordeServer.Experimental.Notifications
 			{
 				UpdateDefinitionBuilder<JobNotificationStateDocument> updateBuilder = Builders<JobNotificationStateDocument>.Update;
 				UpdateDefinition<JobNotificationStateDocument> updates = updateBuilder.Set(x => x.Channel, channel).Set(x => x.Ts, timestamp);
-				jobNotificationState = await _jobNotificationStates.FindOneAndUpdateAsync(x => x.Id == jobNotificationState.Id, updates, null, cancellationToken);
+				jobNotificationState = await _jobNotificationStates.FindOneAndUpdateAsync<JobNotificationStateDocument>(x => x.Id == jobNotificationState.Id, updates, new FindOneAndUpdateOptions<JobNotificationStateDocument> { IsUpsert = false, ReturnDocument = ReturnDocument.After }, cancellationToken);
 				_logger.LogInformation("Updated message {StateId} (jobId: {JobId}, templateId: {TemplateId}, recipient: {Recipient}, messageId: {MessageId})",
 										jobNotificationState.Id, jobNotificationState.JobId, jobNotificationState.TemplateId, jobNotificationState.Recipient, jobNotificationState.MessageId);
 			}
@@ -450,7 +450,7 @@ namespace HordeServer.Experimental.Notifications
 			{
 				UpdateDefinitionBuilder<JobStepNotificationStateDocument> updateBuilder = Builders<JobStepNotificationStateDocument>.Update;
 				UpdateDefinition<JobStepNotificationStateDocument> updates = updateBuilder.Set(x => x.Badge, badge).Set(x => x.Channel, channel).Set(x => x.ThreadTs, threadTs).Set(x => x.Ts, ts);
-				jobStepNotificationState = await _jobStepNotificationStates.FindOneAndUpdateAsync(x => x.Id == jobStepNotificationState.Id, updates, null, cancellationToken);
+				jobStepNotificationState = await _jobStepNotificationStates.FindOneAndUpdateAsync<JobStepNotificationStateDocument>(x => x.Id == jobStepNotificationState.Id, updates, new FindOneAndUpdateOptions<JobStepNotificationStateDocument> { IsUpsert = false, ReturnDocument = ReturnDocument.After }, cancellationToken);
 				_logger.LogInformation("Updated message {StateId} (jobId: {JobId}, templateId: {TemplateId}, recipient: {Recipient}, jobStepId: {JobStepId}, group: {Group}, platform: {Platform}, badge: {Badge}, parentJobId: {ParentJobId}, parentJobTemplateId: {ParentJobTemplateId}, messageId: {MessageId})",
 										jobStepNotificationState.Id, jobStepNotificationState.JobId, jobStepNotificationState.TemplateId, jobStepNotificationState.Recipient, jobStepNotificationState.JobStepId, jobStepNotificationState.Group, jobStepNotificationState.TargetPlatform, jobStepNotificationState.Badge, jobStepNotificationState.ParentJobId, jobStepNotificationState.ParentJobTemplateId, jobStepNotificationState.MessageId);
 			}
