Horde Analytics: Fix issue with sum values updating without matching on group

diff --git a/Engine/Source/Programs/Horde/Plugins/Analytics/HordeServer.Analytics/Telemetry/Metrics/MetricCollection.cs b/Engine/Source/Programs/Horde/Plugins/Analytics/HordeServer.Analytics/Telemetry/Metrics/MetricCollection.cs
--- a/Engine/Source/Programs/Horde/Plugins/Analytics/HordeServer.Analytics/Telemetry/Metrics/MetricCollection.cs
+++ b/Engine/Source/Programs/Horde/Plugins/Analytics/HordeServer.Analytics/Telemetry/Metrics/MetricCollection.cs
@@ -442,7 +442,7 @@ namespace HordeServer.Telemetry.Metrics
 				}
 
 				// Update the document
-				ReplaceOneResult result = await _metrics.ReplaceOneAsync(x => x.TelemetryStoreId == telemetryStoreId && x.Id == metric.Id && x.Count == prevCount, metric, cancellationToken: cancellationToken);
+				ReplaceOneResult result = await _metrics.ReplaceOneAsync(x => x.TelemetryStoreId == telemetryStoreId && x.Id == metric.Id && x.Group == group && x.Count == prevCount, metric, cancellationToken: cancellationToken);
 				if (result.MatchedCount > 0)
 				{
 					break;
