Fixing duplication of notifications caused by original owner of the job cancelling a step.

#rnx

diff --git a/Engine/Source/Programs/Horde/Plugins/Experimental/HordeServer.Experimental/Notifications/Sinks/ExperimentalSlackNotificationSink.cs b/Engine/Source/Programs/Horde/Plugins/Experimental/HordeServer.Experimental/Notifications/Sinks/ExperimentalSlackNotificationSink.cs
--- a/Engine/Source/Programs/Horde/Plugins/Experimental/HordeServer.Experimental/Notifications/Sinks/ExperimentalSlackNotificationSink.cs
+++ b/Engine/Source/Programs/Horde/Plugins/Experimental/HordeServer.Experimental/Notifications/Sinks/ExperimentalSlackNotificationSink.cs
@@ -79,6 +79,12 @@ namespace HordeServer.Notifications.Sinks
 		/// <inheritdoc/>
 		public async Task NotifyJobStepAbortedAsync(IEnumerable<IUser>? usersToNotify, IJob job, IJobStepBatch batch, IJobStep step, INode node, List<ILogEventData> jobStepEventData, CancellationToken cancellationToken)
 		{
+			// Currently we only want to notify from the plugin's configuration which would not have this set
+			if (usersToNotify is not null)
+			{
+				return;
+			}
+
 			await _slackNotificationProcessor.ProcessJobStepAbortedAsync(job, step, cancellationToken);
 		}
 
