Set minimum value for the Package and Content charts in Horde Analytics to zero

diff --git a/Engine/Source/Programs/Horde/HordeServer/Defaults/default-views.telemetry.json b/Engine/Source/Programs/Horde/HordeServer/Defaults/default-views.telemetry.json
--- a/Engine/Source/Programs/Horde/HordeServer/Defaults/default-views.telemetry.json
+++ b/Engine/Source/Programs/Horde/HordeServer/Defaults/default-views.telemetry.json
@@ -550,6 +550,7 @@
 						{
 							"name": "Content Size Total",
 							"display": "Value",
+							"min": 0,
 							"metrics": [
 								{
 									"id": "core-content-totalsize"
@@ -559,6 +560,7 @@
 						{
 							"name": "Content Size By Mount Point",
 							"display": "Value",
+							"min": 0,
 							"metrics": [
 								{
 									"id": "core-content-totalsize-mp"
@@ -568,6 +570,7 @@
 						{
 							"name": "Virtualization By Mount Point",
 							"display": "Value",
+							"min": 0,
 							"metrics": [
 								{
 									"id": "core-content-virtualization-mp"
@@ -747,15 +750,17 @@
 						{
 							"name": "Size By Install Type",
 							"display": "Value",
+							"min": 0,
 							"metrics": [
 								{
-									"id": "package-total-size"	
+									"id": "package-total-size"
 								}
 							]
 						},
 						{
 							"name": "Size Per Package",
 							"display": "Value",
+							"min": 0,
 							"metrics": [
 								{
 									"id": "package-size"
