Add missing tags in test session response

#rnx

diff --git a/Engine/Source/Programs/Horde/Plugins/Build/HordeServer.Build/Jobs/TestData/TestDataController.cs b/Engine/Source/Programs/Horde/Plugins/Build/HordeServer.Build/Jobs/TestData/TestDataController.cs
--- a/Engine/Source/Programs/Horde/Plugins/Build/HordeServer.Build/Jobs/TestData/TestDataController.cs
+++ b/Engine/Source/Programs/Horde/Plugins/Build/HordeServer.Build/Jobs/TestData/TestDataController.cs
@@ -751,6 +751,7 @@ namespace HordeServer.Jobs.TestData
 					MetadataId = testData.Metadata.ToString(),
 					NameRef = testData.NameRef.ToString(),
 					TestDataId = testData.TestDataId.ToString(),
+					TagIds = testData.Tags?.Select(t => t.ToString()).ToArray(),
 					Outcome = testData.Outcome,
 					PhasesTotalCount = testData.PhasesTotalCount,
 					PhasesSucceededCount = testData.PhasesSucceededCount,
diff --git a/Engine/Source/Programs/Shared/EpicGames.Horde/Jobs/TestData/TestDataMessages.cs b/Engine/Source/Programs/Shared/EpicGames.Horde/Jobs/TestData/TestDataMessages.cs
--- a/Engine/Source/Programs/Shared/EpicGames.Horde/Jobs/TestData/TestDataMessages.cs
+++ b/Engine/Source/Programs/Shared/EpicGames.Horde/Jobs/TestData/TestDataMessages.cs
@@ -703,6 +703,11 @@ namespace EpicGames.Horde.Jobs.TestData
 		/// </summary>
 		public string TestDataId { get; set; } = String.Empty;
 
+		/// <summary>
+		/// The list of tag ids associated with the test session
+		/// </summary>
+		public string[]? TagIds { get; set; }
+
 		/// <summary>
 		/// The commit id associated with the test session
 		/// </summary>
