Horde: Fix issue with artifacts not properly loading between step navigation

diff --git a/Engine/Source/Programs/Horde/HordeDashboard/src/components/jobDetailsV2/StepDetailView.tsx b/Engine/Source/Programs/Horde/HordeDashboard/src/components/jobDetailsV2/StepDetailView.tsx
--- a/Engine/Source/Programs/Horde/HordeDashboard/src/components/jobDetailsV2/StepDetailView.tsx
+++ b/Engine/Source/Programs/Horde/HordeDashboard/src/components/jobDetailsV2/StepDetailView.tsx
@@ -41,10 +41,6 @@ class StepDetailDataView extends JobDataView {
 
    async queryArtifacts() {
 
-      if (this.artifacts) {
-         return;
-      }
-
       if (!this.details?.jobData) {
          return;
       }
@@ -68,8 +64,6 @@ class StepDetailDataView extends JobDataView {
          }
          this.artifacts = artifacts;
       }
-
-
    }
 
    async set(stepId: string) {
