Horde: Add better indicator for inidividual data points

diff --git a/Engine/Source/Programs/Horde/HordeDashboard/plugins/analytics/telemetryLineGraph.ts b/Engine/Source/Programs/Horde/HordeDashboard/plugins/analytics/telemetryLineGraph.ts
--- a/Engine/Source/Programs/Horde/HordeDashboard/plugins/analytics/telemetryLineGraph.ts
+++ b/Engine/Source/Programs/Horde/HordeDashboard/plugins/analytics/telemetryLineGraph.ts
@@ -106,7 +106,7 @@ export class TelemetryLineRenderer {
          })
          .attr("cy", i => i[1])
          .attr("fill", i => { return i[3] !== undefined ? graphColors[i[3] as number] : "#8ab8ff" })
-         .attr("r", 0);
+         .attr("r", 2);
 
       let ticks: number[] = [];
       const inc = (maxTime.getTime() - minTime.getTime()) / 10;
@@ -250,7 +250,7 @@ export class TelemetryLineRenderer {
             closestX.set(key, values[closest.index]);
          })
 
-         svg!.selectAll(`circle`).attr("r", 0);
+         svg!.selectAll(`circle`).attr("r", 2);
 
          let closestY: DataPoint | undefined;
 
@@ -272,7 +272,7 @@ export class TelemetryLineRenderer {
       }
 
       const handleMouseLeave = (event: any) => {
-         svg!.selectAll(`circle`).attr("r", 0);
+         svg!.selectAll(`circle`).attr("r", 2);
          onDataHover("__clear__", 0, 0, new Date(), 0, "");
       }
 
