Horde: Fix step ETA times sometimes being in the past

diff --git a/Engine/Source/Programs/Horde/HordeDashboard/src/base/utilities/timeUtils.ts b/Engine/Source/Programs/Horde/HordeDashboard/src/base/utilities/timeUtils.ts
--- a/Engine/Source/Programs/Horde/HordeDashboard/src/base/utilities/timeUtils.ts
+++ b/Engine/Source/Programs/Horde/HordeDashboard/src/base/utilities/timeUtils.ts
@@ -133,6 +133,15 @@ export const getStepETA = (step: StepData, job: JobData): HordeTime => {
 
     let end = moment(job.createTime).add(step.timing.totalTimeToComplete, "seconds");
 
+    let now = moment(Date.now());
+
+    const seconds = end.diff(now, "seconds");
+    
+    // if less than a minute, push out
+    if (seconds < 60) {
+        end = now.add(5, "minutes");
+    }
+
     end = roundTime(end, moment.duration(5, "minutes"));
 
     const serverETA = moment.utc(end).tz(serverTimeZone);
