Horde: Improve data fetching error and surface network connection as likely cause on job page

diff --git a/Engine/Source/Programs/Horde/HordeDashboard/src/components/jobDetailsV2/JobDetailViewSummary.tsx b/Engine/Source/Programs/Horde/HordeDashboard/src/components/jobDetailsV2/JobDetailViewSummary.tsx
--- a/Engine/Source/Programs/Horde/HordeDashboard/src/components/jobDetailsV2/JobDetailViewSummary.tsx
+++ b/Engine/Source/Programs/Horde/HordeDashboard/src/components/jobDetailsV2/JobDetailViewSummary.tsx
@@ -3,12 +3,13 @@
 import { Label, Stack, Text } from '@fluentui/react';
 import { observer } from 'mobx-react-lite';
 import React, { useEffect } from 'react';
-import { JobState, ReportPlacement } from '../../backend/Api';
+import { DashboardPreference, JobState, ReportPlacement } from '../../backend/Api';
 import { Markdown } from '../../base/components/Markdown';
 import { ISideRailLink } from '../../base/components/SideRail';
 import { getNiceTime } from '../../base/utilities/timeUtils';
 import { JobDataView, JobDetailsV2 } from "./JobDetailsViewCommon";
 import { getHordeStyling } from '../../styles/Styles';
+import dashboard from 'horde/backend/Dashboard';
 
 const sideRail: ISideRailLink = { text: "Summary", url: "rail_summary" };
 
@@ -81,6 +82,16 @@ export const SummaryPanel: React.FC<{ jobDetails: JobDetailsV2 }> = observer(({
 
    }*/
 
+   if (jobDetails.jobError) {
+      document.title = "Error Loading Job";
+      let errorTextColor: string | undefined = dashboard.getPreference(DashboardPreference.ColorError) ?? "#c44525";
+      
+      return <Stack style={{ paddingTop: 64, paddingBottom: 82 }}>
+               <Text variant="large" styles={{ root: { fontFamily: "Horde Open Sans SemiBold", padding: "25px 30px 25px", color: errorTextColor } }}>Error loading job, please check network connection.</Text>
+               <Text variant="mediumPlus" styles={{ root: { fontFamily: "Horde Open Sans SemiBold", padding: "10px 30px 25px", color: errorTextColor } }}>{jobDetails.jobError}</Text>
+         </Stack>
+   }
+
    if (!jobData) {
       return null;
    }
diff --git a/Engine/Source/Programs/Horde/HordeDashboard/src/components/jobDetailsV2/JobDetailViewV2.tsx b/Engine/Source/Programs/Horde/HordeDashboard/src/components/jobDetailsV2/JobDetailViewV2.tsx
--- a/Engine/Source/Programs/Horde/HordeDashboard/src/components/jobDetailsV2/JobDetailViewV2.tsx
+++ b/Engine/Source/Programs/Horde/HordeDashboard/src/components/jobDetailsV2/JobDetailViewV2.tsx
@@ -54,12 +54,12 @@ const JobBreadCrumbs: React.FC<{ jobDetails: JobDetailsV2 }> = observer(({ jobDe
       if (!jobDetails.jobError) {
          return <Breadcrumbs items={[{ text: "Loading Job" }]} title={"Loading Job"} spinner={true} />
       } else {
-         return <Breadcrumbs items={[{ text: jobDetails.jobError }]} title={jobDetails.jobError} spinner={false} />
+         return null;
       }
    }
 
    if (jobDetails.jobError) {
-      return <Breadcrumbs items={[{ text: jobDetails.jobError }]} title={jobDetails.jobError} spinner={false} />
+      return null;
    }
 
    const crumbItems: BreadcrumbItem[] = [];
