Horde: Remove deprecated JetBrains.dotCover.GlobalTool

#rnx

diff --git a/Engine/Source/Programs/Horde/HordeServer/Dockerfile b/Engine/Source/Programs/Horde/HordeServer/Dockerfile
--- a/Engine/Source/Programs/Horde/HordeServer/Dockerfile
+++ b/Engine/Source/Programs/Horde/HordeServer/Dockerfile
@@ -17,9 +17,7 @@ WORKDIR /app
 COPY Source/Programs/Horde/nuget.config ./Source/Programs/Horde/nuget.config
 RUN DEBIAN_FRONTEND=noninteractive \
 	apt-get -qq update && \
-	apt-get -qq install curl zip && \
-	dotnet tool install --global JetBrains.dotCover.GlobalTool && \
-	ln -s /root/.dotnet/tools/dotnet-dotcover /usr/local/bin/dotnet-dotcover
+	apt-get -qq install curl zip
 
 COPY --from=redis /usr/local/bin/redis-server /usr/local/bin/redis-server
 
