Horde: Optional docker prune before building Horde Server

#rnx

diff --git a/Engine/Source/Programs/Horde/BuildHorde.xml b/Engine/Source/Programs/Horde/BuildHorde.xml
--- a/Engine/Source/Programs/Horde/BuildHorde.xml
+++ b/Engine/Source/Programs/Horde/BuildHorde.xml
@@ -71,6 +71,10 @@
 	<Agent Name="HordeServer" Type="Linux;Win64_Docker">
 		<Property Name="StagingDir" Value="$(RootDir)/Engine/Saved/HordeServer"/>
 
+		<Node Name="Docker Prune">
+			<Docker Arguments="system prune --all --force --volumes"/>
+		</Node>
+
 		<Node Name="Build HordeServer">
 			<!-- Tag all the files that need to be staged to build -->
 			<Tag Files="$(StagePaths)" Except=".../.vs/...;.../.git/...;.../bin/...;.../obj/..." With="#InputFiles"/>
