Horde: Display the year when not current to reduce confusion when looking at historical jobs

diff --git a/Engine/Source/Programs/Horde/HordeDashboard/src/base/utilities/timeUtils.ts b/Engine/Source/Programs/Horde/HordeDashboard/src/base/utilities/timeUtils.ts
--- a/Engine/Source/Programs/Horde/HordeDashboard/src/base/utilities/timeUtils.ts
+++ b/Engine/Source/Programs/Horde/HordeDashboard/src/base/utilities/timeUtils.ts
@@ -422,7 +422,13 @@ export const getHumanTime = (timeIn: Date | string | undefined): string => {
         return "Monday";
     }
 
-    return time.format('MMM Do');
+    let timeStr =  time.format('MMM Do');
+
+    if (time.year() !== now.year()) {
+        timeStr += `, ${time.year()}`;
+    }
+
+    return timeStr;
 
 };
 
@@ -455,6 +461,10 @@ export const getShortNiceTime = (timeIn: Date | string | undefined, relative: bo
             timeStr = time.format(format);
     }
 
+    if (time.year() !== now.year()) {
+        timeStr += `, ${time.year()}`;
+    }
+
     return timeStr;
 
 };
@@ -483,6 +493,10 @@ export const getNiceTime = (timeIn: Date | string | undefined, relative: boolean
 
     timeStr += ` at ${time.format(format)}`;
 
+    if (time.year() !== now.year()) {
+        timeStr += `, ${time.year()}`;
+    }
+
     return timeStr;
 
 };
