Horde: Replace display field with unit field, and add Bytes, BytesPerSecond, Seconds, and Percentage options

diff --git a/Engine/Source/Programs/Horde/HordeDashboard/plugins/analytics/api.ts b/Engine/Source/Programs/Horde/HordeDashboard/plugins/analytics/api.ts
--- a/Engine/Source/Programs/Horde/HordeDashboard/plugins/analytics/api.ts
+++ b/Engine/Source/Programs/Horde/HordeDashboard/plugins/analytics/api.ts
@@ -47,7 +47,7 @@ export type GetTelemetryMetricResponse = {
 	groupValues?: Record<string, string>;
 }
 
-export type TelemetryDisplayType = "Time" | "Ratio" | "Value";
+export type TelemetryUnitType = "Ratio" | "Value" | "Bytes" | "BytesPerSecond" | "Seconds" | "Percentage";
 export type TelemetryGraphType = "Line" | "Indicator";
 
 /// Metric attached to a telemetry chart	
@@ -70,8 +70,8 @@ export type GetTelemetryChartResponse = {
 	/// The name of the chart, will be displayed on the dashboard	
 	name: string;
 
-	/// The unit to display	
-	display: TelemetryDisplayType;
+	/// The unit of the data to display
+	unit?: TelemetryUnitType;
 
 	/// The graph type 	
 	graph: TelemetryGraphType;
diff --git a/Engine/Source/Programs/Horde/HordeDashboard/plugins/analytics/telemetryLineGraph.ts b/Engine/Source/Programs/Horde/HordeDashboard/plugins/analytics/telemetryLineGraph.ts
--- a/Engine/Source/Programs/Horde/HordeDashboard/plugins/analytics/telemetryLineGraph.ts
+++ b/Engine/Source/Programs/Horde/HordeDashboard/plugins/analytics/telemetryLineGraph.ts
@@ -16,7 +16,7 @@ type DataPoint = [number, number, string, number, number];
 
 export class TelemetryLineRenderer {
 
-   render(chart: GetTelemetryChartResponse, metrics: GetTelemetryMetricsResponse[], legend: string[], minTime: Date, maxTime: Date, container: HTMLDivElement, onZoom: (name: string, event: any) => void, onTimeSelect: (name: string, minTime: Date, maxTime: Date) => void, onDataHover: (key: string, x: number, y: number, time: Date, value: number, color: string) => void, scale = 1.0): any {
+   render(chart: GetTelemetryChartResponse, metrics: GetTelemetryMetricsResponse[], legend: string[], minTime: Date, maxTime: Date, dataSize: string, container: HTMLDivElement, onZoom: (name: string, event: any) => void, onTimeSelect: (name: string, minTime: Date, maxTime: Date) => void, onDataHover: (key: string, x: number, y: number, time: Date, value: number, color: string) => void, scale = 1.0): any {
 
       let minValue = Number.MAX_SAFE_INTEGER
       let maxValue = Number.MIN_SAFE_INTEGER
@@ -31,9 +31,10 @@ export class TelemetryLineRenderer {
          })
       })
 
-      const ratio = chart.display === "Ratio"
+      const ratio = chart.unit === "Ratio" || chart.unit === "Percentage"
 
       minValue = chart.min ?? 0;
+      maxValue = Math.ceil(maxValue / 10) * 10;
 
       if (ratio) {
          maxValue = 1.0
@@ -142,20 +143,33 @@ export class TelemetryLineRenderer {
                .ticks(10)
                .tickFormat((d) => {
 
-
                   if (ratio) {
-                     return (d as number * 100).toString() + "%"
+                     
+                     if (chart.unit === "Percentage") {
+                        return (d as number * 100).toString() + "%";
+                     }
+
+                     return d.toString();
                   }
 
                   if (!d) {
                      return "";
                   }
 
-                  if (chart.display === "Value") {
-                     return d.toString();
+                  if (chart.unit === "Bytes") {
+                     return (d as number).toString() + ` ${dataSize}`;            
+                  }
+
+                  if (chart.unit === "BytesPerSecond") {
+                     return (d as number).toString() + ` ${dataSize}/s`;
+                  }
+
+                  if (chart.unit === "Seconds") {
+                     return msecToElapsed((d as number) * 1000, true, true);
                   }
 
-                  return msecToElapsed((d as number) * 1000, true, true);
+                  // fall back to general Value unit type
+                  return d.toString();
 
                }))
       }
diff --git a/Engine/Source/Programs/Horde/HordeDashboard/plugins/analytics/telemetryView.tsx b/Engine/Source/Programs/Horde/HordeDashboard/plugins/analytics/telemetryView.tsx
--- a/Engine/Source/Programs/Horde/HordeDashboard/plugins/analytics/telemetryView.tsx
+++ b/Engine/Source/Programs/Horde/HordeDashboard/plugins/analytics/telemetryView.tsx
@@ -14,6 +14,7 @@ import { useNavigate, useSearchParams } from "react-router-dom";
 import { GetTelemetryChartResponse, GetTelemetryMetricResponse, GetTelemetryMetricsResponse, GetTelemetryVariableResponse, GetTelemetryViewResponse, getViews } from "./api";
 import { TelemetryViewData, clearTelemetryViewMetrics, getTelemetryViewData, graphColors } from "./telemetryData";
 import { TelemetryLineRenderer } from "./telemetryLineGraph";
+import { formatBytes } from "horde/base/utilities/stringUtills";
 
 const timeSelections: TimeSelection[] = [
    {
@@ -320,6 +321,32 @@ class MetricsHandler {
          return cmetrics.has(m.metricId);
       }).map(m => { return { ...m } as GetTelemetryMetricsResponse });
 
+
+      if (chart.unit === "Bytes" || chart.unit === "BytesPerSecond") {
+
+         if (this.chartDataSize.get(chart.name) === undefined) {
+            // find best fit data size
+            let dataSizes = new Map();
+            metrics.forEach(metric => {
+               metric.metrics.forEach(m => {
+                  let formattedSize = formatBytes(m.value).replace(/[0-9. ]/g, '');
+                  dataSizes.set(formattedSize, (dataSizes.get(formattedSize) ?? 0) + 1);
+               })
+            })
+            let dataSize = [...dataSizes.entries()].reduce((a, b) => a[1] > b[1] ? a : b)[0];
+
+            // apply best fit data size to all metrics in the chart
+            metrics.forEach(metric => {
+               metric.metrics.forEach(m => {
+                  m.value = Number(formatBytes(m.value, 2, dataSize).replace(/[a-zA-Z ]/g, ''));
+               })
+            })
+
+            this.chartDataSize.set(chart.name, dataSize);
+         }
+      }
+
+
       if (chart.min !== undefined || chart.max !== undefined) {
          metrics.forEach(metric => {
             metric.metrics = metric.metrics.filter(m => {
@@ -395,14 +422,24 @@ class MetricsHandler {
 
       let value = "";
 
-      if (chart.display === "Ratio") {
-         value = Math.round((valueIn * 100)).toString() + "%"
-      }
-      else if (chart.display === "Value") {
-         value = Math.round(valueIn).toString();
-      }
-      else {
-         value = msecToElapsed((valueIn) * 1000, true, true);
+      switch(chart.unit) {
+         case "Percentage":
+            value = Math.round((valueIn * 100)).toString() + "%";
+            break;
+         case "Ratio":
+            value = valueIn.toFixed(2).toString();
+            break;
+         case "Seconds":
+            value = msecToElapsed((valueIn) * 1000, true, true);
+            break;
+         case "Bytes":
+            value = valueIn.toFixed(2).toString() + ` ${this.chartDataSize.get(chart.name)}`;
+            break;
+         case "BytesPerSecond":
+            value = valueIn.toFixed(2).toString() + ` ${this.chartDataSize.get(chart.name)}/s`;
+            break;
+         default:
+            value = Math.round(valueIn).toString();
       }
 
       return value;
@@ -706,6 +743,7 @@ class MetricsHandler {
 
    filteredKeys = new Set<string>();
    zoomHandler = new Map<string, any>();
+   chartDataSize = new Map<string, string>();
 
    searchState: SearchState = {}
 
@@ -968,7 +1006,7 @@ const Legend: React.FC<{ chart: GetTelemetryChartResponse }> = observer(({ chart
                      <Text variant="small">Max: {handler.valueToString(chart, entry.max)}</Text>
                   </Stack>
                   <Stack>
-                     <Text variant="small">Change: {Math.round((entry.change * 100)).toString() + "%"}</Text>
+                     <Text variant="small">Change: {chart.unit === "Ratio" ? entry.change.toFixed(2).toString() : Math.round((entry.change * 100)).toString() + "%"}</Text>
                   </Stack>
                </Stack>
             </Stack>
@@ -1185,7 +1223,7 @@ const IndicatorTile: React.FC<{ chart: GetTelemetryChartResponse }> = observer((
             <IndicatorBar stack={barStack} width={160} height={14} />
          </Stack>
          <Stack style={{ width: 90 }} horizontalAlign="end">
-            <Text variant="small">{chart.display === "Value" ? m.value : msecToElapsed(m.value * 1000)}</Text>
+            <Text variant="small">{handler.valueToString(chart, m.value)}</Text>
          </Stack>
       </Stack>
 
@@ -1343,7 +1381,7 @@ const LineGraphTile: React.FC<{ chart: GetTelemetryChartResponse }> = observer((
             tooltip.set(key, x, y, time, value, color);
          }
 
-         const zoomed = renderer.render(chart, metrics, legend.map(v => v.key), handler.minDate!, handler.maxDate!, container, onZoom, onTimeSelect, onDataHover, scale);
+         const zoomed = renderer.render(chart, metrics, legend.map(v => v.key), handler.minDate!, handler.maxDate!, handler.chartDataSize.get(chart.name)!, container, onZoom, onTimeSelect, onDataHover, scale);
          handler.setZoomHandler(chart.name, zoomed);
 
       } catch (err) {
diff --git a/Engine/Source/Programs/Horde/HordeDashboard/src/base/utilities/stringUtills.ts b/Engine/Source/Programs/Horde/HordeDashboard/src/base/utilities/stringUtills.ts
--- a/Engine/Source/Programs/Horde/HordeDashboard/src/base/utilities/stringUtills.ts
+++ b/Engine/Source/Programs/Horde/HordeDashboard/src/base/utilities/stringUtills.ts
@@ -1,12 +1,12 @@
 
-export function formatBytes(bytes: number, decimals = 2) {
+export function formatBytes(bytes: number, decimals = 2, size = "") {
     if (!+bytes) return '0 Bytes'
  
     const k = 1024
     const dm = decimals < 0 ? 0 : decimals
     const sizes = ['Bytes', 'KB', 'MB', 'GB', 'TB', 'PB', 'EB', 'ZB', 'YB']
  
-    const i = Math.floor(Math.log(bytes) / Math.log(k))
+    const i = size ? sizes.indexOf(size) : Math.floor(Math.log(bytes) / Math.log(k))
  
     return `${parseFloat((bytes / Math.pow(k, i)).toFixed(dm))} ${sizes[i]}`
  }
diff --git a/Engine/Source/Programs/Horde/Plugins/Analytics/HordeServer.Analytics/Telemetry/TelemetryController.cs b/Engine/Source/Programs/Horde/Plugins/Analytics/HordeServer.Analytics/Telemetry/TelemetryController.cs
--- a/Engine/Source/Programs/Horde/Plugins/Analytics/HordeServer.Analytics/Telemetry/TelemetryController.cs
+++ b/Engine/Source/Programs/Horde/Plugins/Analytics/HordeServer.Analytics/Telemetry/TelemetryController.cs
@@ -228,7 +228,7 @@ namespace HordeServer.Telemetry
 
 					foreach (TelemetryChartConfig chart in category.Charts)
 					{
-						GetTelemetryChartResponse rchart = new GetTelemetryChartResponse { Name = chart.Name, Display = chart.Display.ToString(), Graph = chart.Graph.ToString(), Min = chart.Min, Max = chart.Max, Metrics = new List<GetTelemetryChartMetricResponse>() };
+						GetTelemetryChartResponse rchart = new GetTelemetryChartResponse { Name = chart.Name, Unit = chart.Unit.ToString(), Graph = chart.Graph.ToString(), Min = chart.Min, Max = chart.Max, Metrics = new List<GetTelemetryChartMetricResponse>() };
 
 						foreach (TelemetryChartMetricConfig metric in chart.Metrics)
 						{
diff --git a/Engine/Source/Programs/Shared/EpicGames.Horde/Telemetry/TelemetryMessages.cs b/Engine/Source/Programs/Shared/EpicGames.Horde/Telemetry/TelemetryMessages.cs
--- a/Engine/Source/Programs/Shared/EpicGames.Horde/Telemetry/TelemetryMessages.cs
+++ b/Engine/Source/Programs/Shared/EpicGames.Horde/Telemetry/TelemetryMessages.cs
@@ -175,9 +175,9 @@ namespace EpicGames.Horde.Telemetry
 		public string Name { get; set; } = null!;
 
 		/// <summary>
-		/// The unit to display
+		/// The unit of the data to display
 		/// </summary>
-		public TelemetryMetricUnitType Display { get; set; } = TelemetryMetricUnitType.Time;
+		public TelemetryMetricUnitType Unit { get; set; } = TelemetryMetricUnitType.Time;
 
 		/// <summary>
 		/// The graph type 
@@ -343,9 +343,9 @@ namespace EpicGames.Horde.Telemetry
 		public string Name { get; set; } = null!;
 
 		/// <summary>
-		/// The unit to display
+		/// The unit of the data to display
 		/// </summary>
-		public string Display { get; set; } = null!;
+		public string Unit { get; set; } = null!;
 
 		/// <summary>
 		/// The graph type 
