Horde: Add admin claim resolution for Horde auth mode

Admin claims are now properly resolved when users sign in using Horde auth mode, matching the existing behavior in OIDC/Okta auth handlers.

diff --git a/Engine/Source/Programs/Horde/HordeServer/Server/AccountController.cs b/Engine/Source/Programs/Horde/HordeServer/Server/AccountController.cs
--- a/Engine/Source/Programs/Horde/HordeServer/Server/AccountController.cs
+++ b/Engine/Source/Programs/Horde/HordeServer/Server/AccountController.cs
@@ -66,16 +66,18 @@ namespace HordeServer.Server
 		readonly IAccountCollection _hordeAccounts;
 		readonly string _authenticationScheme;
 		readonly IOptionsSnapshot<GlobalConfig> _globalConfig;
+		readonly IOptionsMonitor<ServerSettings> _settings;
 
 		/// <summary>
 		/// Constructor
 		/// </summary>
-		public AccountController(IUserCollection users, IAccountCollection hordeAccounts, IOptionsMonitor<ServerSettings> serverSettings, IOptionsSnapshot<GlobalConfig> globalConfig)
+		public AccountController(IUserCollection users, IAccountCollection hordeAccounts, IOptionsMonitor<ServerSettings> serverSettings, IOptionsSnapshot<GlobalConfig> globalConfig, IOptionsMonitor<ServerSettings> settings)
 		{
 			_users = users;
 			_hordeAccounts = hordeAccounts;
 			_authenticationScheme = GetAuthScheme(serverSettings.CurrentValue.AuthMethod);
 			_globalConfig = globalConfig;
+			_settings = settings;
 		}
 
 		/// <summary>
@@ -285,7 +287,7 @@ namespace HordeServer.Server
 			}
 
 			IAccount? account = await _hordeAccounts.FindByLoginAsync(login);
-			if (account == null || !account.Enabled)
+			if (account is not { Enabled: true })
 			{
 				return false;
 			}
@@ -313,6 +315,13 @@ namespace HordeServer.Server
 			{
 				claims.Add(new Claim(ClaimTypes.Email, account.Email));
 			}
+			
+			if (!String.IsNullOrEmpty(_settings.CurrentValue.AdminClaimType) &&
+			    !String.IsNullOrEmpty(_settings.CurrentValue.AdminClaimValue) &&
+			    account.HasClaim(_settings.CurrentValue.AdminClaimType, _settings.CurrentValue.AdminClaimValue))
+			{
+				claims.Add(HordeClaims.AdminClaim.ToClaim());
+			}
 
 			foreach (IUserClaim claim in account.Claims)
 			{
