Horde : Add topological sort of plugins by their dependencies to call the PostLoad Method in the correct order

diff --git a/Engine/Source/Programs/Horde/HordeServer.Shared/Plugins/PluginAttribute.cs b/Engine/Source/Programs/Horde/HordeServer.Shared/Plugins/PluginAttribute.cs
--- a/Engine/Source/Programs/Horde/HordeServer.Shared/Plugins/PluginAttribute.cs
+++ b/Engine/Source/Programs/Horde/HordeServer.Shared/Plugins/PluginAttribute.cs
@@ -19,6 +19,11 @@ namespace HordeServer.Plugins
 		/// </summary>
 		public bool EnabledByDefault { get; set; } = true;
 
+		/// <summary>
+		/// Other plugins that this plugin depends on
+		/// </summary>
+		public string[] DependsOn { get; set; } = [];
+
 		/// <summary>
 		/// Type containing server settings for this plugin. Will be injected into a <see cref="Microsoft.Extensions.Options.IOptions{TOptions}"/>
 		/// object in the service container, as well as (optionally) the constructor for the the concrete plugin startup object.
diff --git a/Engine/Source/Programs/Horde/HordeServer.Tests/Plugins/PluginTests.cs b/Engine/Source/Programs/Horde/HordeServer.Tests/Plugins/PluginTests.cs
new file mode 100644
--- /dev/null
+++ b/Engine/Source/Programs/Horde/HordeServer.Tests/Plugins/PluginTests.cs
@@ -0,0 +1,135 @@
+// Copyright Epic Games, Inc. All Rights Reserved.
+
+using System;
+using System.Collections.Generic;
+using HordeServer.Plugins;
+using Microsoft.AspNetCore.Builder;
+using Microsoft.Extensions.DependencyInjection;
+using Microsoft.VisualStudio.TestTools.UnitTesting;
+
+namespace HordeServer.Tests.Plugins
+{
+	[TestClass]
+	public class PluginsTest
+	{
+		private static class PluginsWithDependencies
+		{
+			[Plugin("A", DependsOn = ["C"])]
+			public class PluginA : IPluginStartup
+			{
+				public void Configure(IApplicationBuilder app) { }
+				public void ConfigureServices(IServiceCollection serviceCollection) { }
+			}
+
+			[Plugin("B")]
+			public class PluginB : IPluginStartup
+			{
+				public void Configure(IApplicationBuilder app) { }
+				public void ConfigureServices(IServiceCollection serviceCollection) { }
+			}
+
+			[Plugin("C", DependsOn = ["B"])]
+			public class PluginC : IPluginStartup
+			{
+				public void Configure(IApplicationBuilder app) { }
+				public void ConfigureServices(IServiceCollection serviceCollection) { }
+			}
+		}
+
+		private static class PluginsWithCycle
+		{
+			[Plugin("A", DependsOn = ["B"])]
+			public class PluginA : IPluginStartup
+			{
+				public void Configure(IApplicationBuilder app) { }
+				public void ConfigureServices(IServiceCollection serviceCollection) { }
+			}
+
+			[Plugin("B", DependsOn = ["A"])]
+			public class PluginB : IPluginStartup
+			{
+				public void Configure(IApplicationBuilder app) { }
+				public void ConfigureServices(IServiceCollection serviceCollection) { }
+			}
+		}
+
+		private static class PluginsWithoutDependencies
+		{
+			[Plugin("A")]
+			public class PluginA : IPluginStartup
+			{
+				public void Configure(IApplicationBuilder app) { }
+				public void ConfigureServices(IServiceCollection serviceCollection) { }
+			}
+
+			[Plugin("B")]
+			public class PluginB : IPluginStartup
+			{
+				public void Configure(IApplicationBuilder app) { }
+				public void ConfigureServices(IServiceCollection serviceCollection) { }
+			}
+
+			[Plugin("C")]
+			public class PluginC : IPluginStartup
+			{
+				public void Configure(IApplicationBuilder app) { }
+				public void ConfigureServices(IServiceCollection serviceCollection) { }
+			}
+		}
+
+		private static class PluginsMissingDependencies
+		{
+			[Plugin("A", DependsOn = ["Missing"])]
+			public class PluginA : IPluginStartup
+			{
+				public void Configure(IApplicationBuilder app) { }
+				public void ConfigureServices(IServiceCollection serviceCollection) { }
+			}
+		}
+
+		[TestMethod]
+		public void TopologicalSort()
+		{
+			PluginCollection pluginCollection = new();
+			pluginCollection.Add(typeof(PluginsWithDependencies.PluginA));
+			pluginCollection.Add(typeof(PluginsWithDependencies.PluginB));
+			pluginCollection.Add(typeof(PluginsWithDependencies.PluginC));
+			IReadOnlyList<ILoadedPlugin> sorted = PluginCollection.GetTopologicalSort(pluginCollection.LoadedPlugins);
+			Assert.AreEqual(sorted[0].Name.ToString(), "b");
+			Assert.AreEqual(sorted[1].Name.ToString(), "c");
+			Assert.AreEqual(sorted[2].Name.ToString(), "a");
+		}
+
+		[TestMethod]
+		public void TopologicalSortEmpty()
+		{
+			PluginCollection pluginCollection = new();
+			pluginCollection.Add(typeof(PluginsWithoutDependencies.PluginA));
+			pluginCollection.Add(typeof(PluginsWithoutDependencies.PluginB));
+			pluginCollection.Add(typeof(PluginsWithoutDependencies.PluginC));
+			IReadOnlyList<ILoadedPlugin> sorted = PluginCollection.GetTopologicalSort(pluginCollection.LoadedPlugins);
+			Assert.AreEqual(sorted[0].Name.ToString(), "a");
+			Assert.AreEqual(sorted[1].Name.ToString(), "b");
+			Assert.AreEqual(sorted[2].Name.ToString(), "c");
+		}
+
+		[TestMethod]
+		[ExpectedException(typeof(InvalidOperationException))]
+		public void TopologicalSortErrorWithCycle()
+		{
+			PluginCollection pluginCollection = new();
+			pluginCollection.Add(typeof(PluginsWithCycle.PluginA));
+			pluginCollection.Add(typeof(PluginsWithCycle.PluginB));
+			_ = PluginCollection.GetTopologicalSort(pluginCollection.LoadedPlugins);
+		}
+
+		[TestMethod]
+		public void TopologicalSortMissingDependency()
+		{
+			PluginCollection pluginCollection = new();
+			pluginCollection.Add(typeof(PluginsMissingDependencies.PluginA));
+			IReadOnlyList<ILoadedPlugin> sorted = PluginCollection.GetTopologicalSort(pluginCollection.LoadedPlugins);
+			Assert.AreEqual(sorted[0].Name.ToString(), "a");
+		}
+	}
+}
diff --git a/Engine/Source/Programs/Horde/HordeServer/Plugins/PluginCollection.cs b/Engine/Source/Programs/Horde/HordeServer/Plugins/PluginCollection.cs
--- a/Engine/Source/Programs/Horde/HordeServer/Plugins/PluginCollection.cs
+++ b/Engine/Source/Programs/Horde/HordeServer/Plugins/PluginCollection.cs
@@ -152,7 +152,7 @@ namespace HordeServer.Plugins
 			PluginAttribute attr = startupType.GetCustomAttribute<PluginAttribute>()
 				?? throw new InvalidOperationException($"Cannot add {startupType.Name} as a plugin. No {nameof(PluginAttribute)} was found.");
 
-			IPluginMetadata metadata = new PluginMetadata { Name = new PluginName(attr.Name) };
+			IPluginMetadata metadata = new PluginMetadata { Name = new PluginName(attr.Name), DependsOn = [..attr.DependsOn] };
 			if (!_loadedPluginNames.Add(metadata.Name))
 			{
 				throw new InvalidOperationException($"An implementation of the {metadata.Name} plugin has already been added");
@@ -165,6 +165,70 @@ namespace HordeServer.Plugins
 			return loadedPlugin;
 		}
 
+		private enum DependencyNodeState
+		{
+			None,
+			InProgress,
+			Done,
+		}
+
+		private class DependencyNode(ILoadedPlugin plugin, IReadOnlyList<PluginName> dependsOn)
+		{
+			public readonly ILoadedPlugin Plugin = plugin;
+			public readonly IReadOnlyList<PluginName> DependsOn = dependsOn;
+			public DependencyNodeState State { get; set; } = DependencyNodeState.None;
+		}
+
+		private static void DependencyNodeVisit(DependencyNode current, Dictionary<PluginName, DependencyNode> nodes, List<ILoadedPlugin> result)
+		{
+			if (current.State == DependencyNodeState.Done)
+			{
+				return;
+			}
+			if (current.State == DependencyNodeState.InProgress)
+			{
+				throw new InvalidOperationException($"Cycle found in plugins for '{current.Plugin.Name}' during dependency sort");
+			}
+			current.State = DependencyNodeState.InProgress;
+
+			foreach (PluginName dependsOn in current.DependsOn)
+			{
+				if (nodes.TryGetValue(dependsOn, out DependencyNode? next))
+				{
+					DependencyNodeVisit(next, nodes, result);
+				}
+			}
+
+			current.State = DependencyNodeState.Done;
+			result.Add(current.Plugin);
+		}
+
+		/// <summary>
+		/// Returns a topological sorted sequence of plugins by plugins they depend upon
+		/// </summary>
+		public static IReadOnlyList<ILoadedPlugin> GetTopologicalSort(IReadOnlyList<ILoadedPlugin> plugins)
+		{
+			List<ILoadedPlugin> result = [];
+			Dictionary<PluginName, DependencyNode> nodes = new();
+
+			foreach (ILoadedPlugin plugin in plugins)
+			{
+				List<PluginName> dependsOn = [];
+				foreach (string pluginName in plugin.Metadata.DependsOn)
+				{
+					dependsOn.Add(new PluginName(pluginName));
+				}
+				DependencyNode node = new(plugin, dependsOn);
+				nodes.Add(plugin.Name, node);
+			}
+
+			foreach (DependencyNode node in nodes.Values)
+			{
+				DependencyNodeVisit(node, nodes, result);
+			}
+			return result;
+		}
+
 		/// <summary>
 		/// Adds a plugin with the given startup class
 		/// </summary>
diff --git a/Engine/Source/Programs/Horde/HordeServer/Server/GlobalConfig.cs b/Engine/Source/Programs/Horde/HordeServer/Server/GlobalConfig.cs
--- a/Engine/Source/Programs/Horde/HordeServer/Server/GlobalConfig.cs
+++ b/Engine/Source/Programs/Horde/HordeServer/Server/GlobalConfig.cs
@@ -129,8 +129,33 @@ namespace HordeServer.Server
 				}
 			}
 
+			IReadOnlyCollection<IPluginConfig> pluginConfigs;
+			if (loadedPlugins.Count == 0)
+			{
+				pluginConfigs = Plugins.Values;
+			}
+			else
+			{
+				List<IPluginConfig> sortedPluginConfigs = [];
+				HashSet<PluginName> seen = [];
+				foreach (ILoadedPlugin loadedPlugin in PluginCollection.GetTopologicalSort(loadedPlugins))
+				{
+					sortedPluginConfigs.Add(Plugins[loadedPlugin.Name]);
+					seen.Add(loadedPlugin.Name);
+				}
+				// Add configs that do not have a plugin
+				foreach (KeyValuePair<PluginName, IPluginConfig> kvp in Plugins)
+				{
+					if (!seen.Contains(kvp.Key))
+					{
+						sortedPluginConfigs.Add(kvp.Value);
+					}
+				}
+				pluginConfigs = sortedPluginConfigs;
+			}
+
 			PluginConfigOptions pluginConfigOptions = new PluginConfigOptions(VersionEnum, Plugins.Values, Acl);
-			foreach (IPluginConfig pluginConfig in Plugins.Values)
+			foreach (IPluginConfig pluginConfig in pluginConfigs)
 			{
 				pluginConfig.PostLoad(pluginConfigOptions);
 			}
