Horde: Make update session call more robust

* Catch HttpRequestException in addition to RpcException
* Implement exponential backoff starting at 1s (max 20s)
* Increase max retries from 3 to 8
* Improve logging slightly to explain attempt count and delay

diff --git a/Engine/Source/Programs/Horde/HordeAgent/Leases/LeaseManager.cs b/Engine/Source/Programs/Horde/HordeAgent/Leases/LeaseManager.cs
--- a/Engine/Source/Programs/Horde/HordeAgent/Leases/LeaseManager.cs
+++ b/Engine/Source/Programs/Horde/HordeAgent/Leases/LeaseManager.cs
@@ -12,7 +12,6 @@ using Microsoft.Extensions.DependencyInjection;
 using Microsoft.Extensions.Logging;
 using Microsoft.Extensions.Options;
 using OpenTelemetry.Trace;
-using StatusCode = Grpc.Core.StatusCode;
 
 namespace HordeAgent.Leases
 {
@@ -473,26 +472,27 @@ namespace HordeAgent.Leases
 		/// <returns>Response from the call</returns>
 		async Task<RpcUpdateSessionResponse?> UpdateSessionAsync(HordeRpc.HordeRpcClient hordeRpc, RpcUpdateSessionRequest updateSessionRequest, Task waitTask)
 		{
+			const int MaxRetries = 8;
+			const int BaseRetryDelayMs = 1000;
+			const int MaxDelayMs = 20000;
+			
 			RpcUpdateSessionResponse? updateSessionResponse = null;
 			try
 			{
 				updateSessionResponse = await UpdateSessionInternalAsync(hordeRpc, updateSessionRequest, waitTask);
 				_updateSessionFailures = 0;
 			}
-			catch (RpcException ex)
+			catch (Exception ex) when (ex is RpcException or HttpRequestException)
 			{
-				if (++_updateSessionFailures >= 3)
+				_updateSessionFailures++;
+				if (_updateSessionFailures >= MaxRetries)
 				{
 					throw;
 				}
-				else if (ex.StatusCode == StatusCode.Unavailable)
-				{
-					_logger.LogInformation(ex, "Service unavailable while calling UpdateSessionAsync(). Will retry...");
-				}
-				else
-				{
-					_logger.LogError(ex, "Error while executing RPC. Will retry...");
-				}
+				
+				int retryDelayMs = Math.Min((int)(BaseRetryDelayMs * Math.Pow(2, _updateSessionFailures - 1)), MaxDelayMs);
+				_logger.LogWarning(ex, "Error while updating session. Backing off {DelayMs} ms before next attempt... ({Attempt} of {MaxAttempts})", retryDelayMs, _updateSessionFailures, MaxRetries);
+				await Task.Delay(retryDelayMs);
 			}
 			return updateSessionResponse;
 		}
