Horde: Change GlobalConfig.PostLoad and IPluginConfig.PostLoad methods to async

This is to allow, in the future, config's to use async APIs for example to fetch secrets.
Some tests had setup in constructors changed these to use the TestInitialize attribute.

diff --git a/Engine/Source/Programs/Horde/HordeServer.Shared/Plugins/IPluginConfig.cs b/Engine/Source/Programs/Horde/HordeServer.Shared/Plugins/IPluginConfig.cs
--- a/Engine/Source/Programs/Horde/HordeServer.Shared/Plugins/IPluginConfig.cs
+++ b/Engine/Source/Programs/Horde/HordeServer.Shared/Plugins/IPluginConfig.cs
@@ -13,11 +13,11 @@ namespace HordeServer.Plugins
 		/// Called to fixup a plugin's configuration after deserialization
 		/// </summary>
 		/// <param name="configOptions">Options for configuring the plugin</param>
-		void PostLoad(PluginConfigOptions configOptions);
+		Task PostLoadAsync(PluginConfigOptions configOptions);
 	}
 
 	/// <summary>
-	/// Options passed to <see cref="IPluginConfig.PostLoad(PluginConfigOptions)"/>
+	/// Options passed to <see cref="IPluginConfig.PostLoadAsync(PluginConfigOptions)"/>
 	/// </summary>
 	public record class PluginConfigOptions(ConfigVersion Version, IEnumerable<IPluginConfig> Plugins, AclConfig ParentAcl);
 
@@ -27,7 +27,9 @@ namespace HordeServer.Plugins
 	public sealed class EmptyPluginConfig : IPluginConfig
 	{
 		/// <inheritdoc/>
-		public void PostLoad(PluginConfigOptions configOptions)
-		{ }
+		public Task PostLoadAsync(PluginConfigOptions configOptions)
+		{
+			return Task.CompletedTask;
+		}
 	}
 }
diff --git a/Engine/Source/Programs/Horde/HordeServer.Tests.Shared/ServerTestSetup.cs b/Engine/Source/Programs/Horde/HordeServer.Tests.Shared/ServerTestSetup.cs
--- a/Engine/Source/Programs/Horde/HordeServer.Tests.Shared/ServerTestSetup.cs
+++ b/Engine/Source/Programs/Horde/HordeServer.Tests.Shared/ServerTestSetup.cs
@@ -72,17 +72,17 @@ namespace HordeServer.Tests
 				.CreateLogger();
 		}
 
-		protected void SetConfig(GlobalConfig globalConfig)
+		protected async Task SetConfigAsync(GlobalConfig globalConfig)
 		{
-			globalConfig.PostLoad(ServerSettings, _pluginCollection.LoadedPlugins, DefaultAclModifiers);
+			await globalConfig.PostLoadAsync(ServerSettings, _pluginCollection.LoadedPlugins, DefaultAclModifiers);
 			ConfigService.OverrideConfig(globalConfig);
 		}
 
-		protected void UpdateConfig(Action<GlobalConfig> action)
+		protected async Task UpdateConfigAsync(Action<GlobalConfig> action)
 		{
 			GlobalConfig globalConfig = GlobalConfig.CurrentValue;
 			action(globalConfig);
-			globalConfig.PostLoad(ServerSettings, _pluginCollection.LoadedPlugins, DefaultAclModifiers);
+			await globalConfig.PostLoadAsync(ServerSettings, _pluginCollection.LoadedPlugins, DefaultAclModifiers);
 			ConfigService.OverrideConfig(globalConfig);
 		}
 
diff --git a/Engine/Source/Programs/Horde/HordeServer.Tests/Configuration/ConfigTests.cs b/Engine/Source/Programs/Horde/HordeServer.Tests/Configuration/ConfigTests.cs
--- a/Engine/Source/Programs/Horde/HordeServer.Tests/Configuration/ConfigTests.cs
+++ b/Engine/Source/Programs/Horde/HordeServer.Tests/Configuration/ConfigTests.cs
@@ -173,7 +173,7 @@ namespace HordeServer.Tests.Configuration
 		}
 
 		[TestMethod]
-		public void WorkspaceConfig()
+		public async Task WorkspaceConfigAsync()
 		{
 			Dictionary<string, WorkspaceConfig> inputWorkspaces = new();
 			inputWorkspaces["base"] = new WorkspaceConfig { Identifier = "base", Cluster = "myCluster", MinScratchSpace = 111 };
@@ -186,7 +186,7 @@ namespace HordeServer.Tests.Configuration
 			GlobalConfig gc = new();
 			gc.Plugins.AddComputeConfig(new ComputeConfig());
 			gc.Plugins.AddBuildConfig(buildConfig);
-			gc.PostLoad(new ServerSettings(), Array.Empty<ILoadedPlugin>(), Array.Empty<IDefaultAclModifier>());
+			await gc.PostLoadAsync(new ServerSettings(), Array.Empty<ILoadedPlugin>(), Array.Empty<IDefaultAclModifier>());
 
 			Dictionary<string,WorkspaceConfig> workspaces = gc.Plugins.GetBuildConfig().Projects[0].Streams[0].WorkspaceTypes;
 			Assert.AreEqual(3, workspaces.Count);
@@ -196,7 +196,7 @@ namespace HordeServer.Tests.Configuration
 		}
 
 		[TestMethod]
-		public void WorkspaceInheritFromProjectConfig()
+		public async Task WorkspaceInheritFromProjectConfigAsync()
 		{
 			List<string> autoSdkViews = ["foo", "bar"];
 
@@ -221,7 +221,7 @@ namespace HordeServer.Tests.Configuration
 			GlobalConfig gc = new();
 			gc.Plugins.AddComputeConfig(new ComputeConfig());
 			gc.Plugins.AddBuildConfig(buildConfig);
-			gc.PostLoad(new ServerSettings(), Array.Empty<ILoadedPlugin>(), Array.Empty<IDefaultAclModifier>());
+			await gc.PostLoadAsync(new ServerSettings(), Array.Empty<ILoadedPlugin>(), Array.Empty<IDefaultAclModifier>());
 
 			Dictionary<string, WorkspaceConfig> workspaces = gc.Plugins.GetBuildConfig().Projects[0].Streams[0].WorkspaceTypes;
 			Assert.AreEqual(4, workspaces.Count);
diff --git a/Engine/Source/Programs/Horde/HordeServer.Tests/Server/DowntimeServiceTests.cs b/Engine/Source/Programs/Horde/HordeServer.Tests/Server/DowntimeServiceTests.cs
--- a/Engine/Source/Programs/Horde/HordeServer.Tests/Server/DowntimeServiceTests.cs
+++ b/Engine/Source/Programs/Horde/HordeServer.Tests/Server/DowntimeServiceTests.cs
@@ -15,12 +15,12 @@ namespace HordeServer.Tests.Server
 	[TestClass]
 	public class DowntimeServiceTest
 	{
-		private static ScheduledDowntime GetDowntimeFromJson(string json)
+		private static async Task<ScheduledDowntime> GetDowntimeFromJsonAsync(string json)
 		{
 			ScheduledDowntime scheduledDowntime = JsonSerializer.Deserialize<ScheduledDowntime>(json, JsonUtils.DefaultSerializerOptions)!;
 			GlobalConfig config = new();
 			config.Downtime.Add(scheduledDowntime);
-			config.PostLoad(new ServerSettings(), [], []);
+			await config.PostLoadAsync(new ServerSettings(), [], []);
 			return scheduledDowntime;
 		}
 
@@ -47,7 +47,7 @@ namespace HordeServer.Tests.Server
 				FinishTime = clock.UtcNow + TimeSpan.FromSeconds(90)
 			};
 			config.Downtime.Add(downtime);
-			config.PostLoad(new ServerSettings(), [], []);
+			await config.PostLoadAsync(new ServerSettings(), [], []);
 			TestOptionsMonitor<GlobalConfig> optionsMonitor = new(config);
 			await using DowntimeService downtimeService = new(clock, optionsMonitor, NullLogger<DowntimeService>.Instance);
 			Assert.IsFalse(downtimeService.IsDowntimeActive);
@@ -57,9 +57,9 @@ namespace HordeServer.Tests.Server
 		}
 
 		[TestMethod]
-		public void IsNotActive()
+		public async Task IsNotActiveAsync()
 		{
-			ScheduledDowntime downtime = GetDowntimeFromJson("""
+			ScheduledDowntime downtime = await GetDowntimeFromJsonAsync("""
 				{
 					"startTime": "2020-09-06T03:00:00+00:00",
 					"finishTime": "2020-09-06T03:45:00+00:00"
@@ -69,9 +69,9 @@ namespace HordeServer.Tests.Server
 		}
 
 		[TestMethod]
-		public void IsActive()
+		public async Task IsActiveAsync()
 		{
-			ScheduledDowntime downtime = GetDowntimeFromJson("""
+			ScheduledDowntime downtime = await GetDowntimeFromJsonAsync("""
 				{
 					"startTime": "2020-09-06T03:00:00+00:00",
 					"finishTime": "2020-09-06T03:45:00+00:00"
@@ -81,9 +81,9 @@ namespace HordeServer.Tests.Server
 		}
 
 		[TestMethod]
-		public void IsNotActiveDaily()
+		public async Task IsNotActiveDailyAsync()
 		{
-			ScheduledDowntime downtime = GetDowntimeFromJson("""
+			ScheduledDowntime downtime = await GetDowntimeFromJsonAsync("""
 				{
 					"startTime": "2020-09-06T03:00:00+00:00",
 					"finishTime": "2020-09-06T03:45:00+00:00",
@@ -94,9 +94,9 @@ namespace HordeServer.Tests.Server
 		}
 
 		[TestMethod]
-		public void IsActiveDaily()
+		public async Task IsActiveDailyAsync()
 		{
-			ScheduledDowntime downtime = GetDowntimeFromJson("""
+			ScheduledDowntime downtime = await GetDowntimeFromJsonAsync("""
 				{
 					"startTime": "2020-09-06T03:00:00+00:00",
 					"finishTime": "2020-09-06T03:45:00+00:00",
@@ -107,9 +107,9 @@ namespace HordeServer.Tests.Server
 		}
 
 		[TestMethod]
-		public void IsNotActiveWeekly()
+		public async Task IsNotActiveWeeklyAsync()
 		{
-			ScheduledDowntime downtime = GetDowntimeFromJson("""
+			ScheduledDowntime downtime = await GetDowntimeFromJsonAsync("""
 				{
 					"startTime": "2020-09-06T03:00:00+00:00",
 					"finishTime": "2020-09-06T03:45:00+00:00",
@@ -120,9 +120,9 @@ namespace HordeServer.Tests.Server
 		}
 
 		[TestMethod]
-		public void IsActiveWeekly()
+		public async Task IsActiveWeeklyAsync()
 		{
-			ScheduledDowntime downtime = GetDowntimeFromJson("""
+			ScheduledDowntime downtime = await GetDowntimeFromJsonAsync("""
 				{
 					"startTime": "2020-09-06T03:00:00+00:00",
 					"finishTime": "2020-09-06T03:45:00+00:00",
@@ -133,9 +133,9 @@ namespace HordeServer.Tests.Server
 		}
 
 		[TestMethod]
-		public void FinishTimeFromDuration()
+		public async Task FinishTimeFromDurationAsync()
 		{
-			ScheduledDowntime downtime = GetDowntimeFromJson("""
+			ScheduledDowntime downtime = await GetDowntimeFromJsonAsync("""
 				{
 					"startTime": "2020-09-06T03:00:00+00:00",
 					"duration": "0:45"
@@ -145,9 +145,9 @@ namespace HordeServer.Tests.Server
 		}
 
 		[TestMethod]
-		public void DurationFromFinishTime()
+		public async Task DurationFromFinishTimeAsync()
 		{
-			ScheduledDowntime downtime = GetDowntimeFromJson("""
+			ScheduledDowntime downtime = await GetDowntimeFromJsonAsync("""
 				{
 					"startTime": "2020-09-06T03:00:00+00:00",
 					"finishTime": "2020-09-06T03:30:00+00:00",
@@ -157,9 +157,9 @@ namespace HordeServer.Tests.Server
 		}
 
 		[TestMethod]
-		public void DurationHigherPrecedence()
+		public async Task DurationHigherPrecedenceAsync()
 		{
-			ScheduledDowntime downtime = GetDowntimeFromJson("""
+			ScheduledDowntime downtime = await GetDowntimeFromJsonAsync("""
 				{
 					"startTime": "2020-09-06T03:00:00+00:00",
 					"finishTime": "2020-09-06T03:30:00+00:00",
@@ -170,9 +170,9 @@ namespace HordeServer.Tests.Server
 		}
 
 		[TestMethod]
-		public void StartTimeAndDurationMinuteAbbreviation()
+		public async Task StartTimeAndDurationMinuteAbbreviationAsync()
 		{
-			ScheduledDowntime downtime = GetDowntimeFromJson("""
+			ScheduledDowntime downtime = await GetDowntimeFromJsonAsync("""
 				{
 					"startTime": "Sun Sep 6, 2020, 3AM",
 					"duration": "0:45"
@@ -183,11 +183,11 @@ namespace HordeServer.Tests.Server
 		}
 		
 		[TestMethod]
-		public void WrongDayOfWeekNameError()
+		public async Task WrongDayOfWeekNameErrorAsync()
 		{
-			Assert.ThrowsException<FormatException>(() =>
+			await Assert.ThrowsExceptionAsync<FormatException>(async () =>
 			{
-				GetDowntimeFromJson("""
+				await GetDowntimeFromJsonAsync("""
 					{
 						"startTime": "Mon Sep 6, 2020, 3AM",
 						"finishTime": "Sun Sep 6, 2020, 3:30AM",
@@ -197,9 +197,9 @@ namespace HordeServer.Tests.Server
 		}
 
 		[TestMethod]
-		public void StartTimeAndDurationHourAbbreviation()
+		public async Task StartTimeAndDurationHourAbbreviationAsync()
 		{
-			ScheduledDowntime downtime = GetDowntimeFromJson("""
+			ScheduledDowntime downtime = await GetDowntimeFromJsonAsync("""
 				{
 					"startTime": "Sun Sep 6, 2020, 3AM",
 					"duration": "1:30"
@@ -210,9 +210,9 @@ namespace HordeServer.Tests.Server
 		}
 
 		[TestMethod]
-		public void StartTimeAndFinishTimeAbbreviation()
+		public async Task StartTimeAndFinishTimeAbbreviationAsync()
 		{
-			ScheduledDowntime downtime = GetDowntimeFromJson("""
+			ScheduledDowntime downtime = await GetDowntimeFromJsonAsync("""
 				{
 					"startTime": "Dec 1, 2024, 3AM",
 					"finishTime": "Sunday December 1 2024, 3:15AM"
@@ -223,9 +223,9 @@ namespace HordeServer.Tests.Server
 		}
 
 		[TestMethod]
-		public void WithTimeZone()
+		public async Task WithTimeZoneAsync()
 		{
-			ScheduledDowntime downtime = GetDowntimeFromJson("""
+			ScheduledDowntime downtime = await GetDowntimeFromJsonAsync("""
 				{
 					"startTime": "Jan 1, 2020, 3AM",
 					"duration": "1:30",
@@ -235,11 +235,11 @@ namespace HordeServer.Tests.Server
 			Assert.AreEqual(DateTimeOffsetParse("2020-01-01T03:00-05:00"), downtime.StartTime);
 			Assert.AreEqual(DateTimeOffsetParse("2020-01-01T04:30-05:00"), downtime.FinishTime);
 		}
-	
+
 		[TestMethod]
-		public void WithTimeZoneAndDaylightSaving()
+		public async Task WithTimeZoneAndDaylightSavingAsync()
 		{
-			ScheduledDowntime downtime = GetDowntimeFromJson("""
+			ScheduledDowntime downtime = await GetDowntimeFromJsonAsync("""
 				{
 					"startTime": "Sun Sep 6, 2020, 3AM",
 					"duration": "1:30",
@@ -249,11 +249,11 @@ namespace HordeServer.Tests.Server
 			Assert.AreEqual(DateTimeOffsetParse("2020-09-06T03:00-04:00"), downtime.StartTime);
 			Assert.AreEqual(DateTimeOffsetParse("2020-09-06T04:30-04:00"), downtime.FinishTime);
 		}
-	
+
 		[TestMethod]
-		public void WithTimeZoneIana()
+		public async Task WithTimeZoneIanaAsync()
 		{
-			ScheduledDowntime downtime = GetDowntimeFromJson("""
+			ScheduledDowntime downtime = await GetDowntimeFromJsonAsync("""
 				{
 					"startTime": "Jan 1, 2020, 3AM",
 					"duration": "1:30",
@@ -263,11 +263,11 @@ namespace HordeServer.Tests.Server
 			Assert.AreEqual(DateTimeOffsetParse("2020-01-01T03:00-05:00"), downtime.StartTime);
 			Assert.AreEqual(DateTimeOffsetParse("2020-01-01T04:30-05:00"), downtime.FinishTime);
 		}
-	
+
 		[TestMethod]
-		public void WithTimeZoneIanaAndDaylightSaving()
+		public async Task WithTimeZoneIanaAndDaylightSavingAsync()
 		{
-			ScheduledDowntime downtime = GetDowntimeFromJson("""
+			ScheduledDowntime downtime = await GetDowntimeFromJsonAsync("""
 				{
 					"startTime": "Sun Sep 6, 2020, 3AM",
 					"duration": "1:30",
@@ -279,11 +279,11 @@ namespace HordeServer.Tests.Server
 		}
 
 		[TestMethod]
-		public void WrongTimeZoneName()
+		public async Task WrongTimeZoneNameAsync()
 		{
-			Assert.ThrowsException<TimeZoneNotFoundException>(() =>
+			await Assert.ThrowsExceptionAsync<TimeZoneNotFoundException>(async () =>
 			{
-				GetDowntimeFromJson("""
+				await GetDowntimeFromJsonAsync("""
 					{
 						"startTime": "Jan 1, 2020, 3AM",
 						"timezone": "ET"
@@ -291,19 +291,19 @@ namespace HordeServer.Tests.Server
 				""");
 			});
 		}
-	
+
 		[TestMethod]
-		public void DowntimeToJson()
+		public async Task DowntimeToJsonAsync()
 		{
 			// -4 hours is Eastern Daylight Time
-			string json = JsonSerializer.Serialize(GetDowntimeFromJson("""
+			string json = JsonSerializer.Serialize(await GetDowntimeFromJsonAsync("""
 					{
 						"startTime": "Sun Sep 6, 2020, 3AM -4",
 						"duration": "0:30"
 					}
 				"""),
 				JsonUtils.DefaultSerializerOptions);
-			ScheduledDowntime downtime = GetDowntimeFromJson(json);
+			ScheduledDowntime downtime = await GetDowntimeFromJsonAsync(json);
 			Assert.AreEqual(DateTimeOffsetParse("2020-09-06T03:00:00-04:00"), downtime.StartTime);
 			Assert.AreEqual(DateTimeOffsetParse("2020-09-06T03:30:00-04:00"), downtime.FinishTime);
 			Assert.AreEqual(TimeSpan.FromMinutes(30), downtime.Duration);
@@ -312,9 +312,9 @@ namespace HordeServer.Tests.Server
 		}
 
 		[TestMethod]
-		public void DowntimeWithTimeZoneToJson()
+		public async Task DowntimeWithTimeZoneToJsonAsync()
 		{
-			string json = JsonSerializer.Serialize(GetDowntimeFromJson("""
+			string json = JsonSerializer.Serialize(await GetDowntimeFromJsonAsync("""
 					{
 						"startTime": "Jan 1, 2025, 3AM",
 						"duration": "0:30",
@@ -322,7 +322,7 @@ namespace HordeServer.Tests.Server
 					}
 				"""),
 				JsonUtils.DefaultSerializerOptions);
-			ScheduledDowntime downtime = GetDowntimeFromJson(json);
+			ScheduledDowntime downtime = await GetDowntimeFromJsonAsync(json);
 			Assert.AreEqual(DateTimeOffsetParse("2025-01-01T03:00:00-05:00"), downtime.StartTime);
 			Assert.AreEqual(DateTimeOffsetParse("2025-01-01T03:30:00-05:00"), downtime.FinishTime);
 			Assert.AreEqual(TimeSpan.FromMinutes(30), downtime.Duration);
diff --git a/Engine/Source/Programs/Horde/HordeServer/Configuration/ConfigService.cs b/Engine/Source/Programs/Horde/HordeServer/Configuration/ConfigService.cs
--- a/Engine/Source/Programs/Horde/HordeServer/Configuration/ConfigService.cs
+++ b/Engine/Source/Programs/Horde/HordeServer/Configuration/ConfigService.cs
@@ -282,7 +282,7 @@ namespace HordeServer.Configuration
 			try
 			{
 				GlobalConfig globalConfig = await ReadGlobalConfigAsync(context, cancellationToken);
-				globalConfig.PostLoad(_serverSettings, _pluginCollection.LoadedPlugins, _defaultAclModifiers);
+				await globalConfig.PostLoadAsync(_serverSettings, _pluginCollection.LoadedPlugins, _defaultAclModifiers);
 
 				foreach (OverrideConfigFile file in overrideFiles.Values)
 				{
@@ -392,7 +392,7 @@ namespace HordeServer.Configuration
 					data = SerializeSnapshot(snapshot);
 				}
 
-				ConfigState state = new ConfigState(IoHash.Compute(data.Span), CreateGlobalConfig(data));
+				ConfigState state = new ConfigState(IoHash.Compute(data.Span), await CreateGlobalConfigAsync(data));
 				await _health.UpdateAsync(HealthStatus.Healthy);
 
 				return state;
@@ -608,7 +608,7 @@ namespace HordeServer.Configuration
 				}
 
 				// Execute a PostLoad before returning so we can validate that everything is valid
-				globalConfig.PostLoad(_serverSettings, _pluginCollection.LoadedPlugins, _defaultAclModifiers);
+				await globalConfig.PostLoadAsync(_serverSettings, _pluginCollection.LoadedPlugins, _defaultAclModifiers);
 
 				return snapshot;
 			}
@@ -841,13 +841,13 @@ namespace HordeServer.Configuration
 			IoHash hash = IoHash.Compute(data.Span);
 			if (hash != state.Hash && state.Hash != IoHash.Zero) // Don't replace any explicit config override
 			{
-				GlobalConfig globalConfig = CreateGlobalConfig(data);
+				GlobalConfig globalConfig = await CreateGlobalConfigAsync(data);
 				_logger.LogInformation("Updating config state from Redis (hash: {OldHash} -> {NewHash}, new size: {Size})", state.Hash, hash, data.Length);
 				Set(hash, globalConfig);
 			}
 		}
 
-		GlobalConfig CreateGlobalConfig(ReadOnlyMemory<byte> data)
+		async Task<GlobalConfig> CreateGlobalConfigAsync(ReadOnlyMemory<byte> data)
 		{
 			// Decompress the snapshot object
 			ConfigSnapshot snapshot = DeserializeSnapshot(data);
@@ -857,7 +857,7 @@ namespace HordeServer.Configuration
 			globalConfig.Revision = IoHash.Compute(data.Span).ToString();
 
 			// Run the postload callbacks on all the config objects
-			globalConfig.PostLoad(_serverSettings, _pluginCollection.LoadedPlugins, _defaultAclModifiers);
+			await globalConfig.PostLoadAsync(_serverSettings, _pluginCollection.LoadedPlugins, _defaultAclModifiers);
 			return globalConfig;
 		}
 	}
diff --git a/Engine/Source/Programs/Horde/HordeServer/Server/GlobalConfig.cs b/Engine/Source/Programs/Horde/HordeServer/Server/GlobalConfig.cs
--- a/Engine/Source/Programs/Horde/HordeServer/Server/GlobalConfig.cs
+++ b/Engine/Source/Programs/Horde/HordeServer/Server/GlobalConfig.cs
@@ -6,6 +6,7 @@ using System.Diagnostics.CodeAnalysis;
 using System.Linq;
 using System.Security.Claims;
 using System.Text.Json.Serialization;
+using System.Threading.Tasks;
 using EpicGames.Core;
 using EpicGames.Horde.Acls;
 using HordeServer.Accounts;
@@ -104,7 +105,7 @@ namespace HordeServer.Server
 		/// <summary>
 		/// Called after the config file has been read
 		/// </summary>
-		public void PostLoad(ServerSettings serverSettings, IReadOnlyList<ILoadedPlugin> loadedPlugins, IEnumerable<IDefaultAclModifier> defaultAclModifiers)
+		public async Task PostLoadAsync(ServerSettings serverSettings, IReadOnlyList<ILoadedPlugin> loadedPlugins, IEnumerable<IDefaultAclModifier> defaultAclModifiers)
 		{
 			ServerSettings = serverSettings;
 			AclAction[] aclActions = AclConfig.GetActions(
@@ -157,7 +158,7 @@ namespace HordeServer.Server
 			PluginConfigOptions pluginConfigOptions = new PluginConfigOptions(VersionEnum, Plugins.Values, Acl);
 			foreach (IPluginConfig pluginConfig in pluginConfigs)
 			{
-				pluginConfig.PostLoad(pluginConfigOptions);
+				await pluginConfig.PostLoadAsync(pluginConfigOptions);
 			}
 
 			foreach (ScheduledDowntime downtime in Downtime)
diff --git a/Engine/Source/Programs/Horde/Plugins/Analytics/HordeServer.Analytics.Tests/Telemetry/TelemetryTests.cs b/Engine/Source/Programs/Horde/Plugins/Analytics/HordeServer.Analytics.Tests/Telemetry/TelemetryTests.cs
--- a/Engine/Source/Programs/Horde/Plugins/Analytics/HordeServer.Analytics.Tests/Telemetry/TelemetryTests.cs
+++ b/Engine/Source/Programs/Horde/Plugins/Analytics/HordeServer.Analytics.Tests/Telemetry/TelemetryTests.cs
@@ -45,7 +45,7 @@ namespace HordeServer.Analytics.Tests.Telemetry
 
 			GlobalConfig globalConfig = new GlobalConfig();
 			globalConfig.Plugins.AddAnalyticsTestConfig(analyticsConfig);
-			SetConfig(globalConfig);
+			await SetConfigAsync(globalConfig);
 
 			MetricTelemetrySink sink = ServiceProvider.GetRequiredService<MetricTelemetrySink>();
 			IMetricCollection collection = ServiceProvider.GetRequiredService<IMetricCollection>();
@@ -86,7 +86,7 @@ namespace HordeServer.Analytics.Tests.Telemetry
 
 			GlobalConfig globalConfig = new GlobalConfig();
 			globalConfig.Plugins.AddAnalyticsTestConfig(analyticsConfig);
-			SetConfig(globalConfig);
+			await SetConfigAsync(globalConfig);
 
 			MetricTelemetrySink sink = ServiceProvider.GetRequiredService<MetricTelemetrySink>();
 			IMetricCollection collection = ServiceProvider.GetRequiredService<IMetricCollection>();
@@ -167,7 +167,7 @@ namespace HordeServer.Analytics.Tests.Telemetry
 
 			GlobalConfig globalConfig = new GlobalConfig();
 			globalConfig.Plugins.AddAnalyticsTestConfig(analyticsConfig);
-			SetConfig(globalConfig);
+			await SetConfigAsync(globalConfig);
 
 			MetricTelemetrySink sink = ServiceProvider.GetRequiredService<MetricTelemetrySink>();
 			IMetricCollection collection = ServiceProvider.GetRequiredService<IMetricCollection>();
@@ -234,7 +234,7 @@ namespace HordeServer.Analytics.Tests.Telemetry
 
 			GlobalConfig globalConfig = new GlobalConfig();
 			globalConfig.Plugins.AddAnalyticsTestConfig(analyticsConfig);
-			SetConfig(globalConfig);
+			await SetConfigAsync(globalConfig);
 
 			MetricTelemetrySink sink = ServiceProvider.GetRequiredService<MetricTelemetrySink>();
 			IMetricCollection collection = ServiceProvider.GetRequiredService<IMetricCollection>();
@@ -274,7 +274,7 @@ namespace HordeServer.Analytics.Tests.Telemetry
 
 			GlobalConfig globalConfig = new GlobalConfig();
 			globalConfig.Plugins.AddAnalyticsTestConfig(analyticsConfig);
-			SetConfig(globalConfig);
+			await SetConfigAsync(globalConfig);
 
 			MetricTelemetrySink sink = ServiceProvider.GetRequiredService<MetricTelemetrySink>();
 			IMetricCollection collection = ServiceProvider.GetRequiredService<IMetricCollection>();
@@ -325,7 +325,7 @@ namespace HordeServer.Analytics.Tests.Telemetry
 
 			GlobalConfig globalConfig = new GlobalConfig();
 			globalConfig.Plugins.AddAnalyticsTestConfig(analyticsConfig);
-			SetConfig(globalConfig);
+			await SetConfigAsync(globalConfig);
 
 			MetricTelemetrySink sink = ServiceProvider.GetRequiredService<MetricTelemetrySink>();
 			IMetricCollection collection = ServiceProvider.GetRequiredService<IMetricCollection>();
@@ -386,7 +386,7 @@ namespace HordeServer.Analytics.Tests.Telemetry
 
 			GlobalConfig globalConfig = new GlobalConfig();
 			globalConfig.Plugins.AddAnalyticsTestConfig(analyticsConfig);
-			SetConfig(globalConfig);
+			await SetConfigAsync(globalConfig);
 
 			MetricTelemetrySink sink = ServiceProvider.GetRequiredService<MetricTelemetrySink>();
 			IMetricCollection collection = ServiceProvider.GetRequiredService<IMetricCollection>();
diff --git a/Engine/Source/Programs/Horde/Plugins/Analytics/HordeServer.Analytics/AnalyticsConfig.cs b/Engine/Source/Programs/Horde/Plugins/Analytics/HordeServer.Analytics/AnalyticsConfig.cs
--- a/Engine/Source/Programs/Horde/Plugins/Analytics/HordeServer.Analytics/AnalyticsConfig.cs
+++ b/Engine/Source/Programs/Horde/Plugins/Analytics/HordeServer.Analytics/AnalyticsConfig.cs
@@ -26,7 +26,7 @@ namespace HordeServer
 		private readonly Dictionary<TelemetryStoreId, TelemetryStoreConfig> _telemetryStoreLookup = new Dictionary<TelemetryStoreId, TelemetryStoreConfig>();
 
 		/// <inheritdoc/>
-		public void PostLoad(PluginConfigOptions configOptions)
+		public Task PostLoadAsync(PluginConfigOptions configOptions)
 		{
 			_parentAcl = configOptions.ParentAcl;
 
@@ -36,6 +36,8 @@ namespace HordeServer
 				_telemetryStoreLookup.Add(telemetryStore.Id, telemetryStore);
 				telemetryStore.PostLoad(configOptions.ParentAcl);
 			}
+
+			return Task.CompletedTask;
 		}
 
 		/// <inheritdoc cref="AclConfig.Authorize(AclAction, ClaimsPrincipal)"/>
diff --git a/Engine/Source/Programs/Horde/Plugins/Build/HordeServer.Build.Tests/Acls/AclTests.cs b/Engine/Source/Programs/Horde/Plugins/Build/HordeServer.Build.Tests/Acls/AclTests.cs
--- a/Engine/Source/Programs/Horde/Plugins/Build/HordeServer.Build.Tests/Acls/AclTests.cs
+++ b/Engine/Source/Programs/Horde/Plugins/Build/HordeServer.Build.Tests/Acls/AclTests.cs
@@ -17,9 +17,9 @@ namespace HordeServer.Tests.Acls
 	public class AclTests : BuildTestSetup
 	{
 		[TestMethod]
-		public void TestCustomProfile()
+		public async Task TestCustomProfileAsync()
 		{
-			UpdateConfig(globalConfig =>
+			await UpdateConfigAsync(globalConfig =>
 			{
 				globalConfig.Acl.Profiles = new List<AclProfileConfig>
 				{
@@ -78,9 +78,9 @@ namespace HordeServer.Tests.Acls
 		}
 
 		[TestMethod]
-		public void TestDefaultProfile()
+		public async Task TestDefaultProfileAsync()
 		{
-			UpdateConfig(globalConfig =>
+			await UpdateConfigAsync(globalConfig =>
 			{
 				ProjectConfig projectConfig = new ProjectConfig();
 				projectConfig.Id = new ProjectId("ue5");
diff --git a/Engine/Source/Programs/Horde/Plugins/Build/HordeServer.Build.Tests/Agents/AgentServiceTest.cs b/Engine/Source/Programs/Horde/Plugins/Build/HordeServer.Build.Tests/Agents/AgentServiceTest.cs
--- a/Engine/Source/Programs/Horde/Plugins/Build/HordeServer.Build.Tests/Agents/AgentServiceTest.cs
+++ b/Engine/Source/Programs/Horde/Plugins/Build/HordeServer.Build.Tests/Agents/AgentServiceTest.cs
@@ -146,7 +146,7 @@ public class AgentServiceTest : BuildTestSetup
 			new AgentRateConfig() { Condition = "aws-instance-type == 'c4.4xlarge' && osfamily == 'windows'", Rate = 300 }
 		};
 
-		UpdateConfig(x => ((ComputeConfig)x.Plugins[new PluginName("compute")]).Rates = agentRateConfigs);
+		await UpdateConfigAsync(x => ((ComputeConfig)x.Plugins[new PluginName("compute")]).Rates = agentRateConfigs);
 
 		double? rate1 = await AgentService.GetRateAsync(agent1.Id);
 		Assert.AreEqual(200, rate1!.Value, 0.1);
diff --git a/Engine/Source/Programs/Horde/Plugins/Build/HordeServer.Build.Tests/Agents/Pools/PoolUpdateServiceTest.cs b/Engine/Source/Programs/Horde/Plugins/Build/HordeServer.Build.Tests/Agents/Pools/PoolUpdateServiceTest.cs
--- a/Engine/Source/Programs/Horde/Plugins/Build/HordeServer.Build.Tests/Agents/Pools/PoolUpdateServiceTest.cs
+++ b/Engine/Source/Programs/Horde/Plugins/Build/HordeServer.Build.Tests/Agents/Pools/PoolUpdateServiceTest.cs
@@ -14,21 +14,18 @@ namespace HordeServer.Tests.Agents.Pools;
 [TestClass]
 public class PoolUpdateServiceTest : BuildTestSetup
 {
-	private readonly PoolUpdateService _pus;
+	private PoolUpdateService _pus = default!;
 	private IPool _pool = default!;
 	private IAgent _enabledAgent = default!;
 	private IAgent _disabledAgent = default!;
 	private IAgent _disabledAgentBeyondGracePeriod = default!;
 
-	public PoolUpdateServiceTest()
-	{
-		UpdateConfig(x => x.Plugins.GetComputeConfig().Pools.Clear());
-		_pus = new PoolUpdateService(AgentService, PoolService, PoolCollection, Clock, ServiceProvider.GetRequiredService<IOptionsMonitor<BuildConfig>>(), Tracer, new NullLogger<PoolUpdateService>());
-	}
 
 	[TestInitialize]
 	public async Task SetupAsync()
 	{
+		await UpdateConfigAsync(x => x.Plugins.GetComputeConfig().Pools.Clear());
+		_pus = new PoolUpdateService(AgentService, PoolService, PoolCollection, Clock, ServiceProvider.GetRequiredService<IOptionsMonitor<BuildConfig>>(), Tracer, new NullLogger<PoolUpdateService>());
 		_pool = await CreatePoolAsync(new PoolConfig { Name = "testPool", EnableAutoscaling = true });
 		_enabledAgent = await CreateAgentAsync(_pool, true);
 		_disabledAgent = await CreateAgentAsync(_pool, false, status: AgentStatus.Stopped);
@@ -67,7 +64,7 @@ public class PoolUpdateServiceTest : BuildTestSetup
 	public async Task ShutdownDisabledAgents_WithGlobalGracePeriod_RequestsShutdownAsync()
 	{
 		// Arrange
-		UpdateConfig(config => config.Plugins.GetBuildConfig().AgentShutdownIfDisabledGracePeriod = TimeSpan.FromHours(3.0));
+		await UpdateConfigAsync(config => config.Plugins.GetBuildConfig().AgentShutdownIfDisabledGracePeriod = TimeSpan.FromHours(3.0));
 
 		// Act
 		await _pus.ShutdownDisabledAgentsAsync(CancellationToken.None);
@@ -84,7 +81,7 @@ public class PoolUpdateServiceTest : BuildTestSetup
 	{
 		// Arrange
 		// Explicitly set the grace period for the pool to be longer than the default of 8 hours
-		UpdateConfig(config => config.Plugins.GetComputeConfig().Pools[0].ShutdownIfDisabledGracePeriod = TimeSpan.FromHours(24.0));
+		await UpdateConfigAsync(config => config.Plugins.GetComputeConfig().Pools[0].ShutdownIfDisabledGracePeriod = TimeSpan.FromHours(24.0));
 
 		// Act
 		await _pus.ShutdownDisabledAgentsAsync(CancellationToken.None);
@@ -100,7 +97,7 @@ public class PoolUpdateServiceTest : BuildTestSetup
 	public async Task ShutdownDisabledAgents_WithAutoScalingOff_DoesNotRequestShutdownAsync()
 	{
 		// Arrange
-		UpdateConfig(config => config.Plugins.GetComputeConfig().Pools[0].EnableAutoscaling = false);
+		await UpdateConfigAsync(config => config.Plugins.GetComputeConfig().Pools[0].EnableAutoscaling = false);
 
 		// Act
 		await _pus.ShutdownDisabledAgentsAsync(CancellationToken.None);
@@ -147,7 +144,7 @@ public class PoolUpdateServiceTest : BuildTestSetup
 			workspaceTypes[workspaceId] = new WorkspaceConfig { Identifier = workspaceId, ConformDiskFreeSpace = autoConformThresholdsMb[i] };
 		}
 		
-		UpdateConfig(x =>
+		await UpdateConfigAsync(x =>
 		{
 			x.Plugins.GetBuildConfig().Projects.Clear();
 			x.Plugins.GetBuildConfig().Projects.Add(new ProjectConfig { Streams = [new StreamConfig { WorkspaceTypes = workspaceTypes, AgentTypes = agentTypes }]});
diff --git a/Engine/Source/Programs/Horde/Plugins/Build/HordeServer.Build.Tests/Agents/Pools/PoolsControllerTest.cs b/Engine/Source/Programs/Horde/Plugins/Build/HordeServer.Build.Tests/Agents/Pools/PoolsControllerTest.cs
--- a/Engine/Source/Programs/Horde/Plugins/Build/HordeServer.Build.Tests/Agents/Pools/PoolsControllerTest.cs
+++ b/Engine/Source/Programs/Horde/Plugins/Build/HordeServer.Build.Tests/Agents/Pools/PoolsControllerTest.cs
@@ -9,9 +9,10 @@ namespace HordeServer.Tests.Agents.Pools
 	[TestClass]
 	public class PoolsControllerTest : BuildTestSetup
 	{
-		public PoolsControllerTest()
+		[TestInitialize]
+		public async Task UpdateConfigAsync()
 		{
-			UpdateConfig(x => x.Plugins.GetComputeConfig().Pools.Clear());
+			await UpdateConfigAsync(x => x.Plugins.GetComputeConfig().Pools.Clear());
 		}
 
 		[TestMethod]
@@ -24,7 +25,7 @@ namespace HordeServer.Tests.Agents.Pools
 			};
 #pragma warning restore CS0618
 
-			UpdateConfig(config =>
+			await UpdateConfigAsync(config =>
 			{
 				ComputeConfig computeConfig = config.Plugins.GetComputeConfig();
 				computeConfig.Pools.Clear();
@@ -54,7 +55,7 @@ namespace HordeServer.Tests.Agents.Pools
 			};
 #pragma warning restore CS0618
 
-			UpdateConfig(config =>
+			await UpdateConfigAsync(config =>
 			{
 				ComputeConfig computeConfig = config.Plugins.GetComputeConfig();
 				computeConfig.Pools.Clear();
diff --git a/Engine/Source/Programs/Horde/Plugins/Build/HordeServer.Build.Tests/Artifacts/ArtifactAclTests.cs b/Engine/Source/Programs/Horde/Plugins/Build/HordeServer.Build.Tests/Artifacts/ArtifactAclTests.cs
--- a/Engine/Source/Programs/Horde/Plugins/Build/HordeServer.Build.Tests/Artifacts/ArtifactAclTests.cs
+++ b/Engine/Source/Programs/Horde/Plugins/Build/HordeServer.Build.Tests/Artifacts/ArtifactAclTests.cs
@@ -28,9 +28,10 @@ namespace HordeServer.Tests.Artifacts
 		const string NormalReadRole = "normal-read";
 		const string SecretReadRole = "secret-read";
 
-		public ArtifactAclTests()
+		[TestInitialize]
+		public async Task UpdateConfigAsync()
 		{
-			UpdateConfig(SetupFixture);
+			await UpdateConfigAsync(SetupFixture);
 		}
 
 		void SetupFixture(GlobalConfig globalConfig)
diff --git a/Engine/Source/Programs/Horde/Plugins/Build/HordeServer.Build.Tests/Artifacts/ArtifactTests.cs b/Engine/Source/Programs/Horde/Plugins/Build/HordeServer.Build.Tests/Artifacts/ArtifactTests.cs
--- a/Engine/Source/Programs/Horde/Plugins/Build/HordeServer.Build.Tests/Artifacts/ArtifactTests.cs
+++ b/Engine/Source/Programs/Horde/Plugins/Build/HordeServer.Build.Tests/Artifacts/ArtifactTests.cs
@@ -15,7 +15,8 @@ namespace HordeServer.Tests.Artifacts
 	[TestClass]
 	public class ArtifactTests : BuildTestSetup
 	{
-		public ArtifactTests()
+		[TestInitialize]
+		public async Task UpdateConfigAsync()
 		{
 			void SetupFixture(GlobalConfig globalConfig)
 			{
@@ -27,7 +28,7 @@ namespace HordeServer.Tests.Artifacts
 				buildConfig.Projects.Add(projectConfig);
 			}
 
-			UpdateConfig(SetupFixture);
+			await UpdateConfigAsync(SetupFixture);
 		}
 
 		[TestMethod]
@@ -63,7 +64,7 @@ namespace HordeServer.Tests.Artifacts
 			ArtifactCollection expirationService = ServiceProvider.GetRequiredService<ArtifactCollection>();
 
 			ArtifactType type = new ArtifactType("my-artifact");
-			UpdateConfig(config => config.Plugins.GetBuildConfig().ArtifactTypes.Add(new ArtifactTypeConfig { Type = type, KeepDays = 1 }));
+			await UpdateConfigAsync(config => config.Plugins.GetBuildConfig().ArtifactTypes.Add(new ArtifactTypeConfig { Type = type, KeepDays = 1 }));
 
 			await expirationService.StartAsync(CancellationToken.None);
 
@@ -92,7 +93,7 @@ namespace HordeServer.Tests.Artifacts
 			DateTime startTime = Clock.UtcNow;
 
 			ArtifactType type = new ArtifactType("my-artifact");
-			UpdateConfig(config => config.Plugins.GetBuildConfig().ArtifactTypes.Add(new ArtifactTypeConfig { Type = type, KeepDays = 1 }));
+			await UpdateConfigAsync(config => config.Plugins.GetBuildConfig().ArtifactTypes.Add(new ArtifactTypeConfig { Type = type, KeepDays = 1 }));
 
 			ArtifactCollection expirationService = ServiceProvider.GetRequiredService<ArtifactCollection>();
 
@@ -109,7 +110,7 @@ namespace HordeServer.Tests.Artifacts
 				Assert.AreEqual(artifact.Id, artifacts[0].Id);
 			}
 
-			UpdateConfig(config => config.Plugins.GetBuildConfig().ArtifactTypes = new List<ArtifactTypeConfig> { new ArtifactTypeConfig { Type = type, KeepDays = 4 } });
+			await UpdateConfigAsync(config => config.Plugins.GetBuildConfig().ArtifactTypes = new List<ArtifactTypeConfig> { new ArtifactTypeConfig { Type = type, KeepDays = 4 } });
 			await Clock.AdvanceAsync(TimeSpan.FromDays(2.0));
 
 			{
@@ -117,7 +118,7 @@ namespace HordeServer.Tests.Artifacts
 				Assert.AreEqual(1, artifacts.Count);
 			}
 
-			UpdateConfig(config => config.Plugins.GetBuildConfig().ArtifactTypes = new List<ArtifactTypeConfig> { new ArtifactTypeConfig { Type = type, KeepDays = 1 } });
+			await UpdateConfigAsync(config => config.Plugins.GetBuildConfig().ArtifactTypes = new List<ArtifactTypeConfig> { new ArtifactTypeConfig { Type = type, KeepDays = 1 } });
 			await Clock.AdvanceAsync(TimeSpan.FromDays(1.0));
 
 			{
@@ -132,7 +133,7 @@ namespace HordeServer.Tests.Artifacts
 			DateTime startTime = Clock.UtcNow;
 
 			ArtifactType type = new ArtifactType("my-artifact");
-			UpdateConfig(config => config.Plugins.GetBuildConfig().ArtifactTypes.Add(new ArtifactTypeConfig { Type = type, KeepCount = 4 }));
+			await UpdateConfigAsync(config => config.Plugins.GetBuildConfig().ArtifactTypes.Add(new ArtifactTypeConfig { Type = type, KeepCount = 4 }));
 
 			ArtifactCollection expirationService = ServiceProvider.GetRequiredService<ArtifactCollection>();
 			await expirationService.StartAsync(CancellationToken.None);
@@ -171,7 +172,7 @@ namespace HordeServer.Tests.Artifacts
 				}
 			}
 
-			UpdateConfig(config => config.Plugins.GetBuildConfig().ArtifactTypes = new List<ArtifactTypeConfig> { new ArtifactTypeConfig { Type = type, KeepDays = 1 } });
+			await UpdateConfigAsync(config => config.Plugins.GetBuildConfig().ArtifactTypes = new List<ArtifactTypeConfig> { new ArtifactTypeConfig { Type = type, KeepDays = 1 } });
 			await Clock.AdvanceAsync(TimeSpan.FromDays(2.0));
 
 			{
@@ -197,14 +198,14 @@ namespace HordeServer.Tests.Artifacts
 
 			// Orphan the artifact type
 			{
-				UpdateConfig(config => config.Plugins.GetBuildConfig().ArtifactTypes.Add(new ArtifactTypeConfig { Type = type, KeepCount = 4 }));
+				await UpdateConfigAsync(config => config.Plugins.GetBuildConfig().ArtifactTypes.Add(new ArtifactTypeConfig { Type = type, KeepCount = 4 }));
 				IArtifactBuilder artifact = await artifactCollection.CreateAsync(new ArtifactName($"default"), type, null, fooStreamId, CommitIdWithOrder.FromPerforceChange(1), new string[] { "test1", "test2" }, Array.Empty<string>());
 
 				await Clock.AdvanceAsync(TimeSpan.FromDays(30.0));
 				IArtifact? artifact2 = await artifactCollection.GetAsync(artifact.Id);
 				Assert.IsNotNull(artifact2);
 
-				UpdateConfig(config => config.Plugins.GetBuildConfig().ArtifactTypes.Clear());
+				await UpdateConfigAsync(config => config.Plugins.GetBuildConfig().ArtifactTypes.Clear());
 				IArtifact? artifact3 = await artifactCollection.GetAsync(artifact.Id);
 				Assert.IsNotNull(artifact3);
 
@@ -215,14 +216,14 @@ namespace HordeServer.Tests.Artifacts
 
 			// Orphan the stream
 			{
-				UpdateConfig(config => config.Plugins.GetBuildConfig().ArtifactTypes.Add(new ArtifactTypeConfig { Type = type, KeepCount = 4 }));
+				await UpdateConfigAsync(config => config.Plugins.GetBuildConfig().ArtifactTypes.Add(new ArtifactTypeConfig { Type = type, KeepCount = 4 }));
 				IArtifactBuilder artifact = await artifactCollection.CreateAsync(new ArtifactName($"default"), type, null, fooStreamId, CommitIdWithOrder.FromPerforceChange(1), new string[] { "test1", "test2" }, Array.Empty<string>());
 
 				await Clock.AdvanceAsync(TimeSpan.FromDays(30.0));
 				IArtifact? artifact2 = await artifactCollection.GetAsync(artifact.Id);
 				Assert.IsNotNull(artifact2);
 
-				UpdateConfig(config => config.Plugins.GetBuildConfig().Projects.Clear());
+				await UpdateConfigAsync(config => config.Plugins.GetBuildConfig().Projects.Clear());
 				IArtifact? artifact3 = await artifactCollection.GetAsync(artifact.Id);
 				Assert.IsNotNull(artifact3);
 
diff --git a/Engine/Source/Programs/Horde/Plugins/Build/HordeServer.Build.Tests/BuildTestSetup.cs b/Engine/Source/Programs/Horde/Plugins/Build/HordeServer.Build.Tests/BuildTestSetup.cs
--- a/Engine/Source/Programs/Horde/Plugins/Build/HordeServer.Build.Tests/BuildTestSetup.cs
+++ b/Engine/Source/Programs/Horde/Plugins/Build/HordeServer.Build.Tests/BuildTestSetup.cs
@@ -351,7 +351,7 @@ namespace HordeServer.Tests
 		
 		protected async Task<IPool> CreatePoolViaConfigAsync(PoolConfig poolConfig)
 		{
-			UpdateConfig(config => config.Plugins.GetComputeConfig().Pools.Add(poolConfig));
+			await UpdateConfigAsync(config => config.Plugins.GetComputeConfig().Pools.Add(poolConfig));
 			return await PoolCollection.GetAsync(poolConfig.Id) ?? throw new Exception($"Unable to get pool ID {poolConfig.Id}");
 		}
 		
diff --git a/Engine/Source/Programs/Horde/Plugins/Build/HordeServer.Build.Tests/Compute/ComputeServiceTest.cs b/Engine/Source/Programs/Horde/Plugins/Build/HordeServer.Build.Tests/Compute/ComputeServiceTest.cs
--- a/Engine/Source/Programs/Horde/Plugins/Build/HordeServer.Build.Tests/Compute/ComputeServiceTest.cs
+++ b/Engine/Source/Programs/Horde/Plugins/Build/HordeServer.Build.Tests/Compute/ComputeServiceTest.cs
@@ -289,7 +289,7 @@ namespace HordeServer.Tests.Compute
 				new() { Id = cluster1, Condition = $"pool == '{poolFoo.ToString()}'" },
 				new() { Id = cluster2, Condition = $"pool == '{poolBar.ToString()}'" }
 			};
-			GlobalConfig.CurrentValue.PostLoad(new ServerSettings(), Array.Empty<ILoadedPlugin>(), Array.Empty<IDefaultAclModifier>());
+			await GlobalConfig.CurrentValue.PostLoadAsync(new ServerSettings(), Array.Empty<ILoadedPlugin>(), Array.Empty<IDefaultAclModifier>());
 
 			List<string> props = ["ComputeIp=11.0.0.1", "ComputePort=5000"];
 			IAgent agent1 = await CreateAgentAsync(new PoolId("foo"), properties: props);
@@ -324,7 +324,7 @@ namespace HordeServer.Tests.Compute
 		}
 
 		[TestMethod]
-		public void Cluster_FindBest()
+		public async Task Cluster_FindBest_Async()
 		{
 			ComputeConfig computeConfig = new()
 			{
@@ -335,7 +335,7 @@ namespace HordeServer.Tests.Compute
 			};
 
 			PluginConfigOptions configOptions = new PluginConfigOptions(ConfigVersion.Latest, Enumerable.Empty<IPluginConfig>(), new HordeServer.Acls.AclConfig());
-			computeConfig.PostLoad(configOptions);
+			await computeConfig.PostLoadAsync(configOptions);
 
 			Assert.AreEqual(new ClusterId("compute1"), ComputeService.FindBestComputeClusterId(computeConfig, IPAddress.Parse("11.0.0.1")));
 			Assert.AreEqual(new ClusterId("compute2"), ComputeService.FindBestComputeClusterId(computeConfig, IPAddress.Parse("12.0.1.1")));
@@ -345,7 +345,7 @@ namespace HordeServer.Tests.Compute
 			{
 				Networks = [new NetworkConfig { CidrBlock = "0.0.0.0/0", Id = "catchAll", ComputeId = "catchAll" }]
 			};
-			globalConfigCatchAll.PostLoad(configOptions);
+			await globalConfigCatchAll.PostLoadAsync(configOptions);
 			Assert.AreEqual(new ClusterId("catchAll"), ComputeService.FindBestComputeClusterId(globalConfigCatchAll, IPAddress.Parse("123.123.123.123")));
 		}
 
@@ -416,7 +416,7 @@ namespace HordeServer.Tests.Compute
 				NullLogger<ComputeService>.Instance);
 			await CreateAgentAsync(new PoolId(pool), properties: ["ComputeIp=11.0.0.1", "ComputePort=5000"]);
 			computeConfig.CurrentValue.Clusters = [ccc];
-			computeConfig.CurrentValue.PostLoad(new PluginConfigOptions(ConfigVersion.Latest, Enumerable.Empty<IPluginConfig>(), new AclConfig()));
+			await computeConfig.CurrentValue.PostLoadAsync(new PluginConfigOptions(ConfigVersion.Latest, Enumerable.Empty<IPluginConfig>(), new AclConfig()));
 			AgentRelayService.SetRandomSeed(1);
 			return cs;
 		}
diff --git a/Engine/Source/Programs/Horde/Plugins/Build/HordeServer.Build.Tests/Devices/DeviceServiceTests.cs b/Engine/Source/Programs/Horde/Plugins/Build/HordeServer.Build.Tests/Devices/DeviceServiceTests.cs
--- a/Engine/Source/Programs/Horde/Plugins/Build/HordeServer.Build.Tests/Devices/DeviceServiceTests.cs
+++ b/Engine/Source/Programs/Horde/Plugins/Build/HordeServer.Build.Tests/Devices/DeviceServiceTests.cs
@@ -154,7 +154,7 @@ namespace HordeServer.Tests.Devices
 				devices.Platforms.Add(platform);
 			}
 
-			UpdateConfig(x => x.Plugins.GetBuildConfig().Devices = devices);
+			await UpdateConfigAsync(x => x.Plugins.GetBuildConfig().Devices = devices);
 
 			for (int i = 1; i < 4; i++)
 			{
diff --git a/Engine/Source/Programs/Horde/Plugins/Build/HordeServer.Build.Tests/Fixture.cs b/Engine/Source/Programs/Horde/Plugins/Build/HordeServer.Build.Tests/Fixture.cs
--- a/Engine/Source/Programs/Horde/Plugins/Build/HordeServer.Build.Tests/Fixture.cs
+++ b/Engine/Source/Programs/Horde/Plugins/Build/HordeServer.Build.Tests/Fixture.cs
@@ -79,7 +79,7 @@ namespace HordeServer.Tests
 
 			GlobalConfig globalConfig = new GlobalConfig();
 			globalConfig.Plugins.AddBuildConfig(buildConfig);
-			globalConfig.PostLoad(serverSettings, pluginCollection.LoadedPlugins, Array.Empty<IDefaultAclModifier>());
+			await globalConfig.PostLoadAsync(serverSettings, pluginCollection.LoadedPlugins, Array.Empty<IDefaultAclModifier>());
 			configService.OverrideConfig(globalConfig);
 
 			StreamId = streamId;
diff --git a/Engine/Source/Programs/Horde/Plugins/Build/HordeServer.Build.Tests/Fleet/FleetServiceTest.cs b/Engine/Source/Programs/Horde/Plugins/Build/HordeServer.Build.Tests/Fleet/FleetServiceTest.cs
--- a/Engine/Source/Programs/Horde/Plugins/Build/HordeServer.Build.Tests/Fleet/FleetServiceTest.cs
+++ b/Engine/Source/Programs/Horde/Plugins/Build/HordeServer.Build.Tests/Fleet/FleetServiceTest.cs
@@ -95,9 +95,10 @@ namespace HordeServer.Tests.Fleet
 	{
 		readonly FleetManagerSpy _fleetManagerSpy = new();
 
-		public FleetServiceTest()
+		[TestInitialize]
+		public async Task UpdateConfigAsync()
 		{
-			UpdateConfig(x => x.Plugins.GetComputeConfig().Pools.Clear());
+			await UpdateConfigAsync(x => x.Plugins.GetComputeConfig().Pools.Clear());
 		}
 
 		[TestMethod]
diff --git a/Engine/Source/Programs/Horde/Plugins/Build/HordeServer.Build.Tests/Fleet/JobQueueStrategyTest.cs b/Engine/Source/Programs/Horde/Plugins/Build/HordeServer.Build.Tests/Fleet/JobQueueStrategyTest.cs
--- a/Engine/Source/Programs/Horde/Plugins/Build/HordeServer.Build.Tests/Fleet/JobQueueStrategyTest.cs
+++ b/Engine/Source/Programs/Horde/Plugins/Build/HordeServer.Build.Tests/Fleet/JobQueueStrategyTest.cs
@@ -117,7 +117,7 @@ namespace HordeServer.Tests.Fleet
 		/// <param name="isDowntimeActive"></param>
 		private async Task<(JobQueueStrategy, PoolSizeResult, IPool, List<IAgent> agents)> SetUpJobsAsync(int numBatchesRunning, int numBatchesReady, int numAgents = 8, bool isDowntimeActive = false)
 		{
-			SetConfig(new GlobalConfig());
+			await SetConfigAsync(new GlobalConfig());
 
 			IPool pool = await CreatePoolAsync(new PoolConfig { Name = "bogusPool" + ++s_uniqueId, EnableAutoscaling = true, MinAgents = 0, NumReserveAgents = 0 });
 			List<IAgent> agents = new();
@@ -137,7 +137,7 @@ namespace HordeServer.Tests.Fleet
 			projectConfig.Id = new ProjectId("ue5");
 			projectConfig.Streams.Add(streamConfig);
 
-			UpdateConfig(config => config.Plugins.GetBuildConfig().Projects = new List<ProjectConfig> { projectConfig });
+			await UpdateConfigAsync(config => config.Plugins.GetBuildConfig().Projects = new List<ProjectConfig> { projectConfig });
 
 			string nodeForAgentType1 = "bogusNodeOnAgentType" + ++s_uniqueId;
 			IGraph graph = await GraphCollection.AppendAsync(null, new()
diff --git a/Engine/Source/Programs/Horde/Plugins/Build/HordeServer.Build.Tests/Issues/IssueServiceTests.cs b/Engine/Source/Programs/Horde/Plugins/Build/HordeServer.Build.Tests/Issues/IssueServiceTests.cs
--- a/Engine/Source/Programs/Horde/Plugins/Build/HordeServer.Build.Tests/Issues/IssueServiceTests.cs
+++ b/Engine/Source/Programs/Horde/Plugins/Build/HordeServer.Build.Tests/Issues/IssueServiceTests.cs
@@ -116,15 +116,16 @@ namespace HordeServer.Tests.Issues
 
 		const string DevStreamName = "//UE4/Dev";
 		readonly StreamId _devStreamId = new StreamId("ue4-dev");
-		readonly IGraph _graph;
-		readonly PerforceServiceStub _perforce;		
-		readonly UserId _jerryId;
-		readonly UserId _bobId;
-		readonly UserId _chrisId;
-		readonly DirectoryReference _autoSdkDir;
-		readonly DirectoryReference _workspaceDir;
-
-		public IssueServiceTests()
+		IGraph _graph = default!;
+		PerforceServiceStub _perforce = default!;
+		UserId _jerryId;
+		UserId _bobId;
+		UserId _chrisId;
+		DirectoryReference _autoSdkDir = default!;
+		DirectoryReference _workspaceDir = default!;
+
+		[TestInitialize]
+		public async Task SetupAsync()
 		{
 			if (RuntimeInformation.IsOSPlatform(OSPlatform.Windows))
 			{
@@ -155,7 +156,7 @@ namespace HordeServer.Tests.Issues
 			GlobalConfig globalConfig = new GlobalConfig();
 			globalConfig.Plugins.AddBuildConfig(buildConfig);
 			globalConfig.Plugins.AddStorageConfig(storageConfig);
-			SetConfig(globalConfig);
+			await SetConfigAsync(globalConfig);
 
 			static StreamConfig CreateStream(StreamId streamId, string streamName)
 			{
@@ -173,16 +174,16 @@ namespace HordeServer.Tests.Issues
 				};
 			}
 
-			IUser bill = UserCollection.FindOrAddUserByLoginAsync("Bill").Result;
-			IUser anne = UserCollection.FindOrAddUserByLoginAsync("Anne").Result;
-			IUser bob = UserCollection.FindOrAddUserByLoginAsync("Bob").Result;
-			IUser jerry = UserCollection.FindOrAddUserByLoginAsync("Jerry").Result;
-			IUser chris = UserCollection.FindOrAddUserByLoginAsync("Chris").Result;
-			IUser tim = UserCollection.FindOrAddUserByLoginAsync("Tim").Result;
-			
+			IUser bill = await UserCollection.FindOrAddUserByLoginAsync("Bill");
+			IUser anne = await UserCollection.FindOrAddUserByLoginAsync("Anne");
+			IUser bob = await UserCollection.FindOrAddUserByLoginAsync("Bob");
+			IUser jerry = await UserCollection.FindOrAddUserByLoginAsync("Jerry");
+			IUser chris = await UserCollection.FindOrAddUserByLoginAsync("Chris");
+			IUser tim = await UserCollection.FindOrAddUserByLoginAsync("Tim");
+
 			_chrisId = chris.Id;
-			_jerryId = UserCollection.FindOrAddUserByLoginAsync("Jerry").Result.Id;
-			_bobId = UserCollection.FindOrAddUserByLoginAsync("Bob").Result.Id;
+			_jerryId = (await UserCollection.FindOrAddUserByLoginAsync("Jerry")).Id;
+			_bobId = (await UserCollection.FindOrAddUserByLoginAsync("Bob")).Id;
 
 			_perforce = PerforceService;
 			_perforce.AddChange(_mainStreamId, 100, bill, "Description", new string[] { "a/b.cpp" });
@@ -424,7 +425,7 @@ namespace HordeServer.Tests.Issues
 			// Scenario: Stream is deleted
 			// Expected: Issue is closed
 			{
-				UpdateConfig(x => x.Plugins.GetBuildConfig().Projects.Clear());
+				await UpdateConfigAsync(x => x.Plugins.GetBuildConfig().Projects.Clear());
 				await Clock.AdvanceAsync(TimeSpan.FromHours(1.0));
 
 				IReadOnlyList<IIssue> issues = await IssueCollection.FindIssuesAsync();
diff --git a/Engine/Source/Programs/Horde/Plugins/Build/HordeServer.Build.Tests/Jobs/Bisect/BisectTests.cs b/Engine/Source/Programs/Horde/Plugins/Build/HordeServer.Build.Tests/Jobs/Bisect/BisectTests.cs
--- a/Engine/Source/Programs/Horde/Plugins/Build/HordeServer.Build.Tests/Jobs/Bisect/BisectTests.cs
+++ b/Engine/Source/Programs/Horde/Plugins/Build/HordeServer.Build.Tests/Jobs/Bisect/BisectTests.cs
@@ -136,7 +136,7 @@ namespace HordeServer.Tests.Jobs.Bisect
 			project.Id = new ProjectId("ue4");
 			project.Streams.Add(new StreamConfig { Id = StreamId, Templates = new List<TemplateRefConfig> { new TemplateRefConfig { Id = TemplateId } } });
 
-			UpdateConfig(x => x.Plugins.GetBuildConfig().Projects.Add(project));
+			await UpdateConfigAsync(x => x.Plugins.GetBuildConfig().Projects.Add(project));
 
 			// Create the graph
 			List<NewNode> nodes = new List<NewNode>();
diff --git a/Engine/Source/Programs/Horde/Plugins/Build/HordeServer.Build.Tests/Jobs/JobCollectionTests.cs b/Engine/Source/Programs/Horde/Plugins/Build/HordeServer.Build.Tests/Jobs/JobCollectionTests.cs
--- a/Engine/Source/Programs/Horde/Plugins/Build/HordeServer.Build.Tests/Jobs/JobCollectionTests.cs
+++ b/Engine/Source/Programs/Horde/Plugins/Build/HordeServer.Build.Tests/Jobs/JobCollectionTests.cs
@@ -68,9 +68,10 @@ namespace HordeServer.Tests.Jobs
 			return result.Value! as T;
 		}
 
-		public JobCollectionTests()
+		[TestInitialize]
+		public async Task UpdateConfigAsync()
 		{
-			UpdateConfig(globalConfig =>
+			await UpdateConfigAsync(globalConfig =>
 			{
 				ProjectConfig projectConfig = new() { Id = new ProjectId("ue4") };
 				projectConfig.Streams.Add(new StreamConfig { Id = _streamId });
@@ -961,7 +962,7 @@ namespace HordeServer.Tests.Jobs
 			GlobalConfig globalConfig = new GlobalConfig();
 			globalConfig.Plugins.AddBuildConfig(buildConfig);
 
-			SetConfig(globalConfig);
+			await SetConfigAsync(globalConfig);
 
 			CreateJobRequest request = new CreateJobRequest(new StreamId("ue5-main"), new TemplateId("template1"));
 			request.CommitId = CommitIdWithOrder.FromPerforceChange(123);
diff --git a/Engine/Source/Programs/Horde/Plugins/Build/HordeServer.Build.Tests/Jobs/JobExpirationTests.cs b/Engine/Source/Programs/Horde/Plugins/Build/HordeServer.Build.Tests/Jobs/JobExpirationTests.cs
--- a/Engine/Source/Programs/Horde/Plugins/Build/HordeServer.Build.Tests/Jobs/JobExpirationTests.cs
+++ b/Engine/Source/Programs/Horde/Plugins/Build/HordeServer.Build.Tests/Jobs/JobExpirationTests.cs
@@ -36,7 +36,7 @@ namespace HordeServer.Tests.Jobs
 			GlobalConfig globalConfig = new GlobalConfig();
 			globalConfig.Plugins.AddBuildConfig(buildConfig);
 
-			SetConfig(globalConfig);
+			await SetConfigAsync(globalConfig);
 
 			CreateJobOptions options = new CreateJobOptions();
 			options.PreflightCommitId = CommitId.FromPerforceChange(999);
diff --git a/Engine/Source/Programs/Horde/Plugins/Build/HordeServer.Build.Tests/Jobs/JobServiceTests.cs b/Engine/Source/Programs/Horde/Plugins/Build/HordeServer.Build.Tests/Jobs/JobServiceTests.cs
--- a/Engine/Source/Programs/Horde/Plugins/Build/HordeServer.Build.Tests/Jobs/JobServiceTests.cs
+++ b/Engine/Source/Programs/Horde/Plugins/Build/HordeServer.Build.Tests/Jobs/JobServiceTests.cs
@@ -48,7 +48,7 @@ namespace HordeServer.Tests.Jobs
 			GlobalConfig globalConfig = new GlobalConfig();
 			globalConfig.Plugins.AddBuildConfig(buildConfig);
 
-			SetConfig(globalConfig);
+			await SetConfigAsync(globalConfig);
 
 			CreateJobOptions options = new CreateJobOptions();
 			options.PreflightCommitId = CommitId.FromPerforceChange(999);
@@ -91,7 +91,7 @@ namespace HordeServer.Tests.Jobs
 
 			GlobalConfig config = new GlobalConfig();
 			config.Plugins.AddBuildConfig(buildConfig);
-			SetConfig(config);
+			await SetConfigAsync(config);
 
 			IUser user = await UserCollection.FindOrAddUserByLoginAsync("Bob");
 
@@ -241,7 +241,7 @@ namespace HordeServer.Tests.Jobs
 
 			GlobalConfig globalConfig = new();
 			globalConfig.Plugins.AddBuildConfig(buildConfig);
-			SetConfig(globalConfig);
+			await SetConfigAsync(globalConfig);
 
 			// ----
 
diff --git a/Engine/Source/Programs/Horde/Plugins/Build/HordeServer.Build.Tests/Jobs/Schedules/SchedulerTests.cs b/Engine/Source/Programs/Horde/Plugins/Build/HordeServer.Build.Tests/Jobs/Schedules/SchedulerTests.cs
--- a/Engine/Source/Programs/Horde/Plugins/Build/HordeServer.Build.Tests/Jobs/Schedules/SchedulerTests.cs
+++ b/Engine/Source/Programs/Horde/Plugins/Build/HordeServer.Build.Tests/Jobs/Schedules/SchedulerTests.cs
@@ -26,12 +26,13 @@ namespace HordeServer.Tests.Jobs.Schedules
 		StreamId StreamId { get; } = new StreamId("ue5-main");
 		TemplateId TemplateId { get; } = new TemplateId("template1");
 
-		readonly ITemplate _template;
-		readonly HashSet<JobId> _initialJobIds;
+		ITemplate _template = default!;
+		HashSet<JobId> _initialJobIds = default!;
 
-		public SchedulerTests()
+		[TestInitialize]
+		public async Task SetupAsync()
 		{
-			IUser bob = UserCollection.FindOrAddUserByLoginAsync("Bob").Result;
+			IUser bob = await UserCollection.FindOrAddUserByLoginAsync("Bob");
 
 			BuildConfig buildConfig = new BuildConfig();
 			buildConfig.Projects.Add(new ProjectConfig { Id = ProjectId, Name = "UE4" });
@@ -39,11 +40,11 @@ namespace HordeServer.Tests.Jobs.Schedules
 			GlobalConfig globalConfig = new GlobalConfig();
 			globalConfig.Plugins.AddBuildConfig(buildConfig);
 
-			SetConfig(globalConfig);
+			await SetConfigAsync(globalConfig);
 
-			_template = TemplateCollection.GetOrAddAsync(new TemplateConfig { Name = "Test template" }).Result;
+			_template = await TemplateCollection.GetOrAddAsync(new TemplateConfig { Name = "Test template" });
 
-			_initialJobIds = new HashSet<JobId>(JobCollection.FindAsync(new FindJobOptions()).Result.Select(x => x.Id));
+			_initialJobIds = new HashSet<JobId>((await JobCollection.FindAsync(new FindJobOptions())).Select(x => x.Id));
 
 			PerforceService.Changes.Clear();
 			PerforceService.AddChange(StreamId, 100, bob, "", new[] { "code.cpp" });
@@ -61,7 +62,7 @@ namespace HordeServer.Tests.Jobs.Schedules
 			streamConfig.Tabs.Add(new TabConfig { Title = "foo", Templates = new List<TemplateId> { TemplateId } });
 			streamConfig.Templates.Add(new TemplateRefConfig { Id = TemplateId, Name = "Test", Schedule = schedule });
 
-			UpdateConfig(x => x.Plugins.GetBuildConfig().Projects[0].Streams = new List<StreamConfig> { streamConfig });
+			await UpdateConfigAsync(x => x.Plugins.GetBuildConfig().Projects[0].Streams = new List<StreamConfig> { streamConfig });
 
 			return (await StreamCollection.GetAsync(streamConfig.Id))!;
 		}
@@ -446,7 +447,7 @@ namespace HordeServer.Tests.Jobs.Schedules
 			config.Name = "//UE5/Main";
 			config.Tabs.Add(new TabConfig { Title = "foo", Templates = new List<TemplateId> { newTemplateRefId1, newTemplateRefId2 } });
 			config.Templates = new() { new TemplateRefConfig { Id = newTemplateRefId1 }, new TemplateRefConfig { Id = newTemplateRefId2 } };
-			UpdateConfig(x => x.Plugins.GetBuildConfig().Projects[0].Streams = new List<StreamConfig> { config });
+			await UpdateConfigAsync(x => x.Plugins.GetBuildConfig().Projects[0].Streams = new List<StreamConfig> { config });
 
 			IStream stream = (await StreamCollection.GetAsync(config.Id))!;
 
@@ -538,7 +539,7 @@ namespace HordeServer.Tests.Jobs.Schedules
 			config.Tabs.Add(new TabConfig { Title = "foo", Templates = new List<TemplateId> { newTemplateRefId1, newTemplateRefId2 } });
 			config.Templates.Add(newTemplate1);
 			config.Templates.Add(newTemplate2);
-			UpdateConfig(x => x.Plugins.GetBuildConfig().Projects[0].Streams = new List<StreamConfig> { config });
+			await UpdateConfigAsync(x => x.Plugins.GetBuildConfig().Projects[0].Streams = new List<StreamConfig> { config });
 
 			IStream stream = (await StreamCollection.GetAsync(config.Id))!;
 
@@ -701,7 +702,7 @@ namespace HordeServer.Tests.Jobs.Schedules
 			streamConfig.Name = "//UE5/Main";
 			streamConfig.Tabs.Add(new TabConfig { Title = "foo", Templates = new List<TemplateId> { TemplateId } });
 			streamConfig.Templates.Add(templateConfig);
-			UpdateConfig(x => x.Plugins.GetBuildConfig().Projects[0].Streams = new List<StreamConfig> { streamConfig });
+			await UpdateConfigAsync(x => x.Plugins.GetBuildConfig().Projects[0].Streams = new List<StreamConfig> { streamConfig });
 
 			// Make sure we don't have any jobs to start with
 			await ScheduleService.TickForTestingAsync();
diff --git a/Engine/Source/Programs/Horde/Plugins/Build/HordeServer.Build.Tests/Jobs/TestData/TestDataTests.cs b/Engine/Source/Programs/Horde/Plugins/Build/HordeServer.Build.Tests/Jobs/TestData/TestDataTests.cs
--- a/Engine/Source/Programs/Horde/Plugins/Build/HordeServer.Build.Tests/Jobs/TestData/TestDataTests.cs
+++ b/Engine/Source/Programs/Horde/Plugins/Build/HordeServer.Build.Tests/Jobs/TestData/TestDataTests.cs
@@ -37,7 +37,7 @@ namespace HordeServer.Tests.Jobs.TestData
 		const string ReleaseStreamName = "//UE5/Release";
 		readonly StreamId _releaseStreamId = new StreamId(StringId.Sanitize(ReleaseStreamName));
 
-		readonly IGraph _graph;
+		IGraph _graph = default!;
 
 		TestDataController? _testDataController;
 
@@ -80,7 +80,8 @@ namespace HordeServer.Tests.Jobs.TestData
 			return streamConfig;
 		}
 
-		public TestDataTests()
+		[TestInitialize]
+		public async Task SetupAsync()
 		{
 			ProjectConfig projectConfig = new ProjectConfig { Id = new ProjectId("ue5"), Name = "UE5" };
 			projectConfig.Streams.Add(CreateStream(_mainStreamId, MainStreamName));
@@ -92,7 +93,7 @@ namespace HordeServer.Tests.Jobs.TestData
 			GlobalConfig globalConfig = new GlobalConfig();
 			globalConfig.Plugins.AddBuildConfig(buildConfig);
 
-			SetConfig(globalConfig);
+			await SetConfigAsync(globalConfig);
 
 			List<INode> nodes = new List<INode>();
 			nodes.Add(MockNode("Update Version Files", NodeAnnotations.Empty));
diff --git a/Engine/Source/Programs/Horde/Plugins/Build/HordeServer.Build.Tests/Notifications/NotificationServiceTests.cs b/Engine/Source/Programs/Horde/Plugins/Build/HordeServer.Build.Tests/Notifications/NotificationServiceTests.cs
--- a/Engine/Source/Programs/Horde/Plugins/Build/HordeServer.Build.Tests/Notifications/NotificationServiceTests.cs
+++ b/Engine/Source/Programs/Horde/Plugins/Build/HordeServer.Build.Tests/Notifications/NotificationServiceTests.cs
@@ -170,9 +170,9 @@ namespace HordeServer.Tests.Notifications
 			return node;
 		}
 
-		void UpdateGlobalConfig()
+		async Task UpdateGlobalConfigAsync()
 		{
-			UpdateConfig(globalConfig =>
+			await UpdateConfigAsync(globalConfig =>
 			{
 				ProjectConfig projectConfig = new() { Id = new ProjectId("ue4") };
 				projectConfig.Streams.Add(new StreamConfig { Id = _streamId });
@@ -361,7 +361,7 @@ namespace HordeServer.Tests.Notifications
 		[TestMethod]
 		public async Task NotifyJobCompleteAsync()
 		{
-			UpdateGlobalConfig();
+			await UpdateGlobalConfigAsync();
 
 			FakeNotificationSink fakeSink = ServiceProvider.GetRequiredService<FakeNotificationSink>();
 
@@ -401,7 +401,7 @@ namespace HordeServer.Tests.Notifications
 		[TestMethod]
 		public async Task NotifyJobAbortedAsync()
 		{
-			UpdateGlobalConfig();
+			await UpdateGlobalConfigAsync();
 
 			FakeNotificationSink fakeSink = ServiceProvider.GetRequiredService<FakeNotificationSink>();
 
diff --git a/Engine/Source/Programs/Horde/Plugins/Build/HordeServer.Build.Tests/Server/RpcServiceTest.cs b/Engine/Source/Programs/Horde/Plugins/Build/HordeServer.Build.Tests/Server/RpcServiceTest.cs
--- a/Engine/Source/Programs/Horde/Plugins/Build/HordeServer.Build.Tests/Server/RpcServiceTest.cs
+++ b/Engine/Source/Programs/Horde/Plugins/Build/HordeServer.Build.Tests/Server/RpcServiceTest.cs
@@ -313,7 +313,7 @@ namespace HordeServer.Tests.Server
 			ToolsConfig toolsConfig = new ();
 			toolsConfig.Tools.Add(new ToolConfig(_dedicatedAgentToolId) { Acl = new AclConfig() { Entries = [aclEntryConfig] } });
 			toolsConfig.Tools.Add(new ToolConfig(_workstationAgentToolId) { Public = true });
-			UpdateConfig(gc => gc.Plugins.AddToolsConfig(toolsConfig) );
+			await UpdateConfigAsync(gc => gc.Plugins.AddToolsConfig(toolsConfig) );
 			
 			IToolCollection collection = ServiceProvider.GetRequiredService<IToolCollection>();
 			await CreateAgentToolDeploymentAsync(collection, _dedicatedAgentToolId, "1", cancellationToken);
diff --git a/Engine/Source/Programs/Horde/Plugins/Build/HordeServer.Build/BuildConfig.cs b/Engine/Source/Programs/Horde/Plugins/Build/HordeServer.Build/BuildConfig.cs
--- a/Engine/Source/Programs/Horde/Plugins/Build/HordeServer.Build/BuildConfig.cs
+++ b/Engine/Source/Programs/Horde/Plugins/Build/HordeServer.Build/BuildConfig.cs
@@ -92,7 +92,7 @@ namespace HordeServer
 		private readonly Dictionary<StreamId, StreamConfig> _streamLookup = new Dictionary<StreamId, StreamConfig>();
 
 		/// <inheritdoc/>
-		public void PostLoad(PluginConfigOptions configOptions)
+		public Task PostLoadAsync(PluginConfigOptions configOptions)
 		{
 			Acl = configOptions.ParentAcl;
 			Streams = Projects.SelectMany(x => x.Streams).ToList();
@@ -116,6 +116,8 @@ namespace HordeServer
 			}
 
 			UpdateWorkspacesForPools(configOptions.Plugins);
+
+			return Task.CompletedTask;
 		}
 
 		void UpdateWorkspacesForPools(IEnumerable<IPluginConfig> plugins)
diff --git a/Engine/Source/Programs/Horde/Plugins/Compute/HordeServer.Compute/ComputeConfig.cs b/Engine/Source/Programs/Horde/Plugins/Compute/HordeServer.Compute/ComputeConfig.cs
--- a/Engine/Source/Programs/Horde/Plugins/Compute/HordeServer.Compute/ComputeConfig.cs
+++ b/Engine/Source/Programs/Horde/Plugins/Compute/HordeServer.Compute/ComputeConfig.cs
@@ -69,7 +69,7 @@ namespace HordeServer
 		private readonly Dictionary<PoolId, PoolConfig> _poolLookup = new Dictionary<PoolId, PoolConfig>();
 
 		/// <inheritdoc/>
-		public void PostLoad(PluginConfigOptions configOptions)
+		public Task PostLoadAsync(PluginConfigOptions configOptions)
 		{
 			Acl = configOptions.ParentAcl;
 			VersionEnum = configOptions.Version;
@@ -87,6 +87,8 @@ namespace HordeServer
 				_poolLookup.Add(pool.Id, pool);
 			}
 			ConfigObject.MergeDefaults<string, PoolConfig>(Pools.Select(x => (x.Id.ToString(), x.Base?.ToString(), x)));
+
+			return Task.CompletedTask;
 		}
 
 		/// <summary>
diff --git a/Engine/Source/Programs/Horde/Plugins/Ddc/HordeServer.Ddc.Tests/Ddc/Functional/ContentId/CompressedBlobTests.cs b/Engine/Source/Programs/Horde/Plugins/Ddc/HordeServer.Ddc.Tests/Ddc/Functional/ContentId/CompressedBlobTests.cs
--- a/Engine/Source/Programs/Horde/Plugins/Ddc/HordeServer.Ddc.Tests/Ddc/Functional/ContentId/CompressedBlobTests.cs
+++ b/Engine/Source/Programs/Horde/Plugins/Ddc/HordeServer.Ddc.Tests/Ddc/Functional/ContentId/CompressedBlobTests.cs
@@ -31,7 +31,8 @@ namespace HordeServer.Tests.Ddc.FunctionalTests.CompressedBlobs
 			public BlobId? Identifier { get; set; }
 		}
 
-		public CompressedBlobTests()
+		[TestInitialize]
+		public async Task AddConfigsAsync()
 		{
 			ServerSettings serverSettings = ServiceProvider.GetRequiredService<IOptions<ServerSettings>>().Value;
 
@@ -41,7 +42,7 @@ namespace HordeServer.Tests.Ddc.FunctionalTests.CompressedBlobs
 
 			GlobalConfig globalConfig = ServiceProvider.GetRequiredService<IOptions<GlobalConfig>>().Value;
 			globalConfig.Plugins.AddStorageConfig(storageConfig);
-			globalConfig.PostLoad(serverSettings, Array.Empty<ILoadedPlugin>(), Array.Empty<IDefaultAclModifier>());
+			await globalConfig.PostLoadAsync(serverSettings, Array.Empty<ILoadedPlugin>(), Array.Empty<IDefaultAclModifier>());
 
 			ConfigService configService = ServiceProvider.GetRequiredService<ConfigService>();
 			configService.OverrideConfig(globalConfig);
diff --git a/Engine/Source/Programs/Horde/Plugins/Ddc/HordeServer.Ddc.Tests/Ddc/Functional/Objects/ReferencesTests.cs b/Engine/Source/Programs/Horde/Plugins/Ddc/HordeServer.Ddc.Tests/Ddc/Functional/Objects/ReferencesTests.cs
--- a/Engine/Source/Programs/Horde/Plugins/Ddc/HordeServer.Ddc.Tests/Ddc/Functional/Objects/ReferencesTests.cs
+++ b/Engine/Source/Programs/Horde/Plugins/Ddc/HordeServer.Ddc.Tests/Ddc/Functional/Objects/ReferencesTests.cs
@@ -47,7 +47,12 @@ namespace HordeServer.Tests.Ddc.FunctionalTests.References
 		public ReferencesTests()
 		{
 			_httpClient = Client;
+			_serviceScope = ServiceProvider.CreateAsyncScope();
+		}
 
+		[TestInitialize]
+		public async Task AddConfigsAsync()
+		{
 			ServerSettings serverSettings = ServiceProvider.GetRequiredService<IOptions<ServerSettings>>().Value;
 
 			StorageConfig storageConfig = new StorageConfig();
@@ -56,12 +61,10 @@ namespace HordeServer.Tests.Ddc.FunctionalTests.References
 
 			GlobalConfig globalConfig = ServiceProvider.GetRequiredService<IOptions<GlobalConfig>>().Value;
 			globalConfig.Plugins.AddStorageConfig(storageConfig);
-			globalConfig.PostLoad(serverSettings, Array.Empty<ILoadedPlugin>(), Array.Empty<IDefaultAclModifier>());
+			await globalConfig.PostLoadAsync(serverSettings, Array.Empty<ILoadedPlugin>(), Array.Empty<IDefaultAclModifier>());
 
 			ConfigService configService = ServiceProvider.GetRequiredService<ConfigService>();
 			configService.OverrideConfig(globalConfig);
-
-			_serviceScope = ServiceProvider.CreateAsyncScope();
 		}
 
 		public override async ValueTask DisposeAsync()
diff --git a/Engine/Source/Programs/Horde/Plugins/Experimental/HordeServer.Experimental/ExperimentalConfig.cs b/Engine/Source/Programs/Horde/Plugins/Experimental/HordeServer.Experimental/ExperimentalConfig.cs
--- a/Engine/Source/Programs/Horde/Plugins/Experimental/HordeServer.Experimental/ExperimentalConfig.cs
+++ b/Engine/Source/Programs/Horde/Plugins/Experimental/HordeServer.Experimental/ExperimentalConfig.cs
@@ -20,7 +20,7 @@ namespace HordeServer
 		readonly Dictionary<NotificationConfigId, NotificationConfig> _notificationConfigLookup = new Dictionary<NotificationConfigId, NotificationConfig>();
 
 		/// <inheritdoc/>
-		public void PostLoad(PluginConfigOptions configOptions)
+		public Task PostLoadAsync(PluginConfigOptions configOptions)
 		{
 			_notificationConfigLookup.Clear();
 			foreach (NotificationConfig config in Notifications)
@@ -28,6 +28,8 @@ namespace HordeServer
 				config.PostLoad();
 				_notificationConfigLookup.Add(config.Id, config);
 			}
+
+			return Task.CompletedTask;
 		}
 
 		/// <summary>
diff --git a/Engine/Source/Programs/Horde/Plugins/Secrets/HordeServer.Secrets.Tests/Secrets/SecretProviderTests.cs b/Engine/Source/Programs/Horde/Plugins/Secrets/HordeServer.Secrets.Tests/Secrets/SecretProviderTests.cs
--- a/Engine/Source/Programs/Horde/Plugins/Secrets/HordeServer.Secrets.Tests/Secrets/SecretProviderTests.cs
+++ b/Engine/Source/Programs/Horde/Plugins/Secrets/HordeServer.Secrets.Tests/Secrets/SecretProviderTests.cs
@@ -45,7 +45,7 @@ namespace HordeServer.Tests.Secrets
 				return new HttpClient();
 			}
 		}
-		private static SecretCollection CreateSecretCollection()
+		private static async Task<SecretCollection> CreateSecretCollectionAsync()
 		{
 			SecretCollectionInternal secretCollectionInternal = new([
 				new InMemoryProvider(),
@@ -186,14 +186,14 @@ namespace HordeServer.Tests.Secrets
 			SecretCollection secretCollection = new(secretCollectionInternal, new OptionsSnapshotStub<SecretsConfig>(secretsConfig));
 			GlobalConfig config = new();
 			config.Plugins.AddSecretsConfig(secretsConfig);
-			config.PostLoad(new ServerSettings(), [], []);
+			await config.PostLoadAsync(new ServerSettings(), [], []);
 			return secretCollection;
 		}
 
 		[TestMethod]
 		public async Task GetSecretAsync()
 		{
-			SecretCollection secretCollection = CreateSecretCollection(); 
+			SecretCollection secretCollection = await CreateSecretCollectionAsync(); 
 			ISecret? secret = await secretCollection.GetAsync(new SecretId("Secret"), CancellationToken.None);
 			Assert.AreEqual(secret?.Data["username"], "j.doe");
 			Assert.AreEqual(secret?.Data["password"], "password123");
@@ -202,7 +202,7 @@ namespace HordeServer.Tests.Secrets
 		[TestMethod]
 		public async Task GetSecretFromProviderConfigAsync()
 		{
-			SecretCollection secretCollection = CreateSecretCollection(); 
+			SecretCollection secretCollection = await CreateSecretCollectionAsync(); 
 			ISecret? secret = await secretCollection.GetAsync(new SecretId("SecretUsingProviderConfig"), CancellationToken.None);
 			Assert.AreEqual(secret?.Data["username"], "user01");
 			Assert.AreEqual(secret?.Data["password"], "password01");
@@ -211,7 +211,7 @@ namespace HordeServer.Tests.Secrets
 		[TestMethod]
 		public async Task GetEmptySecretAsync()
 		{
-			SecretCollection secretCollection = CreateSecretCollection();
+			SecretCollection secretCollection = await CreateSecretCollectionAsync();
 			ISecret? secret = await secretCollection.GetAsync(new SecretId("EmptySecret"), CancellationToken.None);
 			Assert.AreEqual(secret?.Data.Count, 0);
 		}
@@ -220,7 +220,7 @@ namespace HordeServer.Tests.Secrets
 		[TestMethod]
 		public async Task GetSecretFromVaultAsync()
 		{
-			SecretCollection secretCollection = CreateSecretCollection(); 
+			SecretCollection secretCollection = await CreateSecretCollectionAsync(); 
 			ISecret? secret = await secretCollection.GetAsync(new SecretId("SecretUsingHcpVault"), CancellationToken.None);
 			Assert.AreEqual(secret?.Data["username"], "User01");
 			Assert.AreEqual(secret?.Data["password"], "Password123");
@@ -230,7 +230,7 @@ namespace HordeServer.Tests.Secrets
 		[TestMethod]
 		public async Task GetMultipleSecretFromVaultAsync()
 		{
-			SecretCollection secretCollection = CreateSecretCollection(); 
+			SecretCollection secretCollection = await CreateSecretCollectionAsync(); 
 			ISecret? secret = await secretCollection.GetAsync(new SecretId("MultipleSecretUsingHcpVault"), CancellationToken.None);
 			Assert.AreEqual(secret?.Data["one"], "1");
 			Assert.AreEqual(secret?.Data["two"], "2");
@@ -241,7 +241,7 @@ namespace HordeServer.Tests.Secrets
 		[TestMethod]
 		public async Task SecretNotFoundFromVaultAsync()
 		{
-			SecretCollection secretCollection = CreateSecretCollection(); 
+			SecretCollection secretCollection = await CreateSecretCollectionAsync(); 
 			InvalidOperationException ex = await Assert.ThrowsExceptionAsync<InvalidOperationException>(() =>
 				secretCollection.GetAsync(new SecretId("MissingSecretUsingHcpVault"), CancellationToken.None)
 			);
@@ -257,7 +257,7 @@ namespace HordeServer.Tests.Secrets
 			// 2. set the environment variable AWS_PROFILE to a profile that can call AssumeRole for the required ARN
 			// 3. restart IDE to pick up the environment change
 			// 4. update the SecretProviderConfig options for the SecretId used in this test 
-			SecretCollection secretCollection = CreateSecretCollection(); 
+			SecretCollection secretCollection = await CreateSecretCollectionAsync(); 
 			ISecret? secret = await secretCollection.GetAsync(new SecretId("SecretUsingHcpVaultAwsAuth"), CancellationToken.None);
 			Assert.AreEqual(secret?.Data["test"], "test");
 		}
@@ -283,7 +283,7 @@ namespace HordeServer.Tests.Secrets
 			SecretCollection secretCollection = new(secretCollectionInternal, new OptionsSnapshotStub<SecretsConfig>(secretsConfig));
 			GlobalConfig config = new();
 			config.Plugins.AddSecretsConfig(secretsConfig);
-			config.PostLoad(new ServerSettings(), [], []);
+			await config.PostLoadAsync(new ServerSettings(), [], []);
 			ISecret? secret = await secretCollection.GetAsync(new SecretId("basic-secret"), CancellationToken.None);
 			Assert.AreEqual(secret?.Data["username"], "j.doe");
 			Assert.AreEqual(secret?.Data["password"], "password123");
diff --git a/Engine/Source/Programs/Horde/Plugins/Secrets/HordeServer.Secrets/SecretsConfig.cs b/Engine/Source/Programs/Horde/Plugins/Secrets/HordeServer.Secrets/SecretsConfig.cs
--- a/Engine/Source/Programs/Horde/Plugins/Secrets/HordeServer.Secrets/SecretsConfig.cs
+++ b/Engine/Source/Programs/Horde/Plugins/Secrets/HordeServer.Secrets/SecretsConfig.cs
@@ -31,7 +31,7 @@ namespace HordeServer
 		private readonly Dictionary<SecretId, Dictionary<string, SecretProviderConfig>> _providerConfigLookup = new();
 
 		/// <inheritdoc/>
-		public void PostLoad(PluginConfigOptions configOptions)
+		public Task PostLoadAsync(PluginConfigOptions configOptions)
 		{
 			_secretLookup.Clear();
 			foreach (SecretConfig secret in Secrets)
@@ -63,6 +63,8 @@ namespace HordeServer
 					}
 				}
 			}
+
+			return Task.CompletedTask;
 		}
 
 		/// <summary>
diff --git a/Engine/Source/Programs/Horde/Plugins/Storage/HordeServer.Storage.Tests/Storage/BlobStoreTests.cs b/Engine/Source/Programs/Horde/Plugins/Storage/HordeServer.Storage.Tests/Storage/BlobStoreTests.cs
--- a/Engine/Source/Programs/Horde/Plugins/Storage/HordeServer.Storage.Tests/Storage/BlobStoreTests.cs
+++ b/Engine/Source/Programs/Horde/Plugins/Storage/HordeServer.Storage.Tests/Storage/BlobStoreTests.cs
@@ -21,7 +21,7 @@ namespace HordeServer.Tests.Storage
 			AddPlugin<StoragePlugin>();
 		}
 
-		IStorageNamespace GetStorageNamespace()
+		async Task<IStorageNamespace> GetStorageNamespaceAsync()
 		{
 			StorageConfig storageConfig = new StorageConfig();
 			storageConfig.Backends.Add(new BackendConfig { Id = new BackendId("default-backend"), Type = StorageBackendType.Memory });
@@ -29,7 +29,7 @@ namespace HordeServer.Tests.Storage
 
 			GlobalConfig globalConfig = new GlobalConfig();
 			globalConfig.Plugins.Add(new PluginName("storage"), storageConfig);
-			SetConfig(globalConfig);
+			await SetConfigAsync(globalConfig);
 
 			return StorageService.GetNamespace(new NamespaceId("default"));
 		}
@@ -94,7 +94,7 @@ namespace HordeServer.Tests.Storage
 		[TestMethod]
 		public async Task LeafTestAsync()
 		{
-			IStorageNamespace store = GetStorageNamespace();
+			IStorageNamespace store = await GetStorageNamespaceAsync();
 
 			byte[] input = CreateTestData(256, 0);
 
@@ -109,7 +109,7 @@ namespace HordeServer.Tests.Storage
 		[TestMethod]
 		public async Task ReferenceTestAsync()
 		{
-			IStorageNamespace store = GetStorageNamespace();
+			IStorageNamespace store = await GetStorageNamespaceAsync();
 
 			byte[] input1 = CreateTestData(256, 1);
 			HashedBlobRefValue locator1 = await WriteBlobAsync(store, new Blob(input1, Array.Empty<HashedBlobRefValue>()));
@@ -141,7 +141,7 @@ namespace HordeServer.Tests.Storage
 		[TestMethod]
 		public async Task RefExpiryTestAsync()
 		{
-			IStorageNamespace store = GetStorageNamespace();
+			IStorageNamespace store = await GetStorageNamespaceAsync();
 
 			Blob blob1 = new Blob(new byte[] { 1, 2, 3 }, Array.Empty<HashedBlobRefValue>());
 			HashedBlobRefValue target = await WriteBlobAsync(store, blob1);
@@ -186,7 +186,7 @@ namespace HordeServer.Tests.Storage
 		[TestMethod]
 		public async Task AliasTestAsync()
 		{
-			IStorageNamespace store = GetStorageNamespace();
+			IStorageNamespace store = await GetStorageNamespaceAsync();
 
 			BlobType type = new BlobType(Guid.Empty, 0);
 
@@ -215,7 +215,7 @@ namespace HordeServer.Tests.Storage
 		[TestMethod]
 		public async Task MultipleAliasTestsAsync()
 		{
-			IStorageNamespace store = GetStorageNamespace();
+			IStorageNamespace store = await GetStorageNamespaceAsync();
 
 			BlobType type = new BlobType(Guid.Empty, 0);
 
diff --git a/Engine/Source/Programs/Horde/Plugins/Storage/HordeServer.Storage.Tests/Storage/GcServiceTests.cs b/Engine/Source/Programs/Horde/Plugins/Storage/HordeServer.Storage.Tests/Storage/GcServiceTests.cs
--- a/Engine/Source/Programs/Horde/Plugins/Storage/HordeServer.Storage.Tests/Storage/GcServiceTests.cs
+++ b/Engine/Source/Programs/Horde/Plugins/Storage/HordeServer.Storage.Tests/Storage/GcServiceTests.cs
@@ -30,7 +30,7 @@ namespace HordeServer.Tests.Storage
 
 			GlobalConfig globalConfig = new GlobalConfig();
 			globalConfig.Plugins.Add(new PluginName("storage"), storageConfig);
-			SetConfig(globalConfig);
+			await SetConfigAsync(globalConfig);
 
 			IStorageNamespace store = StorageService.GetNamespace(new NamespaceId("default"));
 
diff --git a/Engine/Source/Programs/Horde/Plugins/Storage/HordeServer.Storage/StorageConfig.cs b/Engine/Source/Programs/Horde/Plugins/Storage/HordeServer.Storage/StorageConfig.cs
--- a/Engine/Source/Programs/Horde/Plugins/Storage/HordeServer.Storage/StorageConfig.cs
+++ b/Engine/Source/Programs/Horde/Plugins/Storage/HordeServer.Storage/StorageConfig.cs
@@ -74,7 +74,7 @@ namespace HordeServer.Storage
 		/// <summary>
 		/// Called after the config has been read from disk
 		/// </summary>
-		public void PostLoad(PluginConfigOptions configOptions)
+		public Task PostLoadAsync(PluginConfigOptions configOptions)
 		{
 			Acl = configOptions.ParentAcl;
 
@@ -123,6 +123,8 @@ namespace HordeServer.Storage
 			// Create a unique id to allow detect changes
 			byte[] data = JsonSerializer.SerializeToUtf8Bytes<StorageConfig>(this, JsonUtils.DefaultSerializerOptions);
 			Revision = IoHash.Compute(data).ToString();
+
+			return Task.CompletedTask;
 		}
 
 		static BackendConfig MergeBackendConfigs(BackendId backendId, Dictionary<BackendId, BackendConfig> baseIdToConfig, Dictionary<BackendId, BackendConfig> mergedIdToConfig)
diff --git a/Engine/Source/Programs/Horde/Plugins/Symbols/HordeServer.Symbols/SymbolsConfig.cs b/Engine/Source/Programs/Horde/Plugins/Symbols/HordeServer.Symbols/SymbolsConfig.cs
--- a/Engine/Source/Programs/Horde/Plugins/Symbols/HordeServer.Symbols/SymbolsConfig.cs
+++ b/Engine/Source/Programs/Horde/Plugins/Symbols/HordeServer.Symbols/SymbolsConfig.cs
@@ -22,7 +22,7 @@ namespace HordeServer.Symbols
 		readonly Dictionary<SymbolStoreId, SymbolStoreConfig> _storeLookup = new Dictionary<SymbolStoreId, SymbolStoreConfig>();
 
 		/// <inheritdoc/>
-		public void PostLoad(PluginConfigOptions configOptions)
+		public Task PostLoadAsync(PluginConfigOptions configOptions)
 		{
 			foreach (SymbolStoreConfig store in Stores)
 			{
@@ -34,6 +34,8 @@ namespace HordeServer.Symbols
 			{
 				_storeLookup.Add(store.Id, store);
 			}
+
+			return Task.CompletedTask;
 		}
 
 		/// <summary>
diff --git a/Engine/Source/Programs/Horde/Plugins/Tools/HordeServer.Tools.Tests/Tools/ToolTests.cs b/Engine/Source/Programs/Horde/Plugins/Tools/HordeServer.Tools.Tests/Tools/ToolTests.cs
--- a/Engine/Source/Programs/Horde/Plugins/Tools/HordeServer.Tools.Tests/Tools/ToolTests.cs
+++ b/Engine/Source/Programs/Horde/Plugins/Tools/HordeServer.Tools.Tests/Tools/ToolTests.cs
@@ -16,7 +16,8 @@ namespace HordeServer.Tests.Tools
 	{
 		readonly ToolId _toolId = new ToolId("ugs");
 
-		public ToolTests()
+		[TestInitialize]
+		public async Task AddPluginsAsync()
 		{
 			AddPlugin<StoragePlugin>();
 			AddPlugin<ToolsPlugin>();
@@ -33,7 +34,7 @@ namespace HordeServer.Tests.Tools
 			GlobalConfig globalConfig = new GlobalConfig();
 			globalConfig.Plugins.AddStorageConfig(storageConfig);
 			globalConfig.Plugins.AddToolsConfig(toolsConfig);
-			SetConfig(globalConfig);
+			await SetConfigAsync(globalConfig);
 		}
 
 		[TestMethod]
diff --git a/Engine/Source/Programs/Horde/Plugins/Tools/HordeServer.Tools.Tests/Tools/ToolsControllerTests.cs b/Engine/Source/Programs/Horde/Plugins/Tools/HordeServer.Tools.Tests/Tools/ToolsControllerTests.cs
--- a/Engine/Source/Programs/Horde/Plugins/Tools/HordeServer.Tools.Tests/Tools/ToolsControllerTests.cs
+++ b/Engine/Source/Programs/Horde/Plugins/Tools/HordeServer.Tools.Tests/Tools/ToolsControllerTests.cs
@@ -46,7 +46,7 @@ public class ToolsControllerTests
 		globalConfig.Plugins.AddToolsConfig(toolsConfig);
 
 		ServerSettings serverSettings = new() { AuthMethod = AuthMethod.Horde };
-		globalConfig.PostLoad(serverSettings, pluginCollection.LoadedPlugins, Array.Empty<IDefaultAclModifier>());
+		await globalConfig.PostLoadAsync(serverSettings, pluginCollection.LoadedPlugins, Array.Empty<IDefaultAclModifier>());
 
 		Dictionary<string, string> settings = new() { { "Horde:AuthMethod", AuthMethod.Horde.ToString() } };
 		await using FakeHordeWebApp app = new(settings: settings);
diff --git a/Engine/Source/Programs/Horde/Plugins/Tools/HordeServer.Tools/ToolsConfig.cs b/Engine/Source/Programs/Horde/Plugins/Tools/HordeServer.Tools/ToolsConfig.cs
--- a/Engine/Source/Programs/Horde/Plugins/Tools/HordeServer.Tools/ToolsConfig.cs
+++ b/Engine/Source/Programs/Horde/Plugins/Tools/HordeServer.Tools/ToolsConfig.cs
@@ -21,7 +21,7 @@ namespace HordeServer.Tools
 		private readonly Dictionary<ToolId, ToolConfig> _toolLookup = new Dictionary<ToolId, ToolConfig>();
 
 		/// <inheritdoc/>
-		public void PostLoad(PluginConfigOptions configOptions)
+		public Task PostLoadAsync(PluginConfigOptions configOptions)
 		{
 			_toolLookup.Clear();
 			foreach (ToolConfig tool in Tools)
@@ -29,6 +29,8 @@ namespace HordeServer.Tools
 				_toolLookup.Add(tool.Id, tool);
 				tool.PostLoad(configOptions.ParentAcl);
 			}
+
+			return Task.CompletedTask;
 		}
 
 		/// <summary>
