Horde: Fix compiler warnings

diff --git a/Engine/Source/Programs/Horde/Plugins/Experimental/HordeServer.Experimental.Tests/Notifications/ExperimentalNotificationTests.cs b/Engine/Source/Programs/Horde/Plugins/Experimental/HordeServer.Experimental.Tests/Notifications/ExperimentalNotificationTests.cs
--- a/Engine/Source/Programs/Horde/Plugins/Experimental/HordeServer.Experimental.Tests/Notifications/ExperimentalNotificationTests.cs
+++ b/Engine/Source/Programs/Horde/Plugins/Experimental/HordeServer.Experimental.Tests/Notifications/ExperimentalNotificationTests.cs
@@ -205,7 +205,7 @@ namespace HordeServer.Experimental.Tests.Notifications
 			GlobalConfig globalConfig = new GlobalConfig();
 			globalConfig.Plugins.AddBuildConfig(buildConfig);
 
-			SetConfig(globalConfig);
+			SetConfigAsync(globalConfig).Wait();
 		}
 
 		protected override void ConfigureServices(IServiceCollection services)
@@ -233,10 +233,10 @@ namespace HordeServer.Experimental.Tests.Notifications
 
 		void UpdateGlobalConfig(ExperimentalConfig experimentalConfig)
 		{
-			UpdateConfig(globalConfig =>
+			UpdateConfigAsync(globalConfig =>
 			{
 				globalConfig.Plugins.AddExperimentalConfig(experimentalConfig);
-			});
+			}).Wait();
 		}
 
 		ExperimentalConfig CreateExperimentalNotificationConfig()
@@ -433,7 +433,7 @@ namespace HordeServer.Experimental.Tests.Notifications
 			return job;
 		}
 
-		static void ValidateJobStepInformation(IJobStepNotificationStateRef jobStepNotificationState, JobId jobId, TemplateId templateId, string recipient, string group, string platform, string badge, JobId? parentJobId = null, TemplateId? parentJobTemplateId = null)
+		static void ValidateJobStepInformationAsync(IJobStepNotificationStateRef jobStepNotificationState, JobId jobId, TemplateId templateId, string recipient, string group, string platform, string badge, JobId? parentJobId = null, TemplateId? parentJobTemplateId = null)
 		{
 			Assert.AreEqual(jobId, jobStepNotificationState.JobId);
 			Assert.AreEqual(templateId, jobStepNotificationState.TemplateId);
@@ -445,7 +445,7 @@ namespace HordeServer.Experimental.Tests.Notifications
 			Assert.AreEqual(parentJobTemplateId, jobStepNotificationState.ParentJobTemplateId);
 		}
 
-		static async Task ValidateJobStepInformation(JobNotificationCollection jobNotificationCollection, JobId jobId, TemplateId templateId, string recipient, JobStepId jobStepId, string group, string platform, string badge, JobId? parentJobId = null, TemplateId? parentJobTemplateId = null)
+		static async Task ValidateJobStepInformationAsync(JobNotificationCollection jobNotificationCollection, JobId jobId, TemplateId templateId, string recipient, JobStepId jobStepId, string group, string platform, string badge, JobId? parentJobId = null, TemplateId? parentJobTemplateId = null)
 		{
 			IJobStepNotificationStateQueryBuilder queryBuilder = JobNotificationCollection.CreateJobStepNotificationStateQueryBuilder();
 			queryBuilder.AddJobFilter(jobId);
@@ -456,7 +456,7 @@ namespace HordeServer.Experimental.Tests.Notifications
 			Assert.IsNotNull(jobStepNotificationStates);
 			Assert.AreEqual(1, jobStepNotificationStates.Count);
 
-			ValidateJobStepInformation(jobStepNotificationStates[0], jobId, templateId, recipient, group, platform, badge, parentJobId, parentJobTemplateId);
+			ValidateJobStepInformationAsync(jobStepNotificationStates[0], jobId, templateId, recipient, group, platform, badge, parentJobId, parentJobTemplateId);
 		}
 
 		static async Task WaitForNotificationAsync(FakeExperimentalNotificationSink sink, IJob job, IJobStep? jobStep, int timeoutMs = 30000)
@@ -602,22 +602,22 @@ namespace HordeServer.Experimental.Tests.Notifications
 			Assert.AreEqual(6, jobStepNotificationStates.Count);
 
 			// Job Step notification for 'Client Boot Test Platform1'
-			ValidateJobStepInformation(jobStepNotificationStates[0], job.Id, job.TemplateId, ChannelId, ClientBootTestGroup, "Platform1", PendingBadge);
+			ValidateJobStepInformationAsync(jobStepNotificationStates[0], job.Id, job.TemplateId, ChannelId, ClientBootTestGroup, "Platform1", PendingBadge);
 
 			// Job Step notification for 'Client Project Test Platform1'
-			ValidateJobStepInformation(jobStepNotificationStates[1], job.Id, job.TemplateId, ChannelId, ClientProjectTestGroup, "Platform1", PendingBadge);
+			ValidateJobStepInformationAsync(jobStepNotificationStates[1], job.Id, job.TemplateId, ChannelId, ClientProjectTestGroup, "Platform1", PendingBadge);
 
 			// Job Step notification for 'Editor Boot Test Platform1'
-			ValidateJobStepInformation(jobStepNotificationStates[2], job.Id, job.TemplateId, ChannelId, EditorBootTestGroup, "Platform1", PendingBadge);
+			ValidateJobStepInformationAsync(jobStepNotificationStates[2], job.Id, job.TemplateId, ChannelId, EditorBootTestGroup, "Platform1", PendingBadge);
 
 			// Job Step notification for 'Editor Project Test Platform2'
-			ValidateJobStepInformation(jobStepNotificationStates[3], job.Id, job.TemplateId, ChannelId, EditorBootTestGroup, "Platform2", PendingBadge);
+			ValidateJobStepInformationAsync(jobStepNotificationStates[3], job.Id, job.TemplateId, ChannelId, EditorBootTestGroup, "Platform2", PendingBadge);
 
 			// Job Step notification for 'Editor Boot Test Platform1'
-			ValidateJobStepInformation(jobStepNotificationStates[4], job.Id, job.TemplateId, ChannelId, EditorProjectTestGroup, "Platform1", PendingBadge);
+			ValidateJobStepInformationAsync(jobStepNotificationStates[4], job.Id, job.TemplateId, ChannelId, EditorProjectTestGroup, "Platform1", PendingBadge);
 
 			// Job Step notification for 'Editor Project Test Platform2'
-			ValidateJobStepInformation(jobStepNotificationStates[5], job.Id, job.TemplateId, ChannelId, EditorProjectTestGroup, "Platform2", PendingBadge);
+			ValidateJobStepInformationAsync(jobStepNotificationStates[5], job.Id, job.TemplateId, ChannelId, EditorProjectTestGroup, "Platform2", PendingBadge);
 
 			job = await RunBatchAsync(_mainStreamConfig, job, 1);
 			job = await RunStepAsync(_mainStreamConfig, job, 1, 0, JobStepOutcome.Success); // Update Version Files
@@ -629,22 +629,22 @@ namespace HordeServer.Experimental.Tests.Notifications
 			await WaitForNotificationAsync(fakeSink, job, job.Batches[1].Steps[4]);
 
 			// Validate that the 'Editor BootTest Platform1' step is recorded as passing
-			await ValidateJobStepInformation(jobNotificationCollection, job.Id, job.TemplateId,ChannelId, job.Batches[1].Steps[4].Id, EditorBootTestGroup, "Platform1", PassingBadge);
+			await ValidateJobStepInformationAsync(jobNotificationCollection, job.Id, job.TemplateId,ChannelId, job.Batches[1].Steps[4].Id, EditorBootTestGroup, "Platform1", PassingBadge);
 
 			job = await RunStepAsync(_mainStreamConfig, job, 1, 5, JobStepOutcome.Warnings); // Editor Project Tests Platform1
 
 			await WaitForNotificationAsync(fakeSink, job, job.Batches[1].Steps[5]);
 
 			// Validate that the 'Editor Project Test Platform1' step is recorded as having warnings
-			await ValidateJobStepInformation(jobNotificationCollection, job.Id, job.TemplateId, ChannelId, job.Batches[1].Steps[5].Id, EditorProjectTestGroup, "Platform1", WarningBadge);
+			await ValidateJobStepInformationAsync(jobNotificationCollection, job.Id, job.TemplateId, ChannelId, job.Batches[1].Steps[5].Id, EditorProjectTestGroup, "Platform1", WarningBadge);
 
 			job = await RunStepAsync(_mainStreamConfig, job, 1, 6, JobStepOutcome.Failure); // Client BootTest Platform1
 
 			await WaitForNotificationAsync(fakeSink, job, job.Batches[1].Steps[6]);
 
 			// Validate that the 'Client BootTest Platform1' step is recorded as having failures and 'Client Project Tests Platform1' is recorded as being skipped due to having a dependency on the BootTest
-			await ValidateJobStepInformation(jobNotificationCollection, job.Id, job.TemplateId, ChannelId, job.Batches[1].Steps[6].Id, ClientBootTestGroup, "Platform1", FailureBadge);
-			await ValidateJobStepInformation(jobNotificationCollection, job.Id, job.TemplateId, ChannelId, job.Batches[1].Steps[7].Id, ClientProjectTestGroup, "Platform1", SkippedBadge);
+			await ValidateJobStepInformationAsync(jobNotificationCollection, job.Id, job.TemplateId, ChannelId, job.Batches[1].Steps[6].Id, ClientBootTestGroup, "Platform1", FailureBadge);
+			await ValidateJobStepInformationAsync(jobNotificationCollection, job.Id, job.TemplateId, ChannelId, job.Batches[1].Steps[7].Id, ClientProjectTestGroup, "Platform1", SkippedBadge);
 
 			job = await RunStepAsync(_mainStreamConfig, job, 1, 8, JobStepOutcome.Failure); // Editor BootTest Platform2
 
@@ -703,16 +703,16 @@ namespace HordeServer.Experimental.Tests.Notifications
 			Assert.AreEqual(4, jobStepNotificationStates.Count);
 
 			// Job Step notification for 'Client Boot Test Platform1'
-			ValidateJobStepInformation(jobStepNotificationStates[0], job.Id, job.TemplateId, ChannelId, ClientBootTestGroup, "Platform1", PendingBadge);
+			ValidateJobStepInformationAsync(jobStepNotificationStates[0], job.Id, job.TemplateId, ChannelId, ClientBootTestGroup, "Platform1", PendingBadge);
 
 			// Job Step notification for 'Client Boot Test Platform3'
-			ValidateJobStepInformation(jobStepNotificationStates[1], job.Id, job.TemplateId, ChannelId, ClientBootTestGroup, "Platform3", PendingBadge);
+			ValidateJobStepInformationAsync(jobStepNotificationStates[1], job.Id, job.TemplateId, ChannelId, ClientBootTestGroup, "Platform3", PendingBadge);
 
 			// Job Step notification for 'Editor Boot Test Platform1'
-			ValidateJobStepInformation(jobStepNotificationStates[2], job.Id, job.TemplateId, ChannelId, EditorBootTestGroup, "Platform1", PendingBadge);
+			ValidateJobStepInformationAsync(jobStepNotificationStates[2], job.Id, job.TemplateId, ChannelId, EditorBootTestGroup, "Platform1", PendingBadge);
 
 			// Job Step notification for 'Editor Boot Test Platform3'
-			ValidateJobStepInformation(jobStepNotificationStates[3], job.Id, job.TemplateId, ChannelId, EditorBootTestGroup, "Platform3", PendingBadge);
+			ValidateJobStepInformationAsync(jobStepNotificationStates[3], job.Id, job.TemplateId, ChannelId, EditorBootTestGroup, "Platform3", PendingBadge);
 
 			// #1
 			// Scenario: Cancellation of a step will report the step as failing
@@ -733,16 +733,16 @@ namespace HordeServer.Experimental.Tests.Notifications
 				Assert.AreEqual(4, jobStepNotificationStates.Count);
 
 				// Job Step notification for the 'Client Boot Test Platform1' group
-				ValidateJobStepInformation(jobStepNotificationStates[0], job.Id, job.TemplateId, ChannelId, ClientBootTestGroup, "Platform1", PendingBadge);
+				ValidateJobStepInformationAsync(jobStepNotificationStates[0], job.Id, job.TemplateId, ChannelId, ClientBootTestGroup, "Platform1", PendingBadge);
 
 				// Job Step notification for the 'Client Boot Test Platform3' group
-				ValidateJobStepInformation(jobStepNotificationStates[1], job.Id, job.TemplateId, ChannelId, ClientBootTestGroup, "Platform3", PendingBadge);
+				ValidateJobStepInformationAsync(jobStepNotificationStates[1], job.Id, job.TemplateId, ChannelId, ClientBootTestGroup, "Platform3", PendingBadge);
 
 				// Job Step notification for the 'Editor Boot Test Platform1' group
-				ValidateJobStepInformation(jobStepNotificationStates[2], job.Id, job.TemplateId, ChannelId, EditorBootTestGroup, "Platform1", FailureBadge);
+				ValidateJobStepInformationAsync(jobStepNotificationStates[2], job.Id, job.TemplateId, ChannelId, EditorBootTestGroup, "Platform1", FailureBadge);
 
 				// Job Step notification for the 'Editor Boot Test Platform3' group
-				ValidateJobStepInformation(jobStepNotificationStates[3], job.Id, job.TemplateId, ChannelId, EditorBootTestGroup, "Platform3", PendingBadge);
+				ValidateJobStepInformationAsync(jobStepNotificationStates[3], job.Id, job.TemplateId, ChannelId, EditorBootTestGroup, "Platform3", PendingBadge);
 			}
 
 			// #2
@@ -764,16 +764,16 @@ namespace HordeServer.Experimental.Tests.Notifications
 				Assert.AreEqual(4, jobStepNotificationStates.Count);
 
 				// Job Step notification for the 'Client Boot Test Platform1' group
-				ValidateJobStepInformation(jobStepNotificationStates[0], job.Id, job.TemplateId, ChannelId, ClientBootTestGroup, "Platform1", SkippedBadge);
+				ValidateJobStepInformationAsync(jobStepNotificationStates[0], job.Id, job.TemplateId, ChannelId, ClientBootTestGroup, "Platform1", SkippedBadge);
 
 				// Job Step notification for the 'Client Boot Test Platform3' group
-				ValidateJobStepInformation(jobStepNotificationStates[1], job.Id, job.TemplateId, ChannelId, ClientBootTestGroup, "Platform3", SkippedBadge);
+				ValidateJobStepInformationAsync(jobStepNotificationStates[1], job.Id, job.TemplateId, ChannelId, ClientBootTestGroup, "Platform3", SkippedBadge);
 
 				// Job Step notification for the 'Editor Boot Test Platform1' group
-				ValidateJobStepInformation(jobStepNotificationStates[2], job.Id, job.TemplateId, ChannelId, EditorBootTestGroup, "Platform1", FailureBadge);
+				ValidateJobStepInformationAsync(jobStepNotificationStates[2], job.Id, job.TemplateId, ChannelId, EditorBootTestGroup, "Platform1", FailureBadge);
 
 				// Job Step notification for the 'Editor Boot Test Platform3' group
-				ValidateJobStepInformation(jobStepNotificationStates[3], job.Id, job.TemplateId, ChannelId, EditorBootTestGroup, "Platform3", PendingBadge);
+				ValidateJobStepInformationAsync(jobStepNotificationStates[3], job.Id, job.TemplateId, ChannelId, EditorBootTestGroup, "Platform3", PendingBadge);
 			}
 
 			// #3
@@ -869,13 +869,13 @@ namespace HordeServer.Experimental.Tests.Notifications
 			Assert.AreEqual(3, jobStepNotificationStates.Count);
 
 			// Job Step notification for 'Spawn Platform1 Test'
-			ValidateJobStepInformation(jobStepNotificationStates[0], parentJob.Id, parentJob.TemplateId, ChannelId, SpawnJobTestGroup, "Platform1", PendingBadge);
+			ValidateJobStepInformationAsync(jobStepNotificationStates[0], parentJob.Id, parentJob.TemplateId, ChannelId, SpawnJobTestGroup, "Platform1", PendingBadge);
 
 			// Job Step notification for 'Spawn Platform2 Test'
-			ValidateJobStepInformation(jobStepNotificationStates[1], parentJob.Id, parentJob.TemplateId, ChannelId, SpawnJobTestGroup, "Platform2", PendingBadge);
+			ValidateJobStepInformationAsync(jobStepNotificationStates[1], parentJob.Id, parentJob.TemplateId, ChannelId, SpawnJobTestGroup, "Platform2", PendingBadge);
 
 			// Job Step notification for 'Spawn Platform3 Test'
-			ValidateJobStepInformation(jobStepNotificationStates[2], parentJob.Id, parentJob.TemplateId, ChannelId, SpawnJobTestGroup, "Platform3", PendingBadge);
+			ValidateJobStepInformationAsync(jobStepNotificationStates[2], parentJob.Id, parentJob.TemplateId, ChannelId, SpawnJobTestGroup, "Platform3", PendingBadge);
 
 			// #1
 			// Scenario: Spawned job has a failing Setup Build step
@@ -892,7 +892,7 @@ namespace HordeServer.Experimental.Tests.Notifications
 
 				await WaitForNotificationAsync(fakeSink, parentJob, parentJob.Batches[1].Steps[0]);
 
-				await ValidateJobStepInformation(jobNotificationCollection, parentJob.Id, parentJob.TemplateId, ChannelId, parentJob.Batches[1].Steps[0].Id, SpawnJobTestGroup, "Platform1", PassingBadge);
+				await ValidateJobStepInformationAsync(jobNotificationCollection, parentJob.Id, parentJob.TemplateId, ChannelId, parentJob.Batches[1].Steps[0].Id, SpawnJobTestGroup, "Platform1", PassingBadge);
 
 				// Execute the Setup Build of the newly spawned job
 				spawnedJob = await RunBatchAsync(_mainStreamConfig, spawnedJob, 0);
@@ -931,13 +931,13 @@ namespace HordeServer.Experimental.Tests.Notifications
 
 				// Job Step notification for the 'Spawn Platform1 Test' will now being labeled as passing since the step spawned the new job
 				// Even though the actual job had issues which will be reported in the Slack thread
-				ValidateJobStepInformation(jobStepNotificationStates[0], parentJob.Id, parentJob.TemplateId, ChannelId, SpawnJobTestGroup, "Platform1", PassingBadge);
+				ValidateJobStepInformationAsync(jobStepNotificationStates[0], parentJob.Id, parentJob.TemplateId, ChannelId, SpawnJobTestGroup, "Platform1", PassingBadge);
 
 				// Job Step notification for the Spawn Platform2 Test group
-				ValidateJobStepInformation(jobStepNotificationStates[1], parentJob.Id, parentJob.TemplateId, ChannelId, SpawnJobTestGroup, "Platform2", PendingBadge);
+				ValidateJobStepInformationAsync(jobStepNotificationStates[1], parentJob.Id, parentJob.TemplateId, ChannelId, SpawnJobTestGroup, "Platform2", PendingBadge);
 
 				// Job Step notification for the Spawn Platform3 Test group
-				ValidateJobStepInformation(jobStepNotificationStates[2], parentJob.Id, parentJob.TemplateId, ChannelId, SpawnJobTestGroup, "Platform3", PendingBadge);
+				ValidateJobStepInformationAsync(jobStepNotificationStates[2], parentJob.Id, parentJob.TemplateId, ChannelId, SpawnJobTestGroup, "Platform3", PendingBadge);
 			}
 
 			// #2
@@ -953,7 +953,7 @@ namespace HordeServer.Experimental.Tests.Notifications
 
 				await WaitForNotificationAsync(fakeSink, parentJob, parentJob.Batches[1].Steps[1]);
 
-				await ValidateJobStepInformation(jobNotificationCollection, parentJob.Id, parentJob.TemplateId, ChannelId, parentJob.Batches[1].Steps[1].Id, SpawnJobTestGroup, "Platform2", PassingBadge);
+				await ValidateJobStepInformationAsync(jobNotificationCollection, parentJob.Id, parentJob.TemplateId, ChannelId, parentJob.Batches[1].Steps[1].Id, SpawnJobTestGroup, "Platform2", PassingBadge);
 
 				// Execute the Setup Build of the newly spawned job
 				spawnedJob = await RunBatchAsync(_mainStreamConfig, spawnedJob, 0);
@@ -978,25 +978,25 @@ namespace HordeServer.Experimental.Tests.Notifications
 				Assert.AreEqual(7, consolidatedJobStepNotificationStates.Count);
 
 				// Job Step notification for 'Spawn Platform1 Test'
-				ValidateJobStepInformation(consolidatedJobStepNotificationStates[0], parentJob.Id, parentJob.TemplateId, ChannelId, SpawnJobTestGroup, "Platform1", PassingBadge);
+				ValidateJobStepInformationAsync(consolidatedJobStepNotificationStates[0], parentJob.Id, parentJob.TemplateId, ChannelId, SpawnJobTestGroup, "Platform1", PassingBadge);
 
 				// Job Step notification for 'Spawn Platform2 Test' group'
-				ValidateJobStepInformation(consolidatedJobStepNotificationStates[1], parentJob.Id, parentJob.TemplateId, ChannelId, SpawnJobTestGroup, "Platform2", PassingBadge);
+				ValidateJobStepInformationAsync(consolidatedJobStepNotificationStates[1], parentJob.Id, parentJob.TemplateId, ChannelId, SpawnJobTestGroup, "Platform2", PassingBadge);
 
 				// Job Step notification for 'Spawn Platform3 Test'
-				ValidateJobStepInformation(consolidatedJobStepNotificationStates[2], parentJob.Id, parentJob.TemplateId, ChannelId, SpawnJobTestGroup, "Platform3", PendingBadge);
+				ValidateJobStepInformationAsync(consolidatedJobStepNotificationStates[2], parentJob.Id, parentJob.TemplateId, ChannelId, SpawnJobTestGroup, "Platform3", PendingBadge);
 
 				// Job Step notification for 'Client BootTest Platform2'
-				ValidateJobStepInformation(consolidatedJobStepNotificationStates[3], spawnedJob.Id, spawnedJob.TemplateId, ChannelId, ClientBootTestGroup, "Platform2", PendingBadge, parentJob.Id, parentJob.TemplateId);
+				ValidateJobStepInformationAsync(consolidatedJobStepNotificationStates[3], spawnedJob.Id, spawnedJob.TemplateId, ChannelId, ClientBootTestGroup, "Platform2", PendingBadge, parentJob.Id, parentJob.TemplateId);
 
 				// Job Step notification for 'Client Project Tests Platform2'
-				ValidateJobStepInformation(consolidatedJobStepNotificationStates[4], spawnedJob.Id, spawnedJob.TemplateId, ChannelId, ClientProjectTestGroup, "Platform2", PendingBadge, parentJob.Id, parentJob.TemplateId);
+				ValidateJobStepInformationAsync(consolidatedJobStepNotificationStates[4], spawnedJob.Id, spawnedJob.TemplateId, ChannelId, ClientProjectTestGroup, "Platform2", PendingBadge, parentJob.Id, parentJob.TemplateId);
 
 				// Job Step notification for 'Editor BootTest Platform2'
-				ValidateJobStepInformation(consolidatedJobStepNotificationStates[5], spawnedJob.Id, spawnedJob.TemplateId, ChannelId, EditorBootTestGroup, "Platform2", PendingBadge, parentJob.Id, parentJob.TemplateId);
+				ValidateJobStepInformationAsync(consolidatedJobStepNotificationStates[5], spawnedJob.Id, spawnedJob.TemplateId, ChannelId, EditorBootTestGroup, "Platform2", PendingBadge, parentJob.Id, parentJob.TemplateId);
 
 				// Job Step notification for 'Editor Project Tests Platform2'
-				ValidateJobStepInformation(consolidatedJobStepNotificationStates[6], spawnedJob.Id, spawnedJob.TemplateId, ChannelId, EditorProjectTestGroup, "Platform2", PendingBadge, parentJob.Id, parentJob.TemplateId);
+				ValidateJobStepInformationAsync(consolidatedJobStepNotificationStates[6], spawnedJob.Id, spawnedJob.TemplateId, ChannelId, EditorProjectTestGroup, "Platform2", PendingBadge, parentJob.Id, parentJob.TemplateId);
 
 				// Simulate the cancellation of the 'Editor BootTest Platform1' step
 				spawnedJob = await AbortStepAsync(_mainStreamConfig, spawnedJob, 1, 0);
@@ -1013,25 +1013,25 @@ namespace HordeServer.Experimental.Tests.Notifications
 				Assert.AreEqual(7, consolidatedJobStepNotificationStates.Count);
 
 				// Job Step notification for 'Spawn Platform1 Test'
-				ValidateJobStepInformation(consolidatedJobStepNotificationStates[0], parentJob.Id, parentJob.TemplateId, ChannelId, SpawnJobTestGroup, "Platform1", PassingBadge);
+				ValidateJobStepInformationAsync(consolidatedJobStepNotificationStates[0], parentJob.Id, parentJob.TemplateId, ChannelId, SpawnJobTestGroup, "Platform1", PassingBadge);
 
 				// Job Step notification for 'Spawn Platform2 Test'
-				ValidateJobStepInformation(consolidatedJobStepNotificationStates[1], parentJob.Id, parentJob.TemplateId, ChannelId, SpawnJobTestGroup, "Platform2", PassingBadge);
+				ValidateJobStepInformationAsync(consolidatedJobStepNotificationStates[1], parentJob.Id, parentJob.TemplateId, ChannelId, SpawnJobTestGroup, "Platform2", PassingBadge);
 
 				// Job Step notification for 'Spawn Platform3 Test'
-				ValidateJobStepInformation(consolidatedJobStepNotificationStates[2], parentJob.Id, parentJob.TemplateId, ChannelId, SpawnJobTestGroup, "Platform3", PendingBadge);
+				ValidateJobStepInformationAsync(consolidatedJobStepNotificationStates[2], parentJob.Id, parentJob.TemplateId, ChannelId, SpawnJobTestGroup, "Platform3", PendingBadge);
 
 				// Job Step notification for 'Client BootTest Platform2'
-				ValidateJobStepInformation(consolidatedJobStepNotificationStates[3], spawnedJob.Id, spawnedJob.TemplateId, ChannelId, ClientBootTestGroup, "Platform2", PendingBadge, parentJob.Id, parentJob.TemplateId);
+				ValidateJobStepInformationAsync(consolidatedJobStepNotificationStates[3], spawnedJob.Id, spawnedJob.TemplateId, ChannelId, ClientBootTestGroup, "Platform2", PendingBadge, parentJob.Id, parentJob.TemplateId);
 
 				// Job Step notification for 'Client Project Tests Platform2'
-				ValidateJobStepInformation(consolidatedJobStepNotificationStates[4], spawnedJob.Id, spawnedJob.TemplateId, ChannelId, ClientProjectTestGroup, "Platform2", PendingBadge, parentJob.Id, parentJob.TemplateId);
+				ValidateJobStepInformationAsync(consolidatedJobStepNotificationStates[4], spawnedJob.Id, spawnedJob.TemplateId, ChannelId, ClientProjectTestGroup, "Platform2", PendingBadge, parentJob.Id, parentJob.TemplateId);
 
 				// Job Step notification for 'Editor BootTest Platform2'
-				ValidateJobStepInformation(consolidatedJobStepNotificationStates[5], spawnedJob.Id, spawnedJob.TemplateId, ChannelId, EditorBootTestGroup, "Platform2", FailureBadge, parentJob.Id, parentJob.TemplateId);
+				ValidateJobStepInformationAsync(consolidatedJobStepNotificationStates[5], spawnedJob.Id, spawnedJob.TemplateId, ChannelId, EditorBootTestGroup, "Platform2", FailureBadge, parentJob.Id, parentJob.TemplateId);
 
 				// Job Step notification for 'Editor Project Tests Platform2'
-				ValidateJobStepInformation(consolidatedJobStepNotificationStates[6], spawnedJob.Id, spawnedJob.TemplateId, ChannelId, EditorProjectTestGroup, "Platform2", SkippedBadge, parentJob.Id, parentJob.TemplateId);
+				ValidateJobStepInformationAsync(consolidatedJobStepNotificationStates[6], spawnedJob.Id, spawnedJob.TemplateId, ChannelId, EditorProjectTestGroup, "Platform2", SkippedBadge, parentJob.Id, parentJob.TemplateId);
 
 				// Simulate the cancellation of a job
 				spawnedJob = await AbortJobAsync(spawnedJob);
@@ -1059,13 +1059,13 @@ namespace HordeServer.Experimental.Tests.Notifications
 
 				// Job Step notification for the 'Spawn Platform2 Test' will now being labeled as passing since the step spawned the new job
 				// Even though the actual job had issues and was canceled which will be reported in the Slack thread
-				ValidateJobStepInformation(consolidatedJobStepNotificationStates[0], parentJob.Id, parentJob.TemplateId, ChannelId, SpawnJobTestGroup, "Platform1", PassingBadge);
+				ValidateJobStepInformationAsync(consolidatedJobStepNotificationStates[0], parentJob.Id, parentJob.TemplateId, ChannelId, SpawnJobTestGroup, "Platform1", PassingBadge);
 
 				// Job Step notification for 'Spawn Platform2 Test'
-				ValidateJobStepInformation(consolidatedJobStepNotificationStates[1], parentJob.Id, parentJob.TemplateId, ChannelId, SpawnJobTestGroup, "Platform2", PassingBadge);
+				ValidateJobStepInformationAsync(consolidatedJobStepNotificationStates[1], parentJob.Id, parentJob.TemplateId, ChannelId, SpawnJobTestGroup, "Platform2", PassingBadge);
 
 				// Job Step notification for 'Spawn Platform3 Test'
-				ValidateJobStepInformation(consolidatedJobStepNotificationStates[2], parentJob.Id, parentJob.TemplateId, ChannelId, SpawnJobTestGroup, "Platform3", PendingBadge);
+				ValidateJobStepInformationAsync(consolidatedJobStepNotificationStates[2], parentJob.Id, parentJob.TemplateId, ChannelId, SpawnJobTestGroup, "Platform3", PendingBadge);
 			}
 
 			// #3
@@ -1081,7 +1081,7 @@ namespace HordeServer.Experimental.Tests.Notifications
 
 				await WaitForNotificationAsync(fakeSink, parentJob, parentJob.Batches[1].Steps[2]);
 
-				await ValidateJobStepInformation(jobNotificationCollection, parentJob.Id, parentJob.TemplateId, ChannelId, parentJob.Batches[1].Steps[2].Id, SpawnJobTestGroup, "Platform3", PassingBadge);
+				await ValidateJobStepInformationAsync(jobNotificationCollection, parentJob.Id, parentJob.TemplateId, ChannelId, parentJob.Batches[1].Steps[2].Id, SpawnJobTestGroup, "Platform3", PassingBadge);
 
 				// Execute the Setup Build of the newly spawned job
 				spawnedJob = await RunBatchAsync(_mainStreamConfig, spawnedJob, 0);
@@ -1106,25 +1106,25 @@ namespace HordeServer.Experimental.Tests.Notifications
 				Assert.AreEqual(7, consolidatedJobStepNotificationStates.Count);
 
 				// Job Step notification for 'Spawn Platform1 Test'
-				ValidateJobStepInformation(consolidatedJobStepNotificationStates[0], parentJob.Id, parentJob.TemplateId, ChannelId, SpawnJobTestGroup, "Platform1", PassingBadge);
+				ValidateJobStepInformationAsync(consolidatedJobStepNotificationStates[0], parentJob.Id, parentJob.TemplateId, ChannelId, SpawnJobTestGroup, "Platform1", PassingBadge);
 
 				// Job Step notification for 'Spawn Platform2 Test'
-				ValidateJobStepInformation(consolidatedJobStepNotificationStates[1], parentJob.Id, parentJob.TemplateId, ChannelId, SpawnJobTestGroup, "Platform2", PassingBadge);
+				ValidateJobStepInformationAsync(consolidatedJobStepNotificationStates[1], parentJob.Id, parentJob.TemplateId, ChannelId, SpawnJobTestGroup, "Platform2", PassingBadge);
 
 				// Job Step notification for 'Spawn Platform3 Test'
-				ValidateJobStepInformation(consolidatedJobStepNotificationStates[2], parentJob.Id, parentJob.TemplateId, ChannelId, SpawnJobTestGroup, "Platform3", PassingBadge);
+				ValidateJobStepInformationAsync(consolidatedJobStepNotificationStates[2], parentJob.Id, parentJob.TemplateId, ChannelId, SpawnJobTestGroup, "Platform3", PassingBadge);
 
 				// Job Step notification for 'Client BootTest Platform3'
-				ValidateJobStepInformation(consolidatedJobStepNotificationStates[3], spawnedJob.Id, spawnedJob.TemplateId, ChannelId, ClientBootTestGroup, "Platform3", PendingBadge, parentJob.Id, parentJob.TemplateId);
+				ValidateJobStepInformationAsync(consolidatedJobStepNotificationStates[3], spawnedJob.Id, spawnedJob.TemplateId, ChannelId, ClientBootTestGroup, "Platform3", PendingBadge, parentJob.Id, parentJob.TemplateId);
 
 				// Job Step notification for 'Client Project Tests Platform3'
-				ValidateJobStepInformation(consolidatedJobStepNotificationStates[4], spawnedJob.Id, spawnedJob.TemplateId, ChannelId, ClientProjectTestGroup, "Platform3", PendingBadge, parentJob.Id, parentJob.TemplateId);
+				ValidateJobStepInformationAsync(consolidatedJobStepNotificationStates[4], spawnedJob.Id, spawnedJob.TemplateId, ChannelId, ClientProjectTestGroup, "Platform3", PendingBadge, parentJob.Id, parentJob.TemplateId);
 
 				// Job Step notification for 'Editor BootTest Platform3'
-				ValidateJobStepInformation(consolidatedJobStepNotificationStates[5], spawnedJob.Id, spawnedJob.TemplateId, ChannelId, EditorBootTestGroup, "Platform3", PendingBadge, parentJob.Id, parentJob.TemplateId);
+				ValidateJobStepInformationAsync(consolidatedJobStepNotificationStates[5], spawnedJob.Id, spawnedJob.TemplateId, ChannelId, EditorBootTestGroup, "Platform3", PendingBadge, parentJob.Id, parentJob.TemplateId);
 
 				// Job Step notification for 'Editor Project Tests Platform3'
-				ValidateJobStepInformation(consolidatedJobStepNotificationStates[6], spawnedJob.Id, spawnedJob.TemplateId, ChannelId, EditorProjectTestGroup, "Platform3", PendingBadge, parentJob.Id, parentJob.TemplateId);
+				ValidateJobStepInformationAsync(consolidatedJobStepNotificationStates[6], spawnedJob.Id, spawnedJob.TemplateId, ChannelId, EditorProjectTestGroup, "Platform3", PendingBadge, parentJob.Id, parentJob.TemplateId);
 
 				spawnedJob = await RunBatchAsync(_mainStreamConfig, spawnedJob, 1);
 				spawnedJob = await RunStepAsync(_mainStreamConfig, spawnedJob, 1, 0, JobStepOutcome.Failure); // Editor BootTest Platform3
@@ -1132,15 +1132,15 @@ namespace HordeServer.Experimental.Tests.Notifications
 				await WaitForNotificationAsync(fakeSink, spawnedJob, spawnedJob.Batches[1].Steps[0]);
 
 				// Validate that the 'Editor BootTest Platform3' step is recorded as having failures and 'Editor Project Tests Platform3' is recorded as being skipped due to having a dependency on the BootTest
-				await ValidateJobStepInformation(jobNotificationCollection, spawnedJob.Id, spawnedJob.TemplateId, ChannelId, spawnedJob.Batches[1].Steps[0].Id, EditorBootTestGroup, "Platform3", FailureBadge, parentJob.Id, parentJob.TemplateId);
-				await ValidateJobStepInformation(jobNotificationCollection, spawnedJob.Id, spawnedJob.TemplateId, ChannelId, spawnedJob.Batches[1].Steps[1].Id, EditorProjectTestGroup, "Platform3", SkippedBadge, parentJob.Id, parentJob.TemplateId);
+				await ValidateJobStepInformationAsync(jobNotificationCollection, spawnedJob.Id, spawnedJob.TemplateId, ChannelId, spawnedJob.Batches[1].Steps[0].Id, EditorBootTestGroup, "Platform3", FailureBadge, parentJob.Id, parentJob.TemplateId);
+				await ValidateJobStepInformationAsync(jobNotificationCollection, spawnedJob.Id, spawnedJob.TemplateId, ChannelId, spawnedJob.Batches[1].Steps[1].Id, EditorProjectTestGroup, "Platform3", SkippedBadge, parentJob.Id, parentJob.TemplateId);
 
 				spawnedJob = await RunStepAsync(_mainStreamConfig, spawnedJob, 1, 2, JobStepOutcome.Warnings); // Client BootTest Platform3
 
 				await WaitForNotificationAsync(fakeSink, spawnedJob, spawnedJob.Batches[1].Steps[2]);
 
 				// Validate that the Client BootTest Platform3 step is recorded as having warnings
-				await ValidateJobStepInformation(jobNotificationCollection, spawnedJob.Id, spawnedJob.TemplateId, ChannelId, spawnedJob.Batches[1].Steps[2].Id, ClientBootTestGroup, "Platform3", WarningBadge, parentJob.Id, parentJob.TemplateId);
+				await ValidateJobStepInformationAsync(jobNotificationCollection, spawnedJob.Id, spawnedJob.TemplateId, ChannelId, spawnedJob.Batches[1].Steps[2].Id, ClientBootTestGroup, "Platform3", WarningBadge, parentJob.Id, parentJob.TemplateId);
 
 				spawnedJob = await RunStepAsync(_mainStreamConfig, spawnedJob, 1, 3, JobStepOutcome.Success); // Client Project Tests Platform3
 
@@ -1226,10 +1226,10 @@ namespace HordeServer.Experimental.Tests.Notifications
 			Assert.AreEqual(2, jobStepNotificationStates.Count);
 
 			// Job Step notification for 'Spawn Platform1 Test'
-			ValidateJobStepInformation(jobStepNotificationStates[0], parentJob.Id, parentJob.TemplateId, ChannelId, SpawnJobTestGroup, "Platform1", PendingBadge);
+			ValidateJobStepInformationAsync(jobStepNotificationStates[0], parentJob.Id, parentJob.TemplateId, ChannelId, SpawnJobTestGroup, "Platform1", PendingBadge);
 
 			// Job Step notification for 'Spawn Platform2 Test'
-			ValidateJobStepInformation(jobStepNotificationStates[1], parentJob.Id, parentJob.TemplateId, ChannelId, SpawnJobTestGroup, "Platform2", PendingBadge);
+			ValidateJobStepInformationAsync(jobStepNotificationStates[1], parentJob.Id, parentJob.TemplateId, ChannelId, SpawnJobTestGroup, "Platform2", PendingBadge);
 
 			parentJob = await RunBatchAsync(_mainStreamConfig, parentJob, 1);
 
@@ -1251,7 +1251,7 @@ namespace HordeServer.Experimental.Tests.Notifications
 
 			// Execute the Setup Build steps of all of our spawned jobs
 			spawnedJob1 = await RunBatchAsync(_mainStreamConfig, spawnedJob1, 0);
-			spawnedJob1 = await RunStepAsync(_mainStreamConfig, spawnedJob1, 0, 0, JobStepOutcome.Success); // Setup Build
+			_ = await RunStepAsync(_mainStreamConfig, spawnedJob1, 0, 0, JobStepOutcome.Success); // Setup Build
 
 			spawnedJob2 = await RunBatchAsync(_mainStreamConfig, spawnedJob2, 0);
 			spawnedJob2 = await RunStepAsync(_mainStreamConfig, spawnedJob2, 0, 0, JobStepOutcome.Success); // Setup Build
diff --git a/Engine/Source/Programs/Shared/EpicGames.Core/LogEventParser.cs b/Engine/Source/Programs/Shared/EpicGames.Core/LogEventParser.cs
--- a/Engine/Source/Programs/Shared/EpicGames.Core/LogEventParser.cs
+++ b/Engine/Source/Programs/Shared/EpicGames.Core/LogEventParser.cs
@@ -430,6 +430,7 @@ namespace EpicGames.Core
 		/// Process any data in the buffer
 		/// </summary>
 		/// <param name="bFlush">Whether we've reached the end of the stream</param>
+		/// <param name="disableMatching">Whether to skip matching</param>
 		void ProcessData(bool bFlush, bool disableMatching = false)
 		{
 			_activeTimer.Start();
