Horde: Add agent property for opting out from server-initiated updates of agent software

In the case of agent running on workstation, this ability allows for more control and avoids the agent process restarting itself.

diff --git a/Engine/Source/Programs/Horde/Plugins/Build/HordeServer.Build.Tests/Agents/AgentServiceTest.cs b/Engine/Source/Programs/Horde/Plugins/Build/HordeServer.Build.Tests/Agents/AgentServiceTest.cs
--- a/Engine/Source/Programs/Horde/Plugins/Build/HordeServer.Build.Tests/Agents/AgentServiceTest.cs
+++ b/Engine/Source/Programs/Horde/Plugins/Build/HordeServer.Build.Tests/Agents/AgentServiceTest.cs
@@ -155,6 +155,19 @@ public class AgentServiceTest : BuildTestSetup
 		Assert.AreEqual(300, rate2!.Value, 0.1);
 	}
 	
+	[TestMethod]
+	[DataRow($"unset=null", true)]
+	[DataRow($"{KnownPropertyNames.AutoUpdate}=true", true)]
+	[DataRow($"{KnownPropertyNames.AutoUpdate}=TrUE", true)]
+	[DataRow($"{KnownPropertyNames.AutoUpdate}=false", false)]
+	[DataRow($"{KnownPropertyNames.AutoUpdate}=FaLse", false)]
+	public async Task Properties_AutoUpdate_Async(string property, bool expected)
+	{
+		IAgent newAgent = await AgentService.CreateAgentAsync(new CreateAgentOptions(new AgentId("agent1"), AgentMode.Dedicated, false, ""));
+		IAgent agent = await AgentService.CreateSessionAsync(newAgent, new RpcAgentCapabilities([property]), "test");
+		Assert.AreEqual(expected, agent.IsAutoUpdateEnabled());
+	}
+	
 	[TestMethod]
 	public async Task Properties_ServerDefined_AgentCannotOverride_Async()
 	{
diff --git a/Engine/Source/Programs/Horde/Plugins/Compute/HordeServer.Compute/Agents/IAgent.cs b/Engine/Source/Programs/Horde/Plugins/Compute/HordeServer.Compute/Agents/IAgent.cs
--- a/Engine/Source/Programs/Horde/Plugins/Compute/HordeServer.Compute/Agents/IAgent.cs
+++ b/Engine/Source/Programs/Horde/Plugins/Compute/HordeServer.Compute/Agents/IAgent.cs
@@ -468,6 +468,17 @@ namespace HordeServer.Agents
 			return values.Count > 0 && values[0].Equals("true", StringComparison.OrdinalIgnoreCase);
 		}
 		
+		/// <summary>
+		/// Checks whether an agent wants to auto-update itself
+		/// </summary>
+		/// <param name="agent">Agent to query</param>
+		/// <returns>Default true, unless explicitly set to false</returns>
+		public static bool IsAutoUpdateEnabled(this IAgent agent)
+		{
+			string? value = agent.GetPropertyValues(KnownPropertyNames.AutoUpdate).FirstOrDefault();
+			return String.IsNullOrEmpty(value) || !String.Equals(value, "false", StringComparison.OrdinalIgnoreCase);
+		}
+		
 		/// <summary>
 		/// Get optional tool ID as reported by agent
 		/// <see cref="KnownPropertyNames.ToolId" />
diff --git a/Engine/Source/Programs/Horde/Plugins/Compute/HordeServer.Compute/ComputeServerConfig.cs b/Engine/Source/Programs/Horde/Plugins/Compute/HordeServer.Compute/ComputeServerConfig.cs
--- a/Engine/Source/Programs/Horde/Plugins/Compute/HordeServer.Compute/ComputeServerConfig.cs
+++ b/Engine/Source/Programs/Horde/Plugins/Compute/HordeServer.Compute/ComputeServerConfig.cs
@@ -14,7 +14,7 @@ namespace HordeServer
 	public class ComputeServerConfig : PluginServerConfig
 	{
 		/// <summary>
-		/// Whether to enable the upgrade task source.
+		/// Whether to enable the upgrade task source, always upgrading agents to the latest version
 		/// </summary>
 		public bool EnableUpgradeTasks { get; set; } = true;
 
diff --git a/Engine/Source/Programs/Horde/Plugins/Compute/HordeServer.Compute/Tasks/UpgradeTaskSource.cs b/Engine/Source/Programs/Horde/Plugins/Compute/HordeServer.Compute/Tasks/UpgradeTaskSource.cs
--- a/Engine/Source/Programs/Horde/Plugins/Compute/HordeServer.Compute/Tasks/UpgradeTaskSource.cs
+++ b/Engine/Source/Programs/Horde/Plugins/Compute/HordeServer.Compute/Tasks/UpgradeTaskSource.cs
@@ -43,6 +43,11 @@ namespace HordeServer.Tasks
 			{
 				return SkipAsync(cancellationToken);
 			}
+			
+			if (!agent.IsAutoUpdateEnabled())
+			{
+				return SkipAsync(cancellationToken);
+			}
 
 			(ITool, IToolDeployment)? required = await GetRequiredSoftwareVersionAsync(agent, cancellationToken);
 			if (required == null)
diff --git a/Engine/Source/Programs/Shared/EpicGames.Horde/Agents/KnownPropertyNames.cs b/Engine/Source/Programs/Shared/EpicGames.Horde/Agents/KnownPropertyNames.cs
--- a/Engine/Source/Programs/Shared/EpicGames.Horde/Agents/KnownPropertyNames.cs
+++ b/Engine/Source/Programs/Shared/EpicGames.Horde/Agents/KnownPropertyNames.cs
@@ -93,10 +93,16 @@ public static class KnownPropertyNames
 	public const string WineEnabled = "WineEnabled";
 	
 	/// <summary>
-	/// Whether the agent is trusted
+	/// Whether the agent is trusted (boolean)
 	/// </summary>
 	public const string Trusted = "Trusted";
 	
+	/// <summary>
+	/// Whether the agent want to receive software updates from the server (boolean)
+	/// In some cases, an external script/process will take care of updating the agent (for example, Unreal Toolbox)
+	/// </summary>
+	public const string AutoUpdate = "AutoUpdate";
+	
 	/// <summary>
 	/// ID of current lease
 	/// </summary>
