Horde: Update telemetry configs to replace display with unit
fyi mark.lintott, josh.engebretson

diff --git a/Engine/Source/Programs/Horde/HordeServer/Defaults/default-views.telemetry.json b/Engine/Source/Programs/Horde/HordeServer/Defaults/default-views.telemetry.json
--- a/Engine/Source/Programs/Horde/HordeServer/Defaults/default-views.telemetry.json
+++ b/Engine/Source/Programs/Horde/HordeServer/Defaults/default-views.telemetry.json
@@ -166,7 +166,7 @@
 						},
 						{
 							"name": "User Loading Count (Live)",
-							"display": "Value",
+							"unit": "Value",
 							"max": 100,
 							"metrics": [
 								{
@@ -192,7 +192,7 @@
 						},
 						{
 							"name": "User Editor LoadMap Count (Live)",
-							"display": "Value",
+							"unit": "Value",
 							"metrics": [
 								{
 									"id": "core-loading-user-loadmap-count"
@@ -339,7 +339,7 @@
 						},
 						{
 							"name": "Cook Shaders Compiled (KPIs)",
-							"display": "Value",
+							"unit": "Value",
 							"metrics": [
 								{
 									"id": "core-cooking-nightly-cook-shaders-compiled-windows-kpi",
@@ -388,7 +388,7 @@
 						},
 						{
 							"name": "Cook Shaders Compiled (KPIs)",
-							"display": "Value",
+							"unit": "Value",
 							"metrics": [
 								{
 									"id": "core-cooking-nightly-cook-shaders-compiled-windows-kpi",
@@ -434,7 +434,7 @@
 						},
 						{
 							"name": "Package And Test Shaders Compiled",
-							"display": "Value",
+							"unit": "Value",
 							"metrics": [
 								{
 									"id": "core-cooking-package-and-test-shaders-compiled"
@@ -451,7 +451,7 @@
 						},
 						{
 							"name": "Nightly Cook Shaders Compiled",
-							"display": "Value",
+							"unit": "Value",
 							"metrics": [
 								{
 									"id": "core-cooking-nightly-shaders-compiled"
@@ -481,7 +481,7 @@
 					"charts": [
 						{
 							"name": "User DDC Backend Efficiency",
-							"display": "Ratio",
+							"unit": "Ratio",
 							"metrics": [
 								{
 									"id": "core-loading-user-totalddchitratio",
@@ -503,7 +503,7 @@
 						},
 						{
 							"name": "User DDC Asset Efficiency",
-							"display": "Ratio",
+							"unit": "Ratio",
 							"metrics": [
 								{
 									"id": "core-loading-user-ddc-asset-efficiency"
@@ -517,7 +517,7 @@
 					"charts": [
 						{
 							"name": "Memory By System",
-							"display": "Value",
+							"unit": "Value",
 							"metrics": [
 								{
 									"id": "core-memory-llm-systems"
@@ -526,7 +526,7 @@
 						},
 						{
 							"name": "Memory By AssetClass",
-							"display": "Value",
+							"unit": "Value",
 							"metrics": [
 								{
 									"id": "core-memory-llm-assetclasses"
@@ -535,7 +535,7 @@
 						},
 						{
 							"name": "Memory By Asset",
-							"display": "Value",
+							"unit": "Value",
 							"metrics": [
 								{
 									"id": "core-memory-llm-assets"
@@ -549,7 +549,7 @@
 					"charts": [
 						{
 							"name": "Content Size Total",
-							"display": "Value",
+							"unit": "Value",
 							"min": 0,
 							"metrics": [
 								{
@@ -559,7 +559,7 @@
 						},
 						{
 							"name": "Content Size By Mount Point",
-							"display": "Value",
+							"unit": "Value",
 							"min": 0,
 							"metrics": [
 								{
@@ -569,7 +569,7 @@
 						},
 						{
 							"name": "Virtualization By Mount Point",
-							"display": "Value",
+							"unit": "Value",
 							"min": 0,
 							"metrics": [
 								{
@@ -584,7 +584,7 @@
 					"charts": [
 						{
 							"name": "Remote Distibution All (Average)",
-							"display": "Ratio",
+							"unit": "Ratio",
 							"metrics": [
 								{
 									"id": "ubt-uba-remote-usage-all"
@@ -593,7 +593,7 @@
 						},
 						{
 							"name": "Cache Rate All (Average)",
-							"display": "Ratio",
+							"unit": "Ratio",
 							"metrics": [
 								{
 									"id": "ubt-uba-cache-rate-all"
@@ -602,7 +602,7 @@
 						},
 						{
 							"name": "Actions All (Sum)",
-							"display": "Value",
+							"unit": "Value",
 							"metrics": [
 								{
 									"id": "ubt-uba-total-actions-all",
@@ -616,7 +616,7 @@
 						},
 						{
 							"name": "Failure Rate All (Average)",
-							"display": "Ratio",
+							"unit": "Ratio",
 							"metrics": [
 								{
 									"id": "ubt-uba-failure-rate-all"
@@ -633,7 +633,7 @@
 						},
 						{
 							"name": "Remote Distribution (Average)",
-							"display": "Ratio",
+							"unit": "Ratio",
 							"metrics": [
 								{
 									"id": "ubt-uba-remote-usage"
@@ -642,7 +642,7 @@
 						},
 						{
 							"name": "Cache Rate (Average)",
-							"display": "Ratio",
+							"unit": "Ratio",
 							"metrics": [
 								{
 									"id": "ubt-uba-cache-rate-avg"
@@ -651,7 +651,7 @@
 						},
 						{
 							"name": "Actions (Average)",
-							"display": "Value",
+							"unit": "Value",
 							"metrics": [
 								{
 									"id": "ubt-uba-total-actions-avg",
@@ -665,7 +665,7 @@
 						},
 						{
 							"name": "Actions (Sum)",
-							"display": "Value",
+							"unit": "Value",
 							"metrics": [
 								{
 									"id": "ubt-uba-total-actions",
@@ -679,7 +679,7 @@
 						},
 						{
 							"name": "Failure Rate (Average)",
-							"display": "Ratio",
+							"unit": "Ratio",
 							"metrics": [
 								{
 									"id": "ubt-uba-failure-rate"
@@ -749,7 +749,7 @@
 					"charts": [
 						{
 							"name": "Size By Install Type",
-							"display": "Value",
+							"unit": "Value",
 							"min": 0,
 							"metrics": [
 								{
@@ -759,7 +759,7 @@
 						},
 						{
 							"name": "Size Per Package",
-							"display": "Value",
+							"unit": "Value",
 							"min": 0,
 							"metrics": [
 								{
@@ -774,7 +774,7 @@
 					"charts": [
 						{
 							"name": "Debug",
-							"display": "Value",
+							"unit": "Value",
 							"metrics": [
 								{
 									"id": "debug-probe"
@@ -783,7 +783,7 @@
 						},
 						{
 							"name": "MountPoint Test",
-							"display": "Value",
+							"unit": "Value",
 							"metrics": [
 								{
 									"id": "core-content-events"
