Horde: Add Horde's user claim when issuing JWT tokens for built-in Horde auth mode

diff --git a/Engine/Source/Programs/Horde/HordeServer/Accounts/OAuthController.cs b/Engine/Source/Programs/Horde/HordeServer/Accounts/OAuthController.cs
--- a/Engine/Source/Programs/Horde/HordeServer/Accounts/OAuthController.cs
+++ b/Engine/Source/Programs/Horde/HordeServer/Accounts/OAuthController.cs
@@ -449,6 +449,7 @@ namespace HordeServer.Accounts
 			payload[ClaimTypes.Name] = subject.Name;
 			payload[ClaimTypes.Email] = subject.Email;
 			payload[HordeClaimTypes.AccountId] = subject.Id.ToString();
+			payload[HordeClaimTypes.User] = subject.Login;
 			
 			foreach (IUserClaim claim in subject.Claims)
 			{
