Horde: Update release notes

#rnx

diff --git a/Engine/Source/Programs/Horde/Docs/ReleaseNotes.md b/Engine/Source/Programs/Horde/Docs/ReleaseNotes.md
--- a/Engine/Source/Programs/Horde/Docs/ReleaseNotes.md
+++ b/Engine/Source/Programs/Horde/Docs/ReleaseNotes.md
@@ -1,5 +1,37 @@
 # Release Notes
 
+## 2025-04-10
+
+* Run tests for experimental plugin in Dockerfile builds (41526895)
+* Add agent property for opting out from server-initiated updates of agent software
+  In the case of agent running on workstation, this ability allows for more control and avoids the agent process restarting itself. (41526788)
+* Do not run matchers on partial json log line entries.
+  This prevents us from trying to match unreliable log lines which were only partially written out e.g. due to the process crashing mid output. (41526045)
+* Fixed warnings in LocalizationIssueHandler.cs (41525967)
+* Improved Horde's LocalizationIssueHandler to get ready for many C++ localization warnings to change to structured logs (UE_LOGFMT) (41519652)
+* Change GlobalConfig.PostLoad and IPluginConfig.PostLoad methods to async
+  This is to allow, in the future, config's to use async APIs for example to fetch secrets.
+  Some tests had setup in constructors changed these to use the TestInitialize attribute. (41483040)
+* Make update session call more robust
+  * Catch HttpRequestException in addition to RpcException
+  * Implement exponential backoff starting at 1s (max 20s)
+  * Increase max retries from 3 to 8
+  * Improve logging slightly to explain attempt count and delay (41414189)
+* Add topological sort of plugins by their dependencies to call the PostLoad Method in the correct order (41410626)
+* Add admin claim resolution for Horde auth mode
+  Admin claims are now properly resolved when users sign in using Horde auth mode, matching the existing behavior in OIDC/Okta auth handlers. (41410594)
+* Replace display field with unit field for Horde Analytics, and add Bytes, BytesPerSecond, Seconds, and Percentage options (41366205)
+* Display the year when not current to reduce confusion when looking at historical jobs (41362940)
+* Improve data fetching error and surface network connection as likely cause on job page (41362331)
+* Fix step ETA times sometimes being in the past (41355458)
+* Add better indicator for inidividual data points in Horde Analytics (41352087)
+* Oidc Token - Added ability to read encrypted oidc configuration from a url. Similar to the horde url but encrypted and not exactly the same configuration layout. Is intended to be able to replace the horde url as a more generic option. --AuthConfigUrl for OidcToken.exe
+  This should also hopefully cover most of the cases for submitting a oidc-configuration.json into source control.
+  Also added --ConfigPath as a way to explicitly set the path to oidc-configuration.json if redistributing this outside of the ue engine sync.
+  Lastly I added source code generators for json serialization and configuration builder to enable AOT compilation of OidcToken (which is not used yet). This also fixes most trimming warnings. (41351365)
+* Fix issue with artifacts not properly loading between step navigation (41327584)
+* Set minimum value for the Package and Content charts in Horde Analytics to zero (41316971)
+
 ## 2025-04-03
 
 * Horde Analytics: Fix issue with sum values updating without matching on group (41278970)
