Horde: Allow only compute tasks to be scheduled on workstation agents

diff --git a/Engine/Source/Programs/Horde/Plugins/Build/HordeServer.Build/Jobs/JobTaskSource.cs b/Engine/Source/Programs/Horde/Plugins/Build/HordeServer.Build/Jobs/JobTaskSource.cs
--- a/Engine/Source/Programs/Horde/Plugins/Build/HordeServer.Build/Jobs/JobTaskSource.cs
+++ b/Engine/Source/Programs/Horde/Plugins/Build/HordeServer.Build/Jobs/JobTaskSource.cs
@@ -714,6 +714,11 @@ namespace HordeServer.Jobs
 		/// <inheritdoc/>
 		public override Task<Task<CreateLeaseOptions?>> AssignLeaseAsync(IAgent agent, CancellationToken cancellationToken)
 		{
+			if (agent.Mode == AgentMode.Workstation)
+			{
+				return Task.FromResult(SkipAsync(cancellationToken));
+			}
+			
 			QueueWaiter waiter = new QueueWaiter(agent);
 			lock (_lockObject)
 			{
diff --git a/Engine/Source/Programs/Horde/Plugins/Build/HordeServer.Build/Tasks/ConformTaskSource.cs b/Engine/Source/Programs/Horde/Plugins/Build/HordeServer.Build/Tasks/ConformTaskSource.cs
--- a/Engine/Source/Programs/Horde/Plugins/Build/HordeServer.Build/Tasks/ConformTaskSource.cs
+++ b/Engine/Source/Programs/Horde/Plugins/Build/HordeServer.Build/Tasks/ConformTaskSource.cs
@@ -150,7 +150,7 @@ namespace HordeServer.Tasks
 		/// <inheritdoc/>
 		public override async Task<Task<CreateLeaseOptions?>> AssignLeaseAsync(IAgent agent, CancellationToken cancellationToken)
 		{
-			if (!_settings.CurrentValue.EnableConformTasks || !_buildConfig.CurrentValue.EnableConformTasks)
+			if (!_settings.CurrentValue.EnableConformTasks || !_buildConfig.CurrentValue.EnableConformTasks || agent.Mode == AgentMode.Workstation)
 			{
 				return SkipAsync(cancellationToken);
 			}
diff --git a/Engine/Source/Programs/Horde/Plugins/Compute/HordeServer.Compute/Tasks/RestartTaskSource.cs b/Engine/Source/Programs/Horde/Plugins/Compute/HordeServer.Compute/Tasks/RestartTaskSource.cs
--- a/Engine/Source/Programs/Horde/Plugins/Compute/HordeServer.Compute/Tasks/RestartTaskSource.cs
+++ b/Engine/Source/Programs/Horde/Plugins/Compute/HordeServer.Compute/Tasks/RestartTaskSource.cs
@@ -27,6 +27,11 @@ namespace HordeServer.Tasks
 
 		public override async Task<Task<CreateLeaseOptions?>> AssignLeaseAsync(IAgent agent, CancellationToken cancellationToken)
 		{
+			if (agent.Mode == AgentMode.Workstation)
+			{
+				return SkipAsync(cancellationToken);
+			}
+			
 			if (!agent.RequestForceRestart)
 			{
 				if (!agent.RequestRestart)
diff --git a/Engine/Source/Programs/Horde/Plugins/Compute/HordeServer.Compute/Tasks/ShutdownTaskSource.cs b/Engine/Source/Programs/Horde/Plugins/Compute/HordeServer.Compute/Tasks/ShutdownTaskSource.cs
--- a/Engine/Source/Programs/Horde/Plugins/Compute/HordeServer.Compute/Tasks/ShutdownTaskSource.cs
+++ b/Engine/Source/Programs/Horde/Plugins/Compute/HordeServer.Compute/Tasks/ShutdownTaskSource.cs
@@ -26,7 +26,7 @@ namespace HordeServer.Tasks
 
 		public override async Task<Task<CreateLeaseOptions?>> AssignLeaseAsync(IAgent agent, CancellationToken cancellationToken)
 		{
-			if (!agent.RequestShutdown)
+			if (!agent.RequestShutdown || agent.Mode == AgentMode.Workstation)
 			{
 				return SkipAsync(cancellationToken);
 			}
