Horde: Log type of lease being executed in agent

diff --git a/Engine/Source/Programs/Horde/HordeAgent/Leases/LeaseHandler.cs b/Engine/Source/Programs/Horde/HordeAgent/Leases/LeaseHandler.cs
--- a/Engine/Source/Programs/Horde/HordeAgent/Leases/LeaseHandler.cs
+++ b/Engine/Source/Programs/Horde/HordeAgent/Leases/LeaseHandler.cs
@@ -107,7 +107,7 @@ namespace HordeAgent.Leases
 			span.SetAttribute("horde.lease.id", RpcLease.Id.ToString());
 			span.SetAttribute("horde.agent.id", session.AgentId.ToString());
 			span.SetAttribute("horde.agent.session_id", session.SessionId.ToString());
-			logger.LogInformation("Handling lease {LeaseId}", Id);
+			logger.LogInformation("Handling lease {LeaseId}: {LeaseTypeUrl}", Id, RpcLease.Payload.TypeUrl);
 
 			// Get the lease outcome
 			LeaseResult result = LeaseResult.Failed;
