Horde: Log server profile used during agent startup

diff --git a/Engine/Source/Programs/Horde/HordeAgent/Services/SessionFactory.cs b/Engine/Source/Programs/Horde/HordeAgent/Services/SessionFactory.cs
--- a/Engine/Source/Programs/Horde/HordeAgent/Services/SessionFactory.cs
+++ b/Engine/Source/Programs/Horde/HordeAgent/Services/SessionFactory.cs
@@ -111,7 +111,7 @@ namespace HordeAgent.Services
 
 			// Print the server info
 			ServerProfile serverProfile = currentSettings.GetCurrentServerProfile();
-			logger.LogInformation("Server: {Server}", serverProfile.Url);
+			logger.LogInformation("Server URL: {Server} Profile: {ProfileName}", serverProfile.Url, serverProfile.Name);
 
 			// Show the worker capabilities
 			RpcAgentCapabilities capabilities = await capabilitiesService.GetCapabilitiesAsync(workingDir);
