Horde: Fix Commits Viewer and Robomerge links not appearing in slack messages

diff --git a/Engine/Source/Programs/Horde/HordeServer/ServerApp.cs b/Engine/Source/Programs/Horde/HordeServer/ServerApp.cs
--- a/Engine/Source/Programs/Horde/HordeServer/ServerApp.cs
+++ b/Engine/Source/Programs/Horde/HordeServer/ServerApp.cs
@@ -353,7 +353,9 @@ namespace HordeServer
 			KeyValuePair.Create("Horde:PerforceConnectionPoolSize", "Horde:Plugins:Build:PerforceConnectionPoolSize"),
 			KeyValuePair.Create("Horde:EnableConformTasks", "Horde:Plugins:Build:EnableConformTasks"),
 			KeyValuePair.Create("Horde:P4SwarmUrl", "Horde:Plugins:Build:P4SwarmUrl"),
-			KeyValuePair.Create("Horde:JiraUsername", "Horde:Plugins:Build:JiraUsername"),
+			KeyValuePair.Create("Horde:RobomergeUrl", "Horde:Plugins:Build:RobomergeUrl"),
+            KeyValuePair.Create("Horde:CommitsViewerUrl", "Horde:Plugins:Build:CommitsViewerUrl"),
+            KeyValuePair.Create("Horde:JiraUsername", "Horde:Plugins:Build:JiraUsername"),
 			KeyValuePair.Create("Horde:JiraApiToken", "Horde:Plugins:Build:JiraApiToken"),
 			KeyValuePair.Create("Horde:JiraUrl", "Horde:Plugins:Build:JiraUrl"),
 			KeyValuePair.Create("Horde:SharedDeviceCheckoutDays", "Horde:Plugins:Build:SharedDeviceCheckoutDays"),
