Horde: Skip executing conform leases if agent is running in workstation mode

Workstation agents only deal with stateless compute leases.

diff --git a/Engine/Source/Programs/Horde/HordeAgent/Leases/Handlers/ConformHandler.cs b/Engine/Source/Programs/Horde/HordeAgent/Leases/Handlers/ConformHandler.cs
--- a/Engine/Source/Programs/Horde/HordeAgent/Leases/Handlers/ConformHandler.cs
+++ b/Engine/Source/Programs/Horde/HordeAgent/Leases/Handlers/ConformHandler.cs
@@ -9,21 +9,24 @@ using HordeAgent.Services;
 using HordeCommon.Rpc.Messages;
 using HordeCommon.Rpc.Tasks;
 using Microsoft.Extensions.Logging;
+using Microsoft.Extensions.Options;
 using OpenTelemetry.Trace;
 
 namespace HordeAgent.Leases.Handlers
 {
-	class ConformHandler : LeaseHandler<ConformTask>
+	internal class ConformHandler(RpcLease lease, IOptionsMonitor<AgentSettings> settings) : LeaseHandler<ConformTask>(lease)
 	{
-		public ConformHandler(RpcLease lease)
-			: base(lease)
-		{ }
-
 		/// <inheritdoc/>
 		protected override async Task<LeaseResult> ExecuteAsync(ISession session, LeaseId leaseId, ConformTask conformTask, Tracer tracer, ILogger localLogger, CancellationToken cancellationToken)
 		{
 			await using IServerLogger serverLogger = session.HordeClient.CreateServerLogger(LogId.Parse(conformTask.LogId)).WithLocalLogger(localLogger);
-
+			
+			if (settings.CurrentValue.Mode == AgentMode.Workstation)
+			{
+				serverLogger.LogInformation("Agent is running in workstation mode. Conform lease ignored.");
+				return LeaseResult.Success;
+			}
+			
 			List<string> arguments = new List<string>();
 			arguments.Add("execute");
 			arguments.Add("conform");
@@ -47,10 +50,9 @@ namespace HordeAgent.Leases.Handlers
 		}
 	}
 
-	class ConformHandlerFactory : LeaseHandlerFactory<ConformTask>
+	internal class ConformHandlerFactory(IOptionsMonitor<AgentSettings> settings) : LeaseHandlerFactory<ConformTask>
 	{
-		public override LeaseHandler<ConformTask> CreateHandler(RpcLease lease)
-			=> new ConformHandler(lease);
+		public override LeaseHandler<ConformTask> CreateHandler(RpcLease lease) => new ConformHandler(lease, settings);
 	}
 }
 
