Horde: Add pool ID argument to "compute run" in Horde CLI tool

diff --git a/Engine/Source/Programs/Horde/HordeCmd/Commands/Compute/ComputeRun.cs b/Engine/Source/Programs/Horde/HordeCmd/Commands/Compute/ComputeRun.cs
--- a/Engine/Source/Programs/Horde/HordeCmd/Commands/Compute/ComputeRun.cs
+++ b/Engine/Source/Programs/Horde/HordeCmd/Commands/Compute/ComputeRun.cs
@@ -36,6 +36,10 @@ namespace Horde.Commands.Compute
 		[CommandLine("-Cluster")]
 		[Description("Name of the cluster to run the task on.")]
 		public string ClusterId { get; set; } = "default";
+		
+		[CommandLine("-Pool")]
+		[Description("ID of pool to run this task in")]
+		public string? Pool { get; set; }
 
 		[CommandLine("-Requirements=")]
 		[Description("Query string to select the agent to run on")]
@@ -102,6 +106,12 @@ namespace Horde.Commands.Compute
 			{
 				requirements = new Requirements(Condition.Parse(Requirements));
 			}
+			
+			if (Pool != null)
+			{
+				requirements ??= new Requirements();
+				requirements.Pool = Pool;
+			}
 
 			ConnectionMetadataRequest cmr = new();
 			switch (Encryption?.ToUpperInvariant())
