Horde: Test case-insensitivity for pool ID when matching resources for compute cluster

diff --git a/Engine/Source/Programs/Horde/Plugins/Build/HordeServer.Build.Tests/Compute/ComputeServiceTest.cs b/Engine/Source/Programs/Horde/Plugins/Build/HordeServer.Build.Tests/Compute/ComputeServiceTest.cs
--- a/Engine/Source/Programs/Horde/Plugins/Build/HordeServer.Build.Tests/Compute/ComputeServiceTest.cs
+++ b/Engine/Source/Programs/Horde/Plugins/Build/HordeServer.Build.Tests/Compute/ComputeServiceTest.cs
@@ -286,7 +286,7 @@ namespace HordeServer.Tests.Compute
 			IOptionsMonitor<ComputeConfig> computeConfig = ServiceProvider.GetRequiredService<IOptionsMonitor<ComputeConfig>>();
 			computeConfig.CurrentValue.Clusters = new List<ComputeClusterConfig>
 			{
-				new() { Id = cluster1, Condition = $"pool == '{poolFoo.ToString()}'" },
+				new() { Id = cluster1, Condition = $"pool == 'fOO'" }, // Test case-insensitivity of pool ID matching
 				new() { Id = cluster2, Condition = $"pool == '{poolBar.ToString()}'" }
 			};
 			await GlobalConfig.CurrentValue.PostLoadAsync(new ServerSettings(), Array.Empty<ILoadedPlugin>(), Array.Empty<IDefaultAclModifier>());
