Horde: Add cloud url copy operation to artifact split button

diff --git a/Engine/Source/Programs/Horde/HordeDashboard/src/components/artifacts/ArtifactButton.tsx b/Engine/Source/Programs/Horde/HordeDashboard/src/components/artifacts/ArtifactButton.tsx
--- a/Engine/Source/Programs/Horde/HordeDashboard/src/components/artifacts/ArtifactButton.tsx
+++ b/Engine/Source/Programs/Horde/HordeDashboard/src/components/artifacts/ArtifactButton.tsx
@@ -14,7 +14,7 @@ export const ArtifactButton: React.FC<{ jobId?: string, stepId?: string, artifac
    const [state, setState] = useState<{ showZipDownloadWarning?: boolean }>({});
 
    const navigate = useNavigate();
-   
+
    const path = pathIn ? pathIn + "/" : "";
    const items = selection?.items;
 
@@ -182,18 +182,30 @@ export const ArtifactButton: React.FC<{ jobId?: string, stepId?: string, artifac
                }
             }
          )
-      }
 
-      downloadProps.items.push(
-         {
-            key: 'download_copy_to_clipboard',
-            text: 'Copy Artifact ID',
-            onClick: () => {
-               copyToClipboard(artifact.id);
+         downloadProps.items.push(
+            {
+               key: 'download_copy_to_clipboard',
+               text: 'Copy Artifact ID',
+               onClick: () => {
+                  copyToClipboard(artifact.id);
+               }
             }
+         )
+      } else {
+         const zenUri = artifact?.metadata?.find(m => m.toLowerCase().startsWith("zenbuilduri="))?.slice(12)?.trim();
+         if (zenUri) {
+            downloadProps.items.push(
+               {
+                  key: 'download_copy_zen_uri_to_clipboard',
+                  text: 'Copy Cloud Artifact URL',
+                  onClick: () => {
+                     copyToClipboard(zenUri);
+                  }
+               }
+            )
          }
-      )
-
+      }
    }
 
    function showDownloadWarning() {
