Horde: Updating schema documentation

#rnx

diff --git a/Engine/Source/Programs/Horde/Docs/Deployment/AgentSettings.md b/Engine/Source/Programs/Horde/Docs/Deployment/AgentSettings.md
--- a/Engine/Source/Programs/Horde/Docs/Deployment/AgentSettings.md
+++ b/Engine/Source/Programs/Horde/Docs/Deployment/AgentSettings.md
@@ -9,6 +9,7 @@ Name | Description
 `serverProfiles` | `string` `->` [ServerProfile](#serverprofile)<br>Known servers to connect to
 `server` | `string`<br>The default server, unless overridden from the command line
 `name` | `string`<br>Name of agent to report as when connecting to server. By default, the computer's hostname will be used.
+`mode` | [AgentMode](#agentmode-enum)<br>Mode of operation for the agent - For trusted agents in controlled environments (e.g., build farms). These agents handle all lease types and run exclusively Horde workloads.<br>- For low-trust workstations, uses interactive authentication (human logs in). These agents yield to non-Horde workloads and only support compute leases for remote execution.
 `installed` | `boolean`<br>Whether the server is running in 'installed' mode. In this mode, on Windows, the default data directory will use the common application data folder (C:\ProgramData\Epic\Horde), and configuration data will be read from here and the registry. This setting is overridden to false for local builds from appsettings.Local.json.
 `ephemeral` | `boolean`<br>Whether agent should register as being ephemeral. Doing so will not persist any long-lived data on the server and once disconnected it's assumed to have been deleted permanently. Ideal for short-lived agents, such as spot instances on AWS EC2.
 `workingDir` | [DirectoryReference](#directoryreference)<br>Working directory for leases and jobs (i.e where files from Perforce will be checked out)
@@ -42,9 +43,19 @@ Name | Description
 `environment` | `string`<br>Name of the environment (currently just used for tracing)
 `url` | `string`<br>Url of the server
 `token` | `string`<br>Bearer token to use to initiate the connection
+`useInteractiveAuth` | `boolean`<br>Whether to authenticate interactively in a desktop environment (for example, when agent is running on a user's workstation)
 `thumbprint` | `string`<br>Thumbprint of a certificate to trust. Allows using self-signed certs for the server.
 `thumbprints` | `string[]`<br>Thumbprints of certificates to trust. Allows using self-signed certs for the server.
 
+## AgentMode (Enum)
+
+Defines the operation mode of the agent Duplicated to prevent depending on Protobuf structures (weak name reference when deserializing JSON)
+
+Name | Description
+---- | -----------
+`Dedicated` | 
+`Workstation` | 
+
 ## DirectoryReference
 
 Representation of an absolute directory path. Allows fast hashing and comparisons.
