Horde: Add logging to debug job step warning/error notifcations coming through

diff --git a/Engine/Source/Programs/Horde/Plugins/Build/HordeServer.Build/Notifications/NotificationService.cs b/Engine/Source/Programs/Horde/Plugins/Build/HordeServer.Build/Notifications/NotificationService.cs
--- a/Engine/Source/Programs/Horde/Plugins/Build/HordeServer.Build/Notifications/NotificationService.cs
+++ b/Engine/Source/Programs/Horde/Plugins/Build/HordeServer.Build/Notifications/NotificationService.cs
@@ -662,17 +662,29 @@ namespace HordeServer.Notifications
 
 			// If this is not a success notification and the author isn't in the list to notify, add them manually if this is the outcome has gotten worse.
 			int failures = job.Batches.Sum(x => x.Steps.Count(y => y.Outcome == JobStepOutcome.Failure));
+			int warnings = job.Batches.Sum(x => x.Steps.Count(y => y.Outcome == JobStepOutcome.Warnings));
 			bool firstFailure = step.Outcome == JobStepOutcome.Failure && failures == 1;
-			bool firstWarning = step.Outcome == JobStepOutcome.Warnings && failures == 0 && job.Batches.Sum(x => x.Steps.Count(y => y.Outcome == JobStepOutcome.Warnings)) == 1;
-			if (job.StartedByUserId.HasValue && !usersToNotify.Any(x => x.Id == job.StartedByUserId) && (firstFailure || firstWarning))
+			bool firstWarning = step.Outcome == JobStepOutcome.Warnings && failures == 0 && warnings == 1;
+			if (job.StartedByUserId.HasValue && !usersToNotify.Any(x => x.Id == job.StartedByUserId))
 			{
-				_logger.LogInformation("Author {AuthorUserId} is not in notify list but step outcome is {JobStepOutcome}, adding them to the list...", job.StartedByUserId, step.Outcome);
-				IUser? authorUser = await _userCollection.GetUserAsync(job.StartedByUserId.Value, cancellationToken);
-				if (authorUser != null)
+				if (firstFailure || firstWarning)
+				{					
+					IUser? authorUser = await _userCollection.GetUserAsync(job.StartedByUserId.Value, cancellationToken);
+					if (authorUser != null)
+					{
+						_logger.LogInformation("Author {AuthorUserId} is not in notify list but step outcome is {JobStepOutcome}, adding them to the list...", job.StartedByUserId, step.Outcome);
+						List<IUser> newUsersToNotify = usersToNotify.ToList();
+						newUsersToNotify.Add(authorUser);
+						usersToNotify = newUsersToNotify;
+					}
+					else
+					{
+						_logger.LogWarning("Author {AuthorUserId} is not in notify list, step outcome is {JobStepOutcome}, though unable to find in user collection", job.StartedByUserId, step.Outcome);
+					}
+				}
+				else
 				{
-					List<IUser> newUsersToNotify = usersToNotify.ToList();
-					newUsersToNotify.Add(authorUser);
-					usersToNotify = newUsersToNotify;
+					_logger.LogInformation("Author {AuthorUserId} is not in notify list, step outcome is {JobStepOutcome}, not adding to list with failure count: {JobStepFailures} and warning count: {JobStepWarnings}...", job.StartedByUserId, step.Outcome, failures, warnings);
 				}
 			}
 
