Horde issues now support the annotations AutoAssignToUser as a fallback to AutoAssign

diff --git a/Engine/Source/Programs/Horde/Plugins/Build/HordeServer.Build/Issues/IssueCollection.cs b/Engine/Source/Programs/Horde/Plugins/Build/HordeServer.Build/Issues/IssueCollection.cs
--- a/Engine/Source/Programs/Horde/Plugins/Build/HordeServer.Build/Issues/IssueCollection.cs
+++ b/Engine/Source/Programs/Horde/Plugins/Build/HordeServer.Build/Issues/IssueCollection.cs
@@ -1470,30 +1470,37 @@ namespace HordeServer.Issues
 				newPromoted = newSpans.Any(x => ((IIssueSpan)x).PromoteByDefault);
 			}
 
-			// Find the default owner
+			// Find an owner automatically if annotations requires it
 			UserId? newDefaultOwnerId = null;
-			string? autoAssignToUser = newSpans.Select(x => x.LastFailure.Annotations?.AutoAssignToUser).Where(x => x != null).FirstOrDefault();
-			if (autoAssignToUser != null)
+			bool tryAutoAssign = newSpans.Any(x => x.LastFailure.Annotations?.AutoAssign ?? false);
+			if (tryAutoAssign)
 			{
-				IUser? user = await _userCollection.FindUserByLoginAsync(autoAssignToUser, cancellationToken);
-				if (user != null)
-				{
-					newDefaultOwnerId = user.Id;
-				}
+				newDefaultOwnerId = TryAutoAssignFromSuspects(newSuspectImpls);
 			}
-			else
+			// Find the default owner (as a fallback) if annotations provides it
+			// This behavior has a lower priority than autoAssign to allow using both (in which case, autoAssignToUser is a fallback)
+			if (newDefaultOwnerId == null)
 			{
-				// Figure out if we can auto-assign an owner
-				bool canAutoAssign = newPromoted || newSpans.Any(x => x.LastFailure.Annotations?.AutoAssign ?? false);
-				if (canAutoAssign && newSuspectImpls.Count > 0)
+				string? autoAssignToUser = newSpans.Select(x => x.LastFailure.Annotations?.AutoAssignToUser).Where(x => x != null).FirstOrDefault();
+				if (autoAssignToUser != null)
 				{
-					UserId possibleOwnerId = newSuspectImpls[0].AuthorId;
-					if (newSuspectImpls.All(x => x.AuthorId == possibleOwnerId) && newSuspectImpls.Any(x => x.DeclinedAt == null))
+					IUser? user = await _userCollection.FindUserByLoginAsync(autoAssignToUser, cancellationToken);
+					if (user != null)
 					{
-						newDefaultOwnerId = possibleOwnerId;
+						newDefaultOwnerId = user.Id;
 					}
 				}
 			}
+			// Find an owner automatically if this issue is promoted (and we did not try already to AutoAssign)
+			// This behavior has a lower priority than autoAssignToUser to maintain backward compatibility
+			if (newDefaultOwnerId == null)
+			{
+				tryAutoAssign = !tryAutoAssign && newPromoted;
+				if (tryAutoAssign)
+				{
+					newDefaultOwnerId = TryAutoAssignFromSuspects(newSuspectImpls);
+				}
+			}
 
 			// Perform the actual update with this data
 			List<UpdateDefinition<IssueDocument>> updates = new List<UpdateDefinition<IssueDocument>>();
@@ -1579,6 +1586,19 @@ namespace HordeServer.Issues
 			return newSuspectImpls;
 		}
 
+		private static UserId? TryAutoAssignFromSuspects(List<IssueSuspectDocument> suspects)
+		{
+			if (suspects.Count > 0)
+			{
+				UserId possibleOwnerId = suspects[0].AuthorId;
+				if (suspects.All(x => x.AuthorId == possibleOwnerId) && suspects.Any(x => x.DeclinedAt == null))
+				{
+					return possibleOwnerId;
+				}
+			}
+			return null;
+		}
+
 		#endregion
 
 		#region Spans
