Horde: Add exception handling when reporting job metric so that it doesn't suppress notifications
#submitool bypass

diff --git a/Engine/Source/Programs/Horde/Plugins/Build/HordeServer.Build/Notifications/NotificationService.cs b/Engine/Source/Programs/Horde/Plugins/Build/HordeServer.Build/Notifications/NotificationService.cs
--- a/Engine/Source/Programs/Horde/Plugins/Build/HordeServer.Build/Notifications/NotificationService.cs
+++ b/Engine/Source/Programs/Horde/Plugins/Build/HordeServer.Build/Notifications/NotificationService.cs
@@ -579,10 +579,17 @@ namespace HordeServer.Notifications
 				_jobDurationHistogram.Record(duration.TotalSeconds, tags);
 			}
 
-			JobStepOutcome jobOutcome = job.Batches.SelectMany(x => x.Steps).Min(x => x.Outcome);
-			DateTime? startTime = job.Batches.Select(x => x.StartTimeUtc).Min();
-			DateTime? finishTime = job.Batches.Select(x => x.FinishTimeUtc).Max();
-			RecordMetric("job", jobOutcome, startTime, finishTime);
+			try
+			{
+				JobStepOutcome jobOutcome = job.Batches.SelectMany(x => x.Steps).Min(x => x.Outcome);
+				DateTime? startTime = job.Batches.Select(x => x.StartTimeUtc).Min();
+				DateTime? finishTime = job.Batches.Select(x => x.FinishTimeUtc).Max();
+				RecordMetric("job", jobOutcome, startTime, finishTime);
+			}
+			catch (Exception e)
+			{
+				_logger.LogWarning(e, "Failed to Record job metric for {JobId}", job.Id);
+			}
 
 			// TODO: record metrics for individual steps
 			// foreach (IJobStepBatch Batch in Job.Batches)
