Horde : Add docs on how to create an external secret provider

diff --git a/Engine/Source/Programs/Horde/Docs/Config/Secrets.md b/Engine/Source/Programs/Horde/Docs/Config/Secrets.md
--- a/Engine/Source/Programs/Horde/Docs/Config/Secrets.md
+++ b/Engine/Source/Programs/Horde/Docs/Config/Secrets.md
@@ -76,7 +76,7 @@ An example secret may be configured as follows:
 
 Secrets are queried from the external provider when requested by a user, and are not cached by Horde.
 
-Secrets from external providers that require additional options, for example HashiCorp Vault, for example to set the endpoint and
+Secrets from external providers that require additional options, for example HashiCorp Vault to set the endpoint and
 credentials, will use a `providerConfig` property instead of `name` property to provide the data. For example to fetch
 three secrets from HashiCorp Vault using different authentication methods:
 
@@ -134,7 +134,7 @@ three secrets from HashiCorp Vault using different authentication methods:
                 {
                     "name": "hcp-vault-aws-assume-role",
                     "provider": "HcpVault",
-                    "hcpCault": {
+                    "hcpVault": {
                         "credentials": "awsauth",
                         "endpoint": "https://vault.example.com",
                         "awsIamServerId": "vault.example.com",
@@ -211,3 +211,120 @@ utility methods for querying them at runtime:
    ```csharp
    string propertyValue = await CommandUtils.GetHordeSecretAsync(new SecretId("my-secret-name"), "propertyName")
    ```
+
+## Implement a custom external provider
+
+To implement a custom external provider instead of using AWS or HCP Vault. Create a file in the folder `Engine\Source\Programs\Horde\Plugins\Secrets\HordeServer.Secrets\Secrets\Providers` with a suitable name for the provider e.g. `MySecretProvider.cs` that provides an implementation of `HordeServer.Secrets.ISecretProvider`.
+
+   ```csharp
+    namespace HordeServer.Secrets.Providers
+    {
+        public class MySecretProvider : ISecretProvider
+        {
+            // The name must be unique across all providers
+            public string Name => "MySecretStore";
+
+            // If access to the provider requires a HTTP client
+            private readonly IHttpClientFactory _httpClientFactory;
+
+            private readonly ILogger _logger;
+
+            public MySecretProvider(IHttpClientFactory httpClientFactory, ILogger<HcpVaultSecretProvider> logger)
+            {
+                 _httpClientFactory = httpClientFactory;
+                 _logger = logger;
+            }
+
+            private HttpClient GetHttpClient()
+            {
+                return _httpClientFactory.CreateClient("HordeSecretMySecretProvider");
+            }
+
+            public async Task<string> GetSecretAsync(string path, SecretProviderConfig? config, CancellationToken cancellationToken)
+            {
+                string secretValue;
+                // Implement how to get the secret for the give 'path' from your provider
+                return secretValue;
+            }
+        }
+    }
+   ```
+
+Add the provider to the service collection in the `ConfigureServices` method found in `Engine\Source\Programs\Horde\Plugins\Secrets\HordeServer.Secrets\SecretsPlugin.cs`.
+
+   ```csharp
+    public void ConfigureServices(IServiceCollection services)
+    {
+        // ...
+
+        services.AddSingleton<ISecretProvider, MySecretProvider>();
+    }
+   ```
+
+To test on a local workstation unit tests can be added to `Engine\Source\Programs\Horde\Plugins\Secrets\HordeServer.Secrets.Tests\Secrets\SecretProviderTests.cs`.
+
+Edit the `globals.json` file to add a secret to use the new provider as shown in the earlier examples. In the examples where there is `"provider": "AwsParameterStore"` use `"provider": "MySecretStore"` instead.
+
+Rebuild and deploy the Horde server. Once deployed can use [Swagger](../Internals/RestApi.md) and the endpoint `/api/v1/secrets/{secretId}` to get a secret declared in `globals.json`.
+
+### Add additional configuration for the external provider
+
+If the external provider requires additional configuration for example different secrets are stored at different domains or an access token is required then a provider configuarion can be implemented.
+
+Create a file in the folder `Engine\Source\Programs\Horde\Plugins\Secrets\HordeServer.Secrets\Secrets\Providers` with a suitable name for the provider configuration e.g. `MySecretProviderConfig.cs`.
+
+   ```csharp
+    namespace HordeServer.Secrets.Providers
+    {
+        public class MySecretProviderConfig
+        {
+            public string? DomainName { get; set; }
+            public string? LoginToken { get; set; }
+        }
+    }
+   ```
+
+Extend the class `HordeServer.Secrets.SecretProviderConfig` in `Engine\Source\Programs\Horde\Plugins\Secrets\HordeServer.Secrets\Secrets\SecretProviderConfig.cs` to
+have a `MySecretProviderConfig` member.
+
+   ```csharp
+    public class SecretProviderConfig
+    {
+        // ...
+
+        public MySecretProviderConfig? MyConfig { get; set; }
+    }
+   ```
+
+Edit the `GetSecretAsync` method created earlier to use this new configuration to get the secret. For example if the configuraion was to be used in a HTTP GET request it might look something like this
+
+   ```csharp
+    public async Task<string> GetSecretAsync(string path, SecretProviderConfig? config, CancellationToken cancellationToken)
+    {
+        if (config?.MyConfig == null)
+        {
+            throw new InvalidOperationException($"Unable to fetch secret {path} from My Secret Provider (No Config)");
+        }
+        if (config.MyConfig.DomainName == null)
+        {
+            throw new InvalidOperationException($"Unable to fetch secret {path} from My Secret Provider (No Domain Name)");
+        }
+        if (config.MyConfig.LoginToken == null)
+        {
+            throw new InvalidOperationException($"Unable to fetch secret {path} from My Secret Provider (No Login Token)");
+        }
+
+        Uri uri = new("https://" + config.MyConfig.DomainName + ".example.com", path);
+        using HttpRequestMessage request = new(HttpMethod.Get, uri);
+        HttpClient httpClient = GetHttpClient();
+        httpClient.DefaultRequestHeaders.Add("X-Token", config.MyConfig.LoginToken);
+        using HttpResponseMessage httpResponseMessage = await httpClient.SendAsync(request, cancellationToken);
+
+        string secretValue;
+        // Implement how to get the secret from the repsone message
+        return secretValue;
+    }
+   ```
+
+Edit the `globals.json` file and add the configuration options to `plugins.secrets.providerConfigs` with an object for each combination of configuration options.
+For the secrets replace `"provider": "MySecretStore"` with the key `"providerConfig"` and its value the unique name of the object in `plugins.secrets.providerConfigs`. See the earlier examples of configuring HCP Vault for more details.
