Horde: Trim search text to avoid searching log for empty strings
#rnx

diff --git a/Engine/Source/Programs/Horde/HordeDashboard/src/components/LogView.tsx b/Engine/Source/Programs/Horde/HordeDashboard/src/components/LogView.tsx
--- a/Engine/Source/Programs/Horde/HordeDashboard/src/components/LogView.tsx
+++ b/Engine/Source/Programs/Horde/HordeDashboard/src/components/LogView.tsx
@@ -1024,7 +1024,7 @@ export const LogList: React.FC<{ logId: string }> = observer(({ logId }) => {
          currentLine = handler.currentLine
       }
 
-      if (!newValue) {
+      if (!newValue.trim()) {
 
          logListKey++;
 
