Horde: Remove Microsoft.Management.Infrastructure from agent telemetry

Removes the NuGet dependency for both agent and JobDriver. Using CIM is slow and also adds dependencies incompatible with Wine.

diff --git a/Engine/Source/Programs/Horde/Drivers/JobDriver/JobDriver.csproj b/Engine/Source/Programs/Horde/Drivers/JobDriver/JobDriver.csproj
--- a/Engine/Source/Programs/Horde/Drivers/JobDriver/JobDriver.csproj
+++ b/Engine/Source/Programs/Horde/Drivers/JobDriver/JobDriver.csproj
@@ -73,7 +73,6 @@
     <PackageReference Include="Microsoft.Extensions.Http.Polly" Version="8.0.10" />
     <PackageReference Include="Microsoft.Extensions.Logging.Abstractions" Version="8.0.2" />
     <PackageReference Include="Microsoft.Extensions.Options.DataAnnotations" Version="8.0.0" />
-    <PackageReference Include="Microsoft.Management.Infrastructure" Version="3.0.0" />
     <PackageReference Include="Microsoft.VisualStudio.Threading.Analyzers" Version="17.11.20">
       <PrivateAssets>all</PrivateAssets>
       <IncludeAssets>runtime; build; native; contentfiles; analyzers; buildtransitive</IncludeAssets>
diff --git a/Engine/Source/Programs/Horde/HordeAgent/HordeAgent.csproj b/Engine/Source/Programs/Horde/HordeAgent/HordeAgent.csproj
--- a/Engine/Source/Programs/Horde/HordeAgent/HordeAgent.csproj
+++ b/Engine/Source/Programs/Horde/HordeAgent/HordeAgent.csproj
@@ -55,7 +55,6 @@
     <PackageReference Include="Microsoft.Extensions.Http.Polly" Version="8.0.10" />
     <PackageReference Include="Microsoft.Extensions.Logging.Abstractions" Version="8.0.2" />
     <PackageReference Include="Microsoft.Extensions.Options.DataAnnotations" Version="8.0.0" />
-    <PackageReference Include="Microsoft.Management.Infrastructure" Version="3.0.0" />
     <PackageReference Include="Microsoft.VisualStudio.Threading.Analyzers" Version="17.11.20">
       <PrivateAssets>all</PrivateAssets>
       <IncludeAssets>runtime; build; native; contentfiles; analyzers; buildtransitive</IncludeAssets>
diff --git a/Engine/Source/Programs/Horde/HordeAgent/Services/TelemetryService.cs b/Engine/Source/Programs/Horde/HordeAgent/Services/TelemetryService.cs
--- a/Engine/Source/Programs/Horde/HordeAgent/Services/TelemetryService.cs
+++ b/Engine/Source/Programs/Horde/HordeAgent/Services/TelemetryService.cs
@@ -8,7 +8,6 @@ using System.Text;
 using EpicGames.Core;
 using Microsoft.Extensions.Hosting;
 using Microsoft.Extensions.Logging;
-using Microsoft.Management.Infrastructure;
 
 namespace HordeAgent.Services;
 
@@ -231,23 +230,7 @@ public sealed class WindowsSystemMetrics : ISystemMetrics, IDisposable
 
 	private static uint GetPhysicalMemory()
 	{
-		using CimSession session = CimSession.Create(null);
-		const string QueryNamespace = @"root\cimv2";
-		const string QueryDialect = "WQL";
-		ulong totalCapacity = 0;
-
-		foreach (CimInstance instance in session.QueryInstances(QueryNamespace, QueryDialect, "select Capacity from Win32_PhysicalMemory"))
-		{
-			foreach (CimProperty property in instance.CimInstanceProperties)
-			{
-				if (property.Name.Equals("Capacity", StringComparison.OrdinalIgnoreCase) && property.Value is ulong capacity)
-				{
-					totalCapacity += capacity;
-				}
-			}
-		}
-
-		return (uint)(totalCapacity / 1024); // as kibibytes
+		return (uint)(WindowsCapabilities.GetPhysicalMemory() / 1024);
 	}
 }
 
diff --git a/Engine/Source/Programs/Horde/HordeAgent/Services/WindowsCapabilities.cs b/Engine/Source/Programs/Horde/HordeAgent/Services/WindowsCapabilities.cs
--- a/Engine/Source/Programs/Horde/HordeAgent/Services/WindowsCapabilities.cs
+++ b/Engine/Source/Programs/Horde/HordeAgent/Services/WindowsCapabilities.cs
@@ -68,6 +68,27 @@ public class WindowsCapabilities(IOptionsMonitor<AgentSettings> settings, ILogge
 	private static void SetProp(RpcAgentCapabilities caps, string key, bool value) => caps.Properties.Add($"{key}={value}");
 	private static void SetResource(RpcAgentCapabilities caps, string key, int value) => caps.Resources.Add(key, value);
 	
+	/// <summary>
+	/// Gets the total amount of physical memory installed on the system.
+	/// </summary>
+	/// <returns>The total physical memory in bytes. Returns 0 if memory status could not be retrieved</returns>
+	public static int GetPhysicalMemory()
+	{
+		try
+		{
+			MEMORYSTATUSEX memStatus = new();
+			if (GlobalMemoryStatusEx(memStatus))
+			{
+				return (int)memStatus.ullTotalPhys;
+			}
+		}
+		catch
+		{
+			// Ignore any error
+		}
+		return 0;
+	}
+	
 	private static int GetNumLogicalCores()
 	{
 		int total = 0;
