Horde Server: Downgrade issue escalation channel not existing to a warning

diff --git a/Engine/Source/Programs/Horde/Plugins/Build/HordeServer.Build/Notifications/Sinks/SlackNotificationSink.cs b/Engine/Source/Programs/Horde/Plugins/Build/HordeServer.Build/Notifications/Sinks/SlackNotificationSink.cs
--- a/Engine/Source/Programs/Horde/Plugins/Build/HordeServer.Build/Notifications/Sinks/SlackNotificationSink.cs
+++ b/Engine/Source/Programs/Horde/Plugins/Build/HordeServer.Build/Notifications/Sinks/SlackNotificationSink.cs
@@ -2858,7 +2858,15 @@ namespace HordeServer.Notifications.Sinks
 					}
 					catch (SlackException ex)
 					{
-						_logger.LogError(ex, "Slack exception while escalating issue {IssueId}; cancelling.", issueId);
+						if (ex.Message.Equals("channel_not_found", StringComparison.Ordinal))
+						{
+							_logger.LogWarning(ex, "Slack channel not found while escalating issue {IssueId}; cancelling.", issueId);
+						}
+						else
+						{
+							_logger.LogError(ex, "Slack exception while escalating issue {IssueId}; cancelling.", issueId);
+						}
+							
 						await _redisService.GetDatabase().SortedSetRemoveAsync(s_escalateIssues, issueId);
 					}
 				}
