Horde: Downgrade unknown batch error due to missing shelved CL to a warning

diff --git a/Engine/Source/Programs/Horde/Plugins/Build/HordeServer.Build/Jobs/JobService.cs b/Engine/Source/Programs/Horde/Plugins/Build/HordeServer.Build/Jobs/JobService.cs
--- a/Engine/Source/Programs/Horde/Plugins/Build/HordeServer.Build/Jobs/JobService.cs
+++ b/Engine/Source/Programs/Horde/Plugins/Build/HordeServer.Build/Jobs/JobService.cs
@@ -785,7 +785,14 @@ namespace HordeServer.Jobs
 						if (error != null)
 						{
 							newError = error.Value;
-							_logger.LogError("Batch {BatchId} marked with error {JobStepBatchError}", batchId, error.Value.ToString());
+							if (newError == JobStepBatchError.UnknownShelf)
+							{
+								_logger.LogWarning("Batch {BatchId} marked with error, Unknown Shelved CL: {PreflightCommitId}", batchId, job.PreflightCommitId?.ToString() ?? "Unknown");
+							}
+							else
+							{
+								_logger.LogError("Batch {BatchId} marked with error {JobStepBatchError}", batchId, error.Value.ToString());
+							}								
 							_jobs.JobAudit(job.Id, "Batch {BatchId} marked with error {JobStepBatchError}", batchId, error.Value.ToString());
 						}
 						else if (!batch.Steps.Any(x => x.StartTimeUtc != null))
