Horde Storage: Do not throw errors for duplicate keys on shadow blobs as these are expected

diff --git a/Engine/Source/Programs/Horde/Plugins/Storage/HordeServer.Storage/Storage/StorageService.cs b/Engine/Source/Programs/Horde/Plugins/Storage/HordeServer.Storage/Storage/StorageService.cs
--- a/Engine/Source/Programs/Horde/Plugins/Storage/HordeServer.Storage/Storage/StorageService.cs
+++ b/Engine/Source/Programs/Horde/Plugins/Storage/HordeServer.Storage/Storage/StorageService.cs
@@ -864,7 +864,20 @@ namespace HordeServer.Storage
 
 						writes.Add(new UpdateOneModel<BlobInfo>(filter, update) { IsUpsert = true });
 					}
-					await _blobCollection.BulkWriteAsync(writes, new BulkWriteOptions { IsOrdered = false }, cancellationToken);
+					try
+					{
+						await _blobCollection.BulkWriteAsync(writes, new BulkWriteOptions { IsOrdered = false }, cancellationToken);
+					}
+					catch (MongoBulkWriteException ex)
+					{
+						foreach (WriteError writeError in ex.WriteErrors)
+						{
+							if (writeError.Category != ServerErrorCategory.DuplicateKey)
+							{
+								throw;
+							}
+						}
+					}
 				}
 
 				// Find all the matching import ids in order
