Added IssueHandlersTests to Horde's unittests (starting with a LocalizationIssueHandler test)

diff --git a/Engine/Source/Programs/Horde/Plugins/Build/HordeServer.Build.Tests/Issues/IssueHandlerTests.cs b/Engine/Source/Programs/Horde/Plugins/Build/HordeServer.Build.Tests/Issues/IssueHandlerTests.cs
new file mode 100644
--- /dev/null
+++ b/Engine/Source/Programs/Horde/Plugins/Build/HordeServer.Build.Tests/Issues/IssueHandlerTests.cs
@@ -0,0 +1,178 @@
+// Copyright Epic Games, Inc. All Rights Reserved.
+
+using System.Text;
+using EpicGames.Core;
+using EpicGames.Horde.Issues;
+using EpicGames.Horde.Issues.Handlers;
+using EpicGames.Horde.Jobs.Templates;
+using EpicGames.Horde.Streams;
+
+namespace HordeServer.Tests.Issues
+{
+	[TestClass]
+	public class IssueHandlerTests
+	{
+		private class IssueHandlerTestData
+		{
+			public IssueHandlerTestData(List<LogEvent> logEvents, List<IssueEvent> issueEvents) =>
+			(LogEvents_ToDebug, IssueEvents) = (logEvents, issueEvents);
+
+			public List<LogEvent> LogEvents_ToDebug { get; init; } // For debug purposes only
+			public List<IssueEvent> IssueEvents { get; init; }
+		}
+
+		// Test a single IssueHandler (issue handler priority is irrelevant in such a test)
+		static List<IssueEventGroup> TestSingleIssueHandler(IssueHandler issueHandler,string[] jsonStructuredLogs)
+		{
+			IssueHandlerTestData issueHandlerParsedTestData = ParseJSONStructuredLogs(jsonStructuredLogs);
+			foreach (IssueEvent issueEvent in issueHandlerParsedTestData.IssueEvents)
+			{
+				issueHandler.HandleEvent(issueEvent);
+			}
+			List<IssueEventGroup> groupedIssues = issueHandler.GetIssues().ToList();
+			return groupedIssues;
+		}
+
+		// Parse structured logs to create IssueEvents (and LogEvents for debug purposes)
+		static IssueHandlerTestData ParseJSONStructuredLogs(string[] jsonStructuredLogs)
+		{
+			List<LogEvent> logEvents = new List<LogEvent>(); // For debug purposes only
+			List<IssueEvent> issueEvents = new List<IssueEvent>();
+			foreach (string jsonStructuredLog in jsonStructuredLogs)
+			{
+				byte[] byteArrayStructuredLog = Encoding.UTF8.GetBytes(jsonStructuredLog);
+
+				logEvents.Add(LogEvent.Read(byteArrayStructuredLog)); // For debug purposes only
+
+				JsonLogEvent jsonLogEventFromStructuredLog;
+				Assert.IsTrue(JsonLogEvent.TryParse(byteArrayStructuredLog, out jsonLogEventFromStructuredLog));
+
+				List<JsonLogEvent> listOfJsonLofEvents = new List<JsonLogEvent>();
+				listOfJsonLofEvents.Add(jsonLogEventFromStructuredLog);
+				issueEvents.Add(new IssueEvent(jsonLogEventFromStructuredLog.LineIndex, jsonLogEventFromStructuredLog.Level, jsonLogEventFromStructuredLog.EventId, listOfJsonLofEvents));
+			}
+
+			return new IssueHandlerTestData(logEvents, issueEvents);
+		}
+
+		[TestMethod]
+		public void LocalizationIssueHandler()
+		{
+			// Logs
+			// 1) Taken from an actual JSON log file
+			// 2) Data Obfuscated (modified any data to a plausible yet fake data)
+			// 3) Formatted for the test by transforming any " into \"
+			string[] jsonStructuredLogs =
+			{
+				"{ \"time\":\"2025-01-01T01:20:30.990Z\",\"level\":\"Warning\",\"message\":\"LogGatherTextFromAssetsCommandlet: Warning: Package '/Path/To/File1' and '/DifferentPathTo/File2' have the same localization ID (ABCDEFGHIJKLMNOPQRSTUVWXYZABCDEF). Please reset one of these (Asset Localization -> Reset Localization ID) to avoid conflicts.\",\"format\":\"{_channel}: {_severity}: Package '{file}' and '{conflictFile}' have the same localization ID ({locKey}). Please reset one of these (Asset Localization -> Reset Localization ID) to avoid conflicts.\",\"properties\":{ \"_channel\":{ \"$type\":\"Channel\",\"$text\":\"LogGatherTextFromAssetsCommandlet\"},\"_severity\":{ \"$type\":\"Severity\",\"$text\":\"Warning\"},\"id\":304,\"file\":\"/Path/To/File1\",\"conflictFile\":\"/DifferentPathTo/File2\",\"locKey\":\"ABCDEFGHIJKLMNOPQRSTUVWXYZABCDEF\"} }",
+				"{ \"time\":\"2025-01-01T01:20:30.991Z\",\"level\":\"Warning\",\"message\":\"LogGatherTextFromAssetsCommandlet: Warning: Package '/Path/To/File3' and '/DifferentPathTo/File4' have the same localization ID (ABCDEFGHIJKLMNOPQRSTUVWXYZABCDEG). Please reset one of these (Asset Localization -> Reset Localization ID) to avoid conflicts.\",\"properties\":{ \"_channel\":{ \"$type\":\"Channel\",\"$text\":\"LogGatherTextFromAssetsCommandlet\"},\"_severity\":{ \"$type\":\"Severity\",\"$text\":\"Warning\"},\"id\":304,\"file\":\"/Path/To/File3\",\"conflictFile\":\"/DifferentPathTo/File4\",\"locKey\":\"ABCDEFGHIJKLMNOPQRSTUVWXYZABCDEG\"},\"format\":\"{_channel}: {_severity}: Package '{file}' and '{conflictFile}' have the same localization ID ({locKey}). Please reset one of these (Asset Localization -> Reset Localization ID) to avoid conflicts.\"}",
+				"{ \"time\":\"2025-01-01T01:20:30.992Z\",\"level\":\"Warning\",\"message\":\"LogLocTextHelper: Warning: Text conflict from LOCTEXT macro for namespace 'NamespaceX' and key 'KeyY'. First entry is Plugins/Path/To/File5.cpp(50):'Text from line 50'. Conflicting entry is Plugins/Path/To/File5.cpp(51):'Text from line 51'.\",\"properties\":{ \"_channel\":{ \"$type\":\"Channel\",\"$text\":\"LogLocTextHelper\"},\"_severity\":{ \"$type\":\"Severity\",\"$text\":\"Warning\"},\"from\":\" from LOCTEXT\",\"locNamespace\":\"NamespaceX\",\"locKey\":\"KeyY\",\"locKeyMetaData\":\"\",\"location\":\"Plugins/Path/To/File5.cpp(50)\",\"text\":\"Text from line 50\",\"textMetaData\":\"\",\"conflictLocation\":\"Plugins/Path/To/File5.cpp(51)\",\"conflictText\":\"Text from line 51\",\"conflictTextMetaData\":\"\",\"id\":304},\"format\":\"{_channel}: {_severity}: Text conflict{from} for namespace '{locNamespace}' and key '{locKey}'{locKeyMetaData}. First entry is {location}:'{text}'{textMetaData}. Conflicting entry is {conflictLocation}:'{conflictText}'{conflictTextMetaData}.\"}",
+				"{ \"time\":\"2025-01-01T01:20:30.993Z\",\"level\":\"Warning\",\"message\":\"LogLocTextHelper: Warning: Text conflict for namespace 'NamespaceA' and key 'KeyA'. Main manifest entry is 'Text from manifest'. Conflict from manifest dependency is Path/To/ManifestFile.manifest(10):'Conflict manifest text'.\",\"properties\":{ \"_channel\":{ \"$type\":\"Channel\",\"$text\":\"LogLocTextHelper\"},\"_severity\":{ \"$type\":\"Severity\",\"$text\":\"Warning\"},\"locNamespace\":\"NamespaceA\",\"locKey\":\"KeyA\",\"locKeyMetaData\":\"\",\"location\":\"\",\"text\":\"Text from manifest\",\"textMetaData\":\"\",\"conflictLocation\":\"Path/To/ManifestFile.manifest(10)\",\"conflictText\":\"Conflict manifest text\",\"conflictTextMetaData\":\"\",\"id\":304},\"format\":\"{_channel}: {_severity}: Text conflict for namespace '{locNamespace}' and key '{locKey}'{locKeyMetaData}. Main manifest entry is '{text}'{textMetaData}. Conflict from manifest dependency is {conflictLocation}:'{conflictText}'{conflictTextMetaData}.\"}",
+				"{ \"time\":\"2025-01-01T01:20:30.994Z\",\"level\":\"Warning\",\"message\":\"LogLocTextHelper: Warning: Text conflict from LOCTEXT macro for namespace 'NamespaceX' and key 'KeyY'. First entry is Plugins/Path/To/Asset.SubComponent.Whatever.Else:'Text from asset component'. Conflicting entry is Plugins/Path/To/File5.cpp(51):'Text from line 51'.\",\"properties\":{ \"_channel\":{ \"$type\":\"Channel\",\"$text\":\"LogLocTextHelper\"},\"_severity\":{ \"$type\":\"Severity\",\"$text\":\"Warning\"},\"from\":\" from LOCTEXT\",\"locNamespace\":\"NamespaceX\",\"locKey\":\"KeyY\",\"locKeyMetaData\":\"\",\"location\":\"Plugins/Path/To/Asset.Subcomponent.Whatever.Else\",\"text\":\"Text from asset component\",\"textMetaData\":\"\",\"conflictLocation\":\"Plugins/Path/To/File5.cpp(51)\",\"conflictText\":\"Text from line 51\",\"conflictTextMetaData\":\"\",\"id\":304},\"format\":\"{_channel}: {_severity}: Text conflict{from} for namespace '{locNamespace}' and key '{locKey}'{locKeyMetaData}. First entry is {location}:'{text}'{textMetaData}. Conflicting entry is {conflictLocation}:'{conflictText}'{conflictTextMetaData}.\"}",
+				"{ \"time\":\"2025-01-01T01:20:30.995Z\",\"level\":\"Warning\",\"message\":\"LogLocTextHelper: Warning: Broken Rich Text Tag detected in a translation. An unbalanced tag (a complete/incomplet opening rich text tag (i.e. <TagName>) with an incomplet/complete closing tag(</>)) was detected in a translation but not in its source text. Find the problematic tag in the translation and fix the translation to remove this warning. Translation File:'D:/Path/To/PluginX/Content/Localization/LocTargetX/it/LocTargetX.archive' Namespace And Key:'NamespaceB,KeyB' Translation Text To Fix:'Text with <Broken> Rich Text Tags />.'.\",\"properties\":{ \"_channel\":{ \"$type\":\"Channel\",\"$text\":\"LogLocTextHelper\"},\"_severity\":{ \"$type\":\"Severity\",\"$text\":\"Warning\"},\"cultudeCode\":\"it\",\"locNamespace\":\"NamespaceB\",\"locKey\":\"KeyB\",\"text\":\"Text with <Correct>Rich Text Tags</>.\",\"conflictFile\":\"D:/Path/To/PluginX/Content/Localization/LocTargetX/it/LocTargetX.archive\",\"conflictText\":\"Text with <Broken>Rich Text Tags</.\",\"id\":304},\"format\":\"{_channel}: {_severity}: Broken Rich Text Tag detected in a translation. An unbalanced tag (a complete/incomplet opening rich text tag (i.e. <TagName>) with an incomplet/complete closing tag(</>)) was detected in the translation but not in the source text. Find the problematic tag in the translation and fix the translation to remove this warning. Translation File:'{conflictFile}' Namespace And Key:'{locNamespace},{locKey}' Translation Text To Fix:'{conflictText}'.\"}",
+				"{ \"time\":\"2025-01-01T01:20:30.996Z\",\"level\":\"Error\",\"message\":\"LogLocTextHelper: Error: /Path/To/File6.cpp(60): Failed to add text from LOCTEXT for namespace 'NamespaceC' and key 'KeyC' with source 'Text from File6'.\",\"properties\":{ \"_channel\":{ \"$type\":\"Channel\",\"$text\":\"LogLocTextHelper\"},\"_severity\":{ \"$type\":\"Severity\",\"$text\":\"Error\"},\"location\":\"/Path/To/File6.cpp(60)\",\"from\":\" from LOCTEXT\",\"locNamespace\":\"NamespaceC\",\"locKey\":\"KeyC\",\"locKeyMetaData\":\"\",\"text\":\"Text from File6\",\"textMetaData\":\"\",\"id\":304},\"format\":\"{_channel}: {_severity}: {location}: Failed to add text{from} for namespace '{locNamespace}' and key '{locKey}'{locKeyMetaData} with source '{text}'{textMetaData}.\"}",
+				"{ \"time\":\"2025-01-01T01:20:30.997Z\",\"level\":\"Error\",\"message\":\"LogLocTextHelper: Error: Failed to add text for namespace 'NamespaceA' and key 'KeyA' with source 'TextA'.\",\"properties\":{ \"_channel\":{ \"$type\":\"Channel\",\"$text\":\"LogLocTextHelper\"},\"_severity\":{ \"$type\":\"Severity\",\"$text\":\"Error\"},\"locNamespace\":\"NamespaceA\",\"locKey\":\"KeyA\",\"locKeyMetaData\":\"\",\"text\":\"TextA\",\"textMetaData\":\"\",\"forPlatform\":\"\",\"id\":304},\"format\":\"{_channel}: {_severity}: Failed to add text for namespace '{locNamespace}' and key '{locKey}'{locKeyMetaData} with source '{text}'{textMetaData}{forPlatform}.\"}",
+				"{ \"time\":\"2025-01-01T01:20:30.998Z\",\"level\":\"Error\",\"message\":\"LogLocTextHelper: Error: Failed to add text for namespace 'NamespaceB' and key 'KeyB' with source 'TextB'.\",\"properties\":{ \"_channel\":{ \"$type\":\"Channel\",\"$text\":\"LogLocTextHelper\"},\"_severity\":{ \"$type\":\"Severity\",\"$text\":\"Error\"},\"locNamespace\":\"NamespaceB\",\"locKey\":\"KeyB\",\"locKeyMetaData\":\"\",\"text\":\"TextB\",\"textMetaData\":\"\",\"forPlatform\":\"\",\"id\":304},\"format\":\"{_channel}: {_severity}: Failed to add text for namespace '{locNamespace}' and key '{locKey}'{locKeyMetaData} with source '{text}'{textMetaData}{forPlatform}.\"}",
+			};
+
+			IssueHandlerContext context = new IssueHandlerContext(new StreamId(), new TemplateId(), "Localize ProjectX", new Dictionary<string, string>());
+			LocalizationIssueHandler localizationIssueHandler = new LocalizationIssueHandler(context);
+			List<IssueEventGroup> groupedIssues = TestSingleIssueHandler(localizationIssueHandler, jsonStructuredLogs);
+
+			Assert.IsTrue(groupedIssues.Count == 8);
+
+			Assert.IsTrue(groupedIssues[0].Keys.Count == 3);
+			Assert.IsTrue(groupedIssues[0].Keys.Contains(new IssueKey("File1", IssueKeyType.File)));
+			Assert.IsTrue(groupedIssues[0].Keys.Contains(new IssueKey("File2", IssueKeyType.File)));
+			Assert.IsTrue(groupedIssues[0].Keys.Contains(new IssueKey(",ABCDEFGHIJKLMNOPQRSTUVWXYZABCDEF", IssueKeyType.None)));
+			Assert.IsTrue(groupedIssues[0].ChangeFilter == "...");
+			Assert.IsTrue(groupedIssues[0].Events.Count == 1);
+			Assert.IsTrue(groupedIssues[0].Events[0].EventId == KnownLogEvents.Engine_Localization);
+			Assert.IsTrue(groupedIssues[0].Events[0].Severity == Microsoft.Extensions.Logging.LogLevel.Warning);
+			Assert.IsTrue(groupedIssues[0].Metadata.Count == 1);
+			Assert.IsTrue(groupedIssues[0].SummaryTemplate == "{Meta:Node}: Localization {Severity} in {Files}");
+			Assert.IsTrue(groupedIssues[0].Type == "Localization");
+
+			Assert.IsTrue(groupedIssues[1].Keys.Count == 3);
+			Assert.IsTrue(groupedIssues[1].Keys.Contains(new IssueKey("File3", IssueKeyType.File)));
+			Assert.IsTrue(groupedIssues[1].Keys.Contains(new IssueKey("File4", IssueKeyType.File)));
+			Assert.IsTrue(groupedIssues[1].Keys.Contains(new IssueKey(",ABCDEFGHIJKLMNOPQRSTUVWXYZABCDEG", IssueKeyType.None)));
+			Assert.IsTrue(groupedIssues[1].ChangeFilter == "...");
+			Assert.IsTrue(groupedIssues[1].Events.Count == 1);
+			Assert.IsTrue(groupedIssues[1].Events[0].EventId == KnownLogEvents.Engine_Localization);
+			Assert.IsTrue(groupedIssues[1].Events[0].Severity == Microsoft.Extensions.Logging.LogLevel.Warning);
+			Assert.IsTrue(groupedIssues[1].Metadata.Count == 1);
+			Assert.IsTrue(groupedIssues[1].SummaryTemplate == "{Meta:Node}: Localization {Severity} in {Files}");
+			Assert.IsTrue(groupedIssues[1].Type == "Localization");
+
+			Assert.IsTrue(groupedIssues[2].Keys.Count == 2);
+			Assert.IsTrue(groupedIssues[2].Keys.Contains(new IssueKey("File5", IssueKeyType.File)));
+			Assert.IsTrue(groupedIssues[2].Keys.Contains(new IssueKey("NamespaceX,KeyY", IssueKeyType.None)));
+			Assert.IsTrue(groupedIssues[2].ChangeFilter == "...");
+			Assert.IsTrue(groupedIssues[2].Events.Count == 1);
+			Assert.IsTrue(groupedIssues[2].Events[0].EventId == KnownLogEvents.Engine_Localization);
+			Assert.IsTrue(groupedIssues[2].Events[0].Severity == Microsoft.Extensions.Logging.LogLevel.Warning);
+			Assert.IsTrue(groupedIssues[2].Metadata.Count == 1);
+			Assert.IsTrue(groupedIssues[2].SummaryTemplate == "{Meta:Node}: Localization {Severity} in {Files}");
+			Assert.IsTrue(groupedIssues[2].Type == "Localization");
+
+			Assert.IsTrue(groupedIssues[3].Keys.Count == 2);
+			Assert.IsTrue(groupedIssues[3].Keys.Contains(new IssueKey("ManifestFile", IssueKeyType.File)));
+			Assert.IsTrue(groupedIssues[3].Keys.Contains(new IssueKey("NamespaceA,KeyA", IssueKeyType.None)));
+			Assert.IsTrue(groupedIssues[3].ChangeFilter == "...");
+			Assert.IsTrue(groupedIssues[3].Events.Count == 1);
+			Assert.IsTrue(groupedIssues[3].Events[0].EventId == KnownLogEvents.Engine_Localization);
+			Assert.IsTrue(groupedIssues[3].Events[0].Severity == Microsoft.Extensions.Logging.LogLevel.Warning);
+			Assert.IsTrue(groupedIssues[3].Metadata.Count == 1);
+			Assert.IsTrue(groupedIssues[3].SummaryTemplate == "{Meta:Node}: Localization {Severity} in {Files}");
+			Assert.IsTrue(groupedIssues[3].Type == "Localization");
+
+			Assert.IsTrue(groupedIssues[4].Keys.Count == 3);
+			Assert.IsTrue(groupedIssues[4].Keys.Contains(new IssueKey("Asset", IssueKeyType.File)));
+			Assert.IsTrue(groupedIssues[4].Keys.Contains(new IssueKey("File5", IssueKeyType.File)));
+			Assert.IsTrue(groupedIssues[4].Keys.Contains(new IssueKey("NamespaceX,KeyY", IssueKeyType.None)));
+			Assert.IsTrue(groupedIssues[4].ChangeFilter == "...");
+			Assert.IsTrue(groupedIssues[4].Events.Count == 1);
+			Assert.IsTrue(groupedIssues[4].Events[0].EventId == KnownLogEvents.Engine_Localization);
+			Assert.IsTrue(groupedIssues[4].Events[0].Severity == Microsoft.Extensions.Logging.LogLevel.Warning);
+			Assert.IsTrue(groupedIssues[4].Metadata.Count == 1);
+			Assert.IsTrue(groupedIssues[4].SummaryTemplate == "{Meta:Node}: Localization {Severity} in {Files}");
+			Assert.IsTrue(groupedIssues[4].Type == "Localization");
+
+			Assert.IsTrue(groupedIssues[5].Keys.Count == 2);
+			Assert.IsTrue(groupedIssues[5].Keys.Contains(new IssueKey("LocTargetX", IssueKeyType.File)));
+			Assert.IsTrue(groupedIssues[5].Keys.Contains(new IssueKey("NamespaceB,KeyB", IssueKeyType.None)));
+			Assert.IsTrue(groupedIssues[5].ChangeFilter == "...");
+			Assert.IsTrue(groupedIssues[5].Events.Count == 1);
+			Assert.IsTrue(groupedIssues[5].Events[0].EventId == KnownLogEvents.Engine_Localization);
+			Assert.IsTrue(groupedIssues[5].Events[0].Severity == Microsoft.Extensions.Logging.LogLevel.Warning);
+			Assert.IsTrue(groupedIssues[5].Metadata.Count == 1);
+			Assert.IsTrue(groupedIssues[5].SummaryTemplate == "{Meta:Node}: Localization {Severity} in {Files}");
+			Assert.IsTrue(groupedIssues[5].Type == "Localization");
+
+			Assert.IsTrue(groupedIssues[6].Keys.Count == 2);
+			Assert.IsTrue(groupedIssues[6].Keys.Contains(new IssueKey("File6", IssueKeyType.File)));
+			Assert.IsTrue(groupedIssues[6].Keys.Contains(new IssueKey("NamespaceC,KeyC", IssueKeyType.None)));
+			Assert.IsTrue(groupedIssues[6].ChangeFilter == "...");
+			Assert.IsTrue(groupedIssues[6].Events.Count == 1);
+			Assert.IsTrue(groupedIssues[6].Events[0].EventId == KnownLogEvents.Engine_Localization);
+			Assert.IsTrue(groupedIssues[6].Events[0].Severity == Microsoft.Extensions.Logging.LogLevel.Error);
+			Assert.IsTrue(groupedIssues[6].Metadata.Count == 1);
+			Assert.IsTrue(groupedIssues[6].SummaryTemplate == "{Meta:Node}: Localization {Severity} in {Files}");
+			Assert.IsTrue(groupedIssues[6].Type == "Localization");
+
+			Assert.IsTrue(groupedIssues[7].Keys.Count == 2);
+			Assert.IsTrue(groupedIssues[7].Keys.Contains(new IssueKey("NamespaceA,KeyA", IssueKeyType.None)));
+			Assert.IsTrue(groupedIssues[7].Keys.Contains(new IssueKey("NamespaceB,KeyB", IssueKeyType.None)));
+			Assert.IsTrue(groupedIssues[7].ChangeFilter == "...");
+			Assert.IsTrue(groupedIssues[7].Events.Count == 2);
+			Assert.IsTrue(groupedIssues[7].Events[0].EventId == KnownLogEvents.Engine_Localization);
+			Assert.IsTrue(groupedIssues[7].Events[0].Severity == Microsoft.Extensions.Logging.LogLevel.Error);
+			Assert.IsTrue(groupedIssues[7].Events[1].EventId == KnownLogEvents.Engine_Localization);
+			Assert.IsTrue(groupedIssues[7].Events[1].Severity == Microsoft.Extensions.Logging.LogLevel.Error);
+			Assert.IsTrue(groupedIssues[7].Metadata.Count == 0);
+			Assert.IsTrue(groupedIssues[7].SummaryTemplate == "Localization {Severity}");
+			Assert.IsTrue(groupedIssues[7].Type == "Localization");
+		}
+	}
+}
