Horde: Updating schema documentation

#rnx

diff --git a/Engine/Source/Programs/Horde/Docs/Config/Schema/AclActions.md b/Engine/Source/Programs/Horde/Docs/Config/Schema/AclActions.md
--- a/Engine/Source/Programs/Horde/Docs/Config/Schema/AclActions.md
+++ b/Engine/Source/Programs/Horde/Docs/Config/Schema/AclActions.md
@@ -608,6 +608,12 @@
 | `UpdateReplicator` |  |
 | `ViewReplicator` |  |
 
+## Secrets
+
+| Name | Description |
+| ---- | ----------- |
+| `ViewSecret` | View a credential |
+
 ## ServiceAccounts
 
 | Name | Description |
@@ -781,7 +787,7 @@
 
 | Name | Description |
 | ---- | ----------- |
-| `ViewSecret` | View a credential |
+| `ViewSecret` |  |
 
 ## ServiceAccounts
 
diff --git a/Engine/Source/Programs/Horde/Docs/Config/Schema/Telemetry.md b/Engine/Source/Programs/Horde/Docs/Config/Schema/Telemetry.md
--- a/Engine/Source/Programs/Horde/Docs/Config/Schema/Telemetry.md
+++ b/Engine/Source/Programs/Horde/Docs/Config/Schema/Telemetry.md
@@ -121,7 +121,7 @@ Telemetry chart configuraton
 Name | Description
 ---- | -----------
 `name` | `string`<br>The name of the chart, will be displayed on the dashboard
-`unit` | [TelemetryMetricUnitType](#telemetrymetricunittype-enum)<br>The unit of the data to display
+`display` | [TelemetryMetricUnitType](#telemetrymetricunittype-enum)<br>The unit to display
 `graph` | [TelemetryMetricGraphType](#telemetrymetricgraphtype-enum)<br>The graph type
 `metrics` | [TelemetryChartMetricConfig](#telemetrychartmetricconfig)`[]`<br>List of configured metrics
 `min` | `integer`<br>The min unit value for clamping chart
