Pass telemetry store id to Epic Telemetry backend.

diff --git a/Engine/Source/Programs/Horde/Plugins/Analytics/HordeServer.Analytics/Telemetry/Sinks/EpicTelemetrySink.cs b/Engine/Source/Programs/Horde/Plugins/Analytics/HordeServer.Analytics/Telemetry/Sinks/EpicTelemetrySink.cs
--- a/Engine/Source/Programs/Horde/Plugins/Analytics/HordeServer.Analytics/Telemetry/Sinks/EpicTelemetrySink.cs
+++ b/Engine/Source/Programs/Horde/Plugins/Analytics/HordeServer.Analytics/Telemetry/Sinks/EpicTelemetrySink.cs
@@ -267,6 +267,7 @@ namespace HordeServer.Telemetry.Sinks
 					queryParams.Add("AppEnvironment", telemetryEvent.RecordMeta.AppEnvironment);
 					queryParams.Add("SessionID", telemetryEvent.RecordMeta.SessionId);
 					queryParams.Add("UploadType", "eteventstream");
+					queryParams.Add("StoreID", telemetryStoreId.ToString());
 
 					Uri uri = new Uri(QueryHelpers.AddQueryString(_baseUrl.ToString(), queryParams));
 					writer = new Writer(uri, _jsonOptions);
