Separate Commit Id and Commit Order in TestData v2

#rnx

diff --git a/Engine/Source/Programs/Horde/Plugins/Build/HordeServer.Build/Jobs/TestData/TestDataController.cs b/Engine/Source/Programs/Horde/Plugins/Build/HordeServer.Build/Jobs/TestData/TestDataController.cs
--- a/Engine/Source/Programs/Horde/Plugins/Build/HordeServer.Build/Jobs/TestData/TestDataController.cs
+++ b/Engine/Source/Programs/Horde/Plugins/Build/HordeServer.Build/Jobs/TestData/TestDataController.cs
@@ -747,7 +747,8 @@ namespace HordeServer.Jobs.TestData
 					StepId = testData.StepId.ToString(),
 					Duration = testData.Duration,
 					StartDateTime = testData.StartDateTime,
-					BuildCommitId = testData.BuildCommitId,
+					CommitId = testData.BuildCommitId.ToString(),
+					CommitOrder = testData.BuildCommitId.Order,
 					MetadataId = testData.Metadata.ToString(),
 					NameRef = testData.NameRef.ToString(),
 					TestDataId = testData.TestDataId.ToString(),
@@ -814,7 +815,8 @@ namespace HordeServer.Jobs.TestData
 					StepId = testData.StepId.ToString(),
 					Duration = testData.Duration,
 					StartDateTime = testData.StartDateTime,
-					BuildCommitId = testData.BuildCommitId,
+					CommitId = testData.BuildCommitId.ToString(),
+					CommitOrder = testData.BuildCommitId.Order,
 					MetadataId = testData.Metadata.ToString(),
 					PhaseRef = testData.PhaseRef.ToString(),
 					SessionId = testData.SessionId.ToString(),
diff --git a/Engine/Source/Programs/Shared/EpicGames.Horde/Jobs/TestData/TestDataMessages.cs b/Engine/Source/Programs/Shared/EpicGames.Horde/Jobs/TestData/TestDataMessages.cs
--- a/Engine/Source/Programs/Shared/EpicGames.Horde/Jobs/TestData/TestDataMessages.cs
+++ b/Engine/Source/Programs/Shared/EpicGames.Horde/Jobs/TestData/TestDataMessages.cs
@@ -711,7 +711,12 @@ namespace EpicGames.Horde.Jobs.TestData
 		/// <summary>
 		/// The commit id associated with the test session
 		/// </summary>
-		public CommitId BuildCommitId { get; set; } = CommitId.Empty;
+		public string CommitId { get; set; } = String.Empty;
+
+		/// <summary>
+		/// The commit order associated with the test session
+		/// </summary>
+		public int CommitOrder { get; set; } = -1;
 
 		/// <summary>
 		/// The duration of the test session
@@ -792,7 +797,12 @@ namespace EpicGames.Horde.Jobs.TestData
 		/// <summary>
 		/// The commit id associated with the phase session
 		/// </summary>
-		public CommitId BuildCommitId { get; set; } = CommitId.Empty;
+		public string CommitId { get; set; } = String.Empty;
+
+		/// <summary>
+		/// The commit order associated with the phase session
+		/// </summary>
+		public int CommitOrder { get; set; } = -1;
 
 		/// <summary>
 		/// The duration of the phase session
