Horde: Update release notes

#rnx

diff --git a/Engine/Source/Programs/Horde/Docs/ReleaseNotes.md b/Engine/Source/Programs/Horde/Docs/ReleaseNotes.md
--- a/Engine/Source/Programs/Horde/Docs/ReleaseNotes.md
+++ b/Engine/Source/Programs/Horde/Docs/ReleaseNotes.md
@@ -1,5 +1,41 @@
 # Release Notes
 
+## 2025-04-24
+
+* Pass telemetry store id to Epic Telemetry backend. (41897043)
+* Added IssueHandlersTests to Horde's unittests (starting with a LocalizationIssueHandler test) (41887180)
+* Improved LocalizationIssueHandler to cover more incoming edgecases (when changing from UE_LOG to UE_LOGFMT) (41886613)
+* Horde Storage: Do not throw errors for duplicate keys on shadow blobs as these are expected (41872054)
+* Downgrade unknown batch error due to missing shelved CL to a warning (41871559)
+* Horde Server: Downgrade issue escalation channel not existing to a warning (41871033)
+* Revert back to display field for Analytics (41861309)
+* Remove Microsoft.Management.Infrastructure from agent telemetry
+  Removes the NuGet dependency for both agent and JobDriver. Using CIM is slow and also adds dependencies incompatible with Wine. (41860782)
+* Remove use of Microsoft.Management.Infrastructure in agent service upgrade
+  Using CIM is slow and also adds dependencies incompatible with Wine. (41860494)
+* Horde Server: Adjust http timeouts and fix issue with storage backend and global http client timeout (41823434)
+* Horde : Add docs on how to create an external secret provider (41816354)
+* Add exception handling when reporting job metric so that it doesn't suppress notifications (41788849)
+* Horde issues now support the annotations AutoAssignToUser as a fallback to AutoAssign (41768375)
+* Add logging to debug job step warning/error notifcations coming through (41681308)
+* Add cloud url copy operation to artifact split button (41645110)
+* Add a new workstation mode to agent
+  * An agent in workstation mode is expected to authenticate as the current user, and not via a pre-shared token like in a buildfarm scenario
+  * Uses interactive auth via EpicGames.Oidc to authenticate (41643324)
+* Log reason when compute task allocation fails in compute client (41640994)
+* Test case-insensitivity for pool ID when matching resources for compute cluster (41640817)
+* Add pool ID argument to "compute run" in Horde CLI tool (41640803)
+* Skip executing conform leases if agent is running in workstation mode
+  Workstation agents only deal with stateless compute leases. (41640666)
+* Fix Commits Viewer and Robomerge links not appearing in slack messages (41612887)
+* Log server profile used during agent startup (41600533)
+* Log type of lease being executed in agent (41600484)
+* Stop disposing of supplied HTTP handler in HordeClient
+  When enabling interactive auth for the Horde agent, this error came up as the handler was improperly disposed.
+  This fix addreses that as lifetime should be managed by caller, in this case, the service provider/DI. (41569693)
+* Allow only compute tasks to be scheduled on workstation agents (41565320)
+* Add Horde's user claim when issuing JWT tokens for built-in Horde auth mode (41533085)
+
 ## 2025-04-10
 
 * Run tests for experimental plugin in Dockerfile builds (41526895)
