Modify the Licensee Agent name to conform with other Agent namings used in other buildgraphs
This CL will require an update to Horde stream jsons to rename the previous agent type to the new one (separate CL)

#rnx

diff --git a/Engine/Source/Programs/Horde/BuildHorde.xml b/Engine/Source/Programs/Horde/BuildHorde.xml
--- a/Engine/Source/Programs/Horde/BuildHorde.xml
+++ b/Engine/Source/Programs/Horde/BuildHorde.xml
@@ -144,7 +144,7 @@
     <!-- Changing the windows architecture value here means updates elsewhere in this XML, and in the Installer wxs file -->
 	<Property Name="ServerArchList" Value="win-x64;osx-arm64;linux-x64"/>
 	
-	<Agent Name="HordeInstallerTools" Type="LicenseeWin64;Win64">
+	<Agent Name="HordeInstallerTools" Type="Win64_Licensee;Win64">
 		<Node Name="Build HordeCmd" Produces="#HordeCmdFiles">
 			<Delete Files="$(ServerToolsStagingDir_HordeCmd)/..."/>
 			<DotNet	Arguments="publish &quot;$(RootDir)/Engine/Source/Programs/Horde/HordeCmd/HordeCmd.csproj&quot; --output &quot;$(ServerToolsStagingDir_HordeCmd)&quot; --runtime win-x64 --self-contained $(VersionArguments)" DotNetPath="$(DotNetPath)" />
