Horde : The Notices API Get method will return the start and finish times when values are set

diff --git a/Engine/Source/Programs/Horde/HordeServer.Tests/HordeServer.Tests.csproj b/Engine/Source/Programs/Horde/HordeServer.Tests/HordeServer.Tests.csproj
--- a/Engine/Source/Programs/Horde/HordeServer.Tests/HordeServer.Tests.csproj
+++ b/Engine/Source/Programs/Horde/HordeServer.Tests/HordeServer.Tests.csproj
@@ -62,6 +62,7 @@
 
   <ItemGroup>
     <Folder Include="Resources\" />
+    <Folder Include="Server\Notices\" />
   </ItemGroup>
 
 </Project>
diff --git a/Engine/Source/Programs/Horde/HordeServer.Tests/Server/Notices/NoticeControllerTests.cs b/Engine/Source/Programs/Horde/HordeServer.Tests/Server/Notices/NoticeControllerTests.cs
new file mode 100644
--- /dev/null
+++ b/Engine/Source/Programs/Horde/HordeServer.Tests/Server/Notices/NoticeControllerTests.cs
@@ -0,0 +1,114 @@
+﻿// Copyright Epic Games, Inc. All Rights Reserved.
+
+using System;
+using System.Collections.Generic;
+using System.Net.Http;
+using System.Net.Http.Json;
+using System.Threading.Tasks;
+using EpicGames.Core;
+using EpicGames.Horde;
+using EpicGames.Horde.Server.Notices;
+using Microsoft.Extensions.DependencyInjection;
+using Microsoft.VisualStudio.TestTools.UnitTesting;
+
+namespace HordeServer.Tests.Server.Notices;
+
+[TestClass]
+public class NoticeControllerTests : ControllerIntegrationTest
+{
+	[TestMethod]
+	public async Task GetNoticesAsync()
+	{
+		CreateNoticeRequest notice = new("Hello, World!");
+		using HttpResponseMessage response = await Client.PostAsJsonAsync("api/v1/notices", notice);
+		response.EnsureSuccessStatusCode();
+		List<GetNoticeResponse> notices = await HordeHttpRequest.GetAsync<List<GetNoticeResponse>>(Client, "api/v1/notices");
+		GetNoticeResponse result = notices[0];
+		Assert.AreEqual(result.Message, "Hello, World!");
+		Assert.IsNull(result.StartTime);
+		Assert.IsNull(result.FinishTime);
+		Assert.IsTrue(result.Active);
+	}
+
+	[TestMethod]
+	public async Task GetNoticesWithStartAndFinishTimesAsync()
+	{
+		IClock clock = ServiceProvider.GetRequiredService<IClock>();
+		CreateNoticeRequest notice = new("Hello, World!")
+		{
+			StartTime = clock.UtcNow - TimeSpan.FromMinutes(1),
+			FinishTime = clock.UtcNow + TimeSpan.FromMinutes(1)
+		};
+		using HttpResponseMessage response = await Client.PostAsJsonAsync("api/v1/notices", notice);
+		response.EnsureSuccessStatusCode();
+		List<GetNoticeResponse> notices = await HordeHttpRequest.GetAsync<List<GetNoticeResponse>>(Client, "api/v1/notices");
+		GetNoticeResponse result = notices[0];
+		Assert.IsNotNull(result.StartTime);
+		Assert.IsNotNull(result.FinishTime);
+		Assert.IsTrue(result.Active);
+	}
+
+	[TestMethod]
+	public async Task GetNoticesWithStartTimeAsync()
+	{
+		IClock clock = ServiceProvider.GetRequiredService<IClock>();
+		CreateNoticeRequest notice = new("Hello, World!")
+		{
+			StartTime = clock.UtcNow - TimeSpan.FromMinutes(1)
+		};
+		using HttpResponseMessage response = await Client.PostAsJsonAsync("api/v1/notices", notice);
+		response.EnsureSuccessStatusCode();
+		List<GetNoticeResponse> notices = await HordeHttpRequest.GetAsync<List<GetNoticeResponse>>(Client, "api/v1/notices");
+		GetNoticeResponse result = notices[0];
+		Assert.IsNotNull(result.StartTime);
+		Assert.IsNull(result.FinishTime);
+		Assert.IsTrue(result.Active);
+	}
+
+	[TestMethod]
+	public async Task GetNoticesWithFinishTimeAsync()
+	{
+		IClock clock = ServiceProvider.GetRequiredService<IClock>();
+		CreateNoticeRequest notice = new("Hello, World!")
+		{
+			FinishTime = clock.UtcNow + TimeSpan.FromMinutes(1)
+		};
+		using HttpResponseMessage response = await Client.PostAsJsonAsync("api/v1/notices", notice);
+		response.EnsureSuccessStatusCode();
+		List<GetNoticeResponse> notices = await HordeHttpRequest.GetAsync<List<GetNoticeResponse>>(Client, "api/v1/notices");
+		GetNoticeResponse result = notices[0];
+		Assert.IsNull(result.StartTime);
+		Assert.IsNotNull(result.FinishTime);
+		Assert.IsTrue(result.Active);
+	}
+
+	[TestMethod]
+	public async Task GetNoticesInactivePastAsync()
+	{
+		IClock clock = ServiceProvider.GetRequiredService<IClock>();
+		CreateNoticeRequest notice = new("Hello, World!")
+		{
+			FinishTime = clock.UtcNow - TimeSpan.FromMinutes(1)
+		};
+		using HttpResponseMessage response = await Client.PostAsJsonAsync("api/v1/notices", notice);
+		response.EnsureSuccessStatusCode();
+		List<GetNoticeResponse> notices = await HordeHttpRequest.GetAsync<List<GetNoticeResponse>>(Client, "api/v1/notices");
+		GetNoticeResponse result = notices[0];
+		Assert.IsFalse(result.Active);
+	}
+
+	[TestMethod]
+	public async Task GetNoticesInactiveFutureAsync()
+	{
+		IClock clock = ServiceProvider.GetRequiredService<IClock>();
+		CreateNoticeRequest notice = new("Hello, World!")
+		{
+			StartTime = clock.UtcNow + TimeSpan.FromMinutes(1)
+		};
+		using HttpResponseMessage response = await Client.PostAsJsonAsync("api/v1/notices", notice);
+		response.EnsureSuccessStatusCode();
+		List<GetNoticeResponse> notices = await HordeHttpRequest.GetAsync<List<GetNoticeResponse>>(Client, "api/v1/notices");
+		GetNoticeResponse result = notices[0];
+		Assert.IsFalse(result.Active);
+	}
+}
diff --git a/Engine/Source/Programs/Horde/HordeServer/Server/Notices/NoticesController.cs b/Engine/Source/Programs/Horde/HordeServer/Server/Notices/NoticesController.cs
--- a/Engine/Source/Programs/Horde/HordeServer/Server/Notices/NoticesController.cs
+++ b/Engine/Source/Programs/Horde/HordeServer/Server/Notices/NoticesController.cs
@@ -129,7 +129,19 @@ namespace HordeServer.Server.Notices
 					userInfo = (await _userCollection.GetCachedUserAsync(notice.UserId, cancellationToken))?.ToThinApiResponse();
 				}
 
-				messages.Add(new GetNoticeResponse() { Id = notice.Id.ToString(), Active = true, Message = notice.Message, CreatedByUser = userInfo });
+				DateTime? startTime = null;
+				if (notice.StartTime != null)
+				{
+					startTime = TimeZoneInfo.ConvertTimeBySystemTimeZoneId(notice.StartTime.Value, "UTC");
+				}
+				DateTime? finishTime = null;
+				if (notice.FinishTime != null)
+				{
+					finishTime = TimeZoneInfo.ConvertTimeBySystemTimeZoneId(notice.FinishTime.Value, "UTC");
+				}
+				bool active = _clock.UtcNow >= (startTime ?? DateTime.MinValue) && _clock.UtcNow < (finishTime ?? DateTime.MaxValue);
+
+				messages.Add(new GetNoticeResponse() { Id = notice.Id.ToString(), Active = active, StartTime = startTime, FinishTime = finishTime, Message = notice.Message, CreatedByUser = userInfo });
 			}
 
 			// Add any status notices for admins
