Horde: Add support for multiple report summaries
#submttool bypass

diff --git a/Engine/Source/Programs/Horde/HordeDashboard/src/components/jobDetailsV2/JobDetailViewSummary.tsx b/Engine/Source/Programs/Horde/HordeDashboard/src/components/jobDetailsV2/JobDetailViewSummary.tsx
--- a/Engine/Source/Programs/Horde/HordeDashboard/src/components/jobDetailsV2/JobDetailViewSummary.tsx
+++ b/Engine/Source/Programs/Horde/HordeDashboard/src/components/jobDetailsV2/JobDetailViewSummary.tsx
@@ -117,7 +117,7 @@ export const SummaryPanel: React.FC<{ jobDetails: JobDetailsV2 }> = observer(({
 
    const summary = text.join(".  ");
 
-   const reportData = jobDetails.getReportData(ReportPlacement.Summary);
+   const reportData = jobDetails.getSummaryReport();
    
    const parentJobUrl = jobData.parentJobId && jobData.parentJobId ? `/job/${jobData.parentJobId}?step=${jobData.parentJobStepId}` : undefined;
 
diff --git a/Engine/Source/Programs/Horde/HordeDashboard/src/components/jobDetailsV2/JobDetailsViewCommon.tsx b/Engine/Source/Programs/Horde/HordeDashboard/src/components/jobDetailsV2/JobDetailsViewCommon.tsx
--- a/Engine/Source/Programs/Horde/HordeDashboard/src/components/jobDetailsV2/JobDetailsViewCommon.tsx
+++ b/Engine/Source/Programs/Horde/HordeDashboard/src/components/jobDetailsV2/JobDetailsViewCommon.tsx
@@ -5,7 +5,7 @@ import moment from "moment";
 import React, { useState } from "react";
 import { useLocation, useNavigate } from "react-router-dom";
 import backend from "../../backend";
-import { BatchData, GetArtifactResponse, GetJobTimingResponse, GetLabelStateResponse, GetTemplateRefResponse, JobData, JobState, LabelState, ReportPlacement, StepData, StreamData } from "../../backend/Api";
+import { BatchData, GetArtifactResponse, GetJobTimingResponse, GetLabelStateResponse, GetReportResponse, GetTemplateRefResponse, JobData, JobState, LabelState, ReportPlacement, StepData, StreamData } from "../../backend/Api";
 import { JobLabel } from "../../backend/JobDetails";
 import { PollBase } from "../../backend/PollBase";
 import { projectStore } from "../../backend/ProjectStore";
@@ -517,7 +517,9 @@ export class JobDetailsV2 extends PollBase {
    }
 
 
-   getReportData(placement: ReportPlacement, stepId?: string): string | undefined {
+   getSummaryReport(stepId?: string): string | undefined {
+
+      let reports: GetReportResponse[] = [];
 
       if (stepId) {
 
@@ -527,12 +529,25 @@ export class JobDetailsV2 extends PollBase {
             return undefined;
          }
 
-         return step.reports?.find(r => r.placement === placement)?.content;
+         reports = step.reports?.filter(r => r.placement === ReportPlacement.Summary) ?? [];
 
       } else {
 
-         return this.jobData?.reports?.find(r => r.placement === placement)?.content;
+         reports = this.jobData?.reports?.filter(r => r.placement === ReportPlacement.Summary) ?? [];         
       }
+
+      if (!reports.length) {
+         return undefined;
+      }      
+      return reports.map(r => {
+
+         if (r.content) {
+            return `**${r.name}**:  \n` + `${r.content}  \n  \n`            
+         }  else {
+            return "";
+         }
+
+      }).filter(r => !!r).join("");
    }
 
    syncTiming() {
diff --git a/Engine/Source/Programs/Horde/HordeDashboard/src/components/jobDetailsV2/StepDetailSummary.tsx b/Engine/Source/Programs/Horde/HordeDashboard/src/components/jobDetailsV2/StepDetailSummary.tsx
--- a/Engine/Source/Programs/Horde/HordeDashboard/src/components/jobDetailsV2/StepDetailSummary.tsx
+++ b/Engine/Source/Programs/Horde/HordeDashboard/src/components/jobDetailsV2/StepDetailSummary.tsx
@@ -233,7 +233,7 @@ export const StepSummaryPanel: React.FC<{ jobDetails: JobDetailsV2; stepId: stri
 
    const step = jobDetails.stepById(stepId)!
 
-   const reportData = jobDetails.getReportData(ReportPlacement.Summary, stepId);
+   const reportData = jobDetails.getSummaryReport(stepId);
 
    const spawnedJobs = step.spawnedJobs ?? [];
    const spawnedJobUrl = spawnedJobs.length > 0 ? `/job/${spawnedJobs[0]}` : undefined;
