Add auto_load param to CV links

diff --git a/Engine/Source/Programs/Horde/Plugins/Build/HordeServer.Build/Notifications/Sinks/SlackNotificationSink.cs b/Engine/Source/Programs/Horde/Plugins/Build/HordeServer.Build/Notifications/Sinks/SlackNotificationSink.cs
--- a/Engine/Source/Programs/Horde/Plugins/Build/HordeServer.Build/Notifications/Sinks/SlackNotificationSink.cs
+++ b/Engine/Source/Programs/Horde/Plugins/Build/HordeServer.Build/Notifications/Sinks/SlackNotificationSink.cs
@@ -1532,7 +1532,8 @@ namespace HordeServer.Notifications.Sinks
 				}
 
 				string query_string = $"?depot={depot_stream_split[0]}&stream={depot_stream_split[1]}" +
-					$"&query_range[CL][start_cl]={span.LastSuccess.CommitId.GetPerforceChange() + 1}&query_range[CL][end_cl]={span.FirstFailure.CommitId}";
+					$"&query_range[CL][start_cl]={span.LastSuccess.CommitId.GetPerforceChange() + 1}&query_range[CL][end_cl]={span.FirstFailure.CommitId}" +
+					"&auto_load=true";
 
 				return new Uri(_settings.CommitsViewerUrl, query_string);
 			}
