Horde Dashboard: Add support for multiple plugin mount points

diff --git a/Engine/Source/Programs/Horde/HordeDashboard/plugins/index.ts b/Engine/Source/Programs/Horde/HordeDashboard/plugins/index.ts
--- a/Engine/Source/Programs/Horde/HordeDashboard/plugins/index.ts
+++ b/Engine/Source/Programs/Horde/HordeDashboard/plugins/index.ts
@@ -18,7 +18,10 @@ export type Mount = {
 export type HordePlugin = {
    id: string;
    routes: RouteObject[];
+   // single plugin mount point
    mount?: Mount;
+   // multiple plugin mount points
+   mounts?: Mount[];
    enabled?: boolean;
 }
 
diff --git a/Engine/Source/Programs/Horde/HordeDashboard/src/components/TopNav.tsx b/Engine/Source/Programs/Horde/HordeDashboard/src/components/TopNav.tsx
--- a/Engine/Source/Programs/Horde/HordeDashboard/src/components/TopNav.tsx
+++ b/Engine/Source/Programs/Horde/HordeDashboard/src/components/TopNav.tsx
@@ -441,14 +441,29 @@ export const TopNav: React.FC<{ suppressServer?: boolean }> = observer(({ suppre
       let serviceItems: IContextualMenuItem[] = [];
 
       // get the plugins which mount into the tools bar
-      const plugins = getHordePlugins().filter(p => p.mount?.type === MountType.TopNav);
+      const plugins = getHordePlugins().filter(p => p.mount?.type === MountType.TopNav || p.mounts?.find(m => m.type === MountType.TopNav));
 
+      // Add plugin topnav mount points
       plugins.forEach(p => {
-         serviceItems.push({
-            key: `admin_plugin_${p.mount!.text}`,
-            text: p.mount!.text,
-            link: p.mount!.route
-         });   
+         if (p.mount) {
+            serviceItems.push({
+               key: `admin_plugin_${p.mount!.text}`,
+               text: p.mount!.text,
+               link: p.mount!.route
+            });      
+         }
+
+         p.mounts?.forEach( m => {
+            if (m.type !== MountType.TopNav) {
+               return;
+            }
+
+            serviceItems.push({
+               key: `admin_plugin_${m.text}`,
+               text: m.text,
+               link: m.route
+            });      
+         })
       })
 
       if (features?.showTests !== false) {
