Horde Server: Adding job and step meta data tagging

diff --git a/Engine/Source/Programs/Horde/Plugins/Build/HordeServer.Build.Tests/Jobs/JobCollectionTests.cs b/Engine/Source/Programs/Horde/Plugins/Build/HordeServer.Build.Tests/Jobs/JobCollectionTests.cs
--- a/Engine/Source/Programs/Horde/Plugins/Build/HordeServer.Build.Tests/Jobs/JobCollectionTests.cs
+++ b/Engine/Source/Programs/Horde/Plugins/Build/HordeServer.Build.Tests/Jobs/JobCollectionTests.cs
@@ -994,6 +994,63 @@ namespace HordeServer.Tests.Jobs
 
 		}
 
+		[TestMethod]
+		public async Task JobMetaDataAsync()
+		{
+			ProjectId projectId = new ProjectId("ue5");
+			StreamId streamId = new StreamId("ue5-main");
+			TemplateId templateRefId1 = new TemplateId("template1");
+			TemplateId templateRefId2 = new TemplateId("template2");
+
+			StreamConfig streamConfig = new StreamConfig { Id = streamId };
+			streamConfig.Templates.Add(new TemplateRefConfig { Id = templateRefId1, Name = "Test Template 1" });
+			streamConfig.Templates.Add(new TemplateRefConfig { Id = templateRefId2, Name = "Test Template 2" });
+			streamConfig.Tabs.Add(new TabConfig { Title = "foo", Templates = new List<TemplateId> { templateRefId1, templateRefId2 } });
+
+			ProjectConfig projectConfig = new ProjectConfig { Id = projectId };
+			projectConfig.Streams.Add(streamConfig);
+
+			BuildConfig buildConfig = new BuildConfig();
+			buildConfig.Projects.Add(projectConfig);
+
+			GlobalConfig globalConfig = new GlobalConfig();
+			globalConfig.Plugins.AddBuildConfig(buildConfig);
+
+			await SetConfigAsync(globalConfig);
+
+			CreateJobRequest request = new CreateJobRequest(new StreamId("ue5-main"), new TemplateId("template1"));
+			request.CommitId = CommitIdWithOrder.FromPerforceChange(123);
+
+			ActionResult<CreateJobResponse> jobResult = await JobsController.CreateJobAsync(request);
+
+			Assert.IsNotNull(jobResult.Value);
+			JobId jobId = JobId.Parse(jobResult.Value!.Id);
+
+			IJob? job = Deref(await JobCollection.GetAsync(jobId));
+
+			PutJobMetadataRequest metaRequest = new PutJobMetadataRequest() { JobMetaData = new List<string> { "TestMeta=42" }, StepMetaData = new Dictionary<JobStepId, List<string>> { { job.Batches[0].Steps[0].Id, new List<string> { "TestStepMeta=2025"} } } };
+
+			await JobsController.UpdateJobMetaDataAsync(jobId, metaRequest);
+
+			job = Deref(await JobCollection.GetAsync(jobId));
+
+			Assert.IsNotNull(job);
+			Assert.AreEqual(job.Metadata.Count, 1);
+			Assert.AreEqual(job.Metadata[0], "TestMeta=42");
+
+			Assert.AreEqual(job.Batches[0].Steps[0].Metadata.Count, 1);
+			Assert.AreEqual(job.Batches[0].Steps[0].Metadata[0], "TestStepMeta=2025");
+			
+			await JobStepRefCollection.UpdateAsync(job, job.Batches[0], job.Batches[0].Steps[0], job.Graph, null);
+
+			IJobStepRef? jobStepRef = await JobStepRefCollection.FindAsync(jobId, job.Batches[0].Id, job.Batches[0].Steps[0].Id);
+
+			Assert.IsNotNull(jobStepRef);
+			Assert.IsNotNull(jobStepRef.Metadata);
+			Assert.AreEqual(jobStepRef.Metadata.Count, 1);
+			Assert.AreEqual(jobStepRef.Metadata[0], "TestStepMeta=2025");
+		}
+
 		private async Task<IJob> AddJobAsync(IGraph graph, string commitId)
 		{
 			return await JobCollection.AddAsync(JobIdUtils.GenerateNewId(), _streamId, _templateId1, _template1.Hash, graph, "jobName", new CommitId(commitId), null, new CreateJobOptions());
diff --git a/Engine/Source/Programs/Horde/Plugins/Build/HordeServer.Build/Jobs/IJobCollection.cs b/Engine/Source/Programs/Horde/Plugins/Build/HordeServer.Build/Jobs/IJobCollection.cs
--- a/Engine/Source/Programs/Horde/Plugins/Build/HordeServer.Build/Jobs/IJobCollection.cs
+++ b/Engine/Source/Programs/Horde/Plugins/Build/HordeServer.Build/Jobs/IJobCollection.cs
@@ -207,6 +207,16 @@ namespace HordeServer.Jobs
 		/// <param name="args"></param>
 		void JobAudit(JobId jobId, string? message, params object?[] args);
 
+		/// <summary>
+		/// Update a job and step metadata
+		/// </summary>
+		/// <param name="jobId"></param>
+		/// <param name="jobMetadata"></param>
+		/// <param name="stepMetadata"></param>
+		/// <param name="cancellationToken"></param>
+		/// <returns></returns>
+		Task<IJob?> UpdateMetadataAsync(JobId jobId, List<string>? jobMetadata = null, Dictionary<JobStepId, List<string>>? stepMetadata = null, CancellationToken cancellationToken = default);
+
 		/// <summary>
 		/// Upgrade all documents in the collection
 		/// </summary>
diff --git a/Engine/Source/Programs/Horde/Plugins/Build/HordeServer.Build/Jobs/IJobStepRefCollection.cs b/Engine/Source/Programs/Horde/Plugins/Build/HordeServer.Build/Jobs/IJobStepRefCollection.cs
--- a/Engine/Source/Programs/Horde/Plugins/Build/HordeServer.Build/Jobs/IJobStepRefCollection.cs
+++ b/Engine/Source/Programs/Horde/Plugins/Build/HordeServer.Build/Jobs/IJobStepRefCollection.cs
@@ -39,7 +39,8 @@ namespace HordeServer.Jobs
 		/// <param name="jobStartTimeUtc">Start time of the job</param>
 		/// <param name="startTimeUtc">Start time</param>
 		/// <param name="finishTimeUtc">Finish time for the step, if known</param>		
-		Task<IJobStepRef> InsertOrReplaceAsync(JobStepRefId id, string jobName, string stepName, StreamId streamId, TemplateId templateId, CommitIdWithOrder commitId, LogId? logId, PoolId? poolId, AgentId? agentId, JobStepState? state, JobStepOutcome? outcome, bool updateIssues, CommitIdWithOrder? lastSuccess, CommitIdWithOrder? lastWarning, float waitTime, float initTime, DateTime jobStartTimeUtc, DateTime startTimeUtc, DateTime? finishTimeUtc);
+		/// <param name="stepMetadata">Step meta data</param>		
+		Task<IJobStepRef> InsertOrReplaceAsync(JobStepRefId id, string jobName, string stepName, StreamId streamId, TemplateId templateId, CommitIdWithOrder commitId, LogId? logId, PoolId? poolId, AgentId? agentId, JobStepState? state, JobStepOutcome? outcome, bool updateIssues, CommitIdWithOrder? lastSuccess, CommitIdWithOrder? lastWarning, float waitTime, float initTime, DateTime jobStartTimeUtc, DateTime startTimeUtc, DateTime? finishTimeUtc, List<string>? stepMetadata = null);
 
 		/// <summary>
 		/// Updates a job step ref 
@@ -150,7 +151,7 @@ namespace HordeServer.Jobs
 					logger.LogInformation("Updating step reference {StepId} for job {JobId}, batch {BatchId}, with outcome {JobStepOutcome}", step.Id, job.Id, batch.Id, outcome);
 				}
 
-				await jobStepRefs.InsertOrReplaceAsync(new JobStepRefId(job.Id, batch.Id, step.Id), job.Name, nodeName, job.StreamId, job.TemplateId, job.CommitId, step.LogId, batch.PoolId, batch.AgentId, step.State, outcome, job.UpdateIssues, lastSuccess, lastWarning, waitTime, initTime, job.CreateTimeUtc, step.StartTimeUtc ?? DateTime.UtcNow, step.FinishTimeUtc);
+				await jobStepRefs.InsertOrReplaceAsync(new JobStepRefId(job.Id, batch.Id, step.Id), job.Name, nodeName, job.StreamId, job.TemplateId, job.CommitId, step.LogId, batch.PoolId, batch.AgentId, step.State, outcome, job.UpdateIssues, lastSuccess, lastWarning, waitTime, initTime, job.CreateTimeUtc, step.StartTimeUtc ?? DateTime.UtcNow, step.FinishTimeUtc, step.Metadata.ToList());
 			}
 		}
 	}
diff --git a/Engine/Source/Programs/Horde/Plugins/Build/HordeServer.Build/Jobs/JobCollection.cs b/Engine/Source/Programs/Horde/Plugins/Build/HordeServer.Build/Jobs/JobCollection.cs
--- a/Engine/Source/Programs/Horde/Plugins/Build/HordeServer.Build/Jobs/JobCollection.cs
+++ b/Engine/Source/Programs/Horde/Plugins/Build/HordeServer.Build/Jobs/JobCollection.cs
@@ -56,6 +56,8 @@ namespace HordeServer.Jobs
 			public IGraph Graph => _graph;
 			public List<JobStepBatch> Batches => _batches;
 
+			public List<string> Metadata => _document.Metadata;
+
 			JobId IJob.Id => _document.Id;
 			StreamId IJob.StreamId => _document.StreamId;
 			TemplateId IJob.TemplateId => _document.TemplateId;
@@ -311,6 +313,7 @@ namespace HordeServer.Jobs
 			bool IJobStep.Warnings => Node.Warnings;
 			IReadOnlyDictionary<string, string>? IJobStep.Credentials => Node.Credentials;
 			IReadOnlyNodeAnnotations IJobStep.Annotations => Node.Annotations;
+			IReadOnlyList<string> IJobStep.Metadata => Document.Metadata;
 
 			// Step properties
 			JobStepState IJobStep.State => Document.State;
@@ -532,6 +535,9 @@ namespace HordeServer.Jobs
 			[BsonIgnoreIfNull]
 			public DateTime? UpdateTimeUtc { get; set; }
 
+			[BsonElement("Metadata")]
+			public List<string> Metadata { get; set; } = new List<string>();
+
 			[BsonRequired]
 			public int UpdateIndex { get; set; }
 
@@ -542,7 +548,7 @@ namespace HordeServer.Jobs
 				GraphHash = null!;
 			}
 
-			public JobDocument(JobId id, StreamId streamId, TemplateId templateId, ContentHash templateHash, ContentHash graphHash, string name, CommitIdWithOrder commitId, CommitIdWithOrder? codeCommitId, CreateJobOptions options, DateTime createTimeUtc)
+			public JobDocument(JobId id, StreamId streamId, TemplateId templateId, ContentHash templateHash, ContentHash graphHash, string name, CommitIdWithOrder commitId, CommitIdWithOrder? codeCommitId, CreateJobOptions options, DateTime createTimeUtc, List<string> metadata)
 			{
 				Id = id;
 				StreamId = streamId;
@@ -583,6 +589,7 @@ namespace HordeServer.Jobs
 
 				NextSubResourceId = SubResourceId.GenerateNewId();
 				UpdateTimeUtc = createTimeUtc;
+				Metadata = metadata;
 			}
 
 			public DateTime GetCreateTimeOrDefault()
@@ -669,6 +676,8 @@ namespace HordeServer.Jobs
 
 			public JobStepError Error { get; set; } = JobStepError.None;
 
+			public List<string> Metadata { get; set; } = new List<string>();
+
 			[BsonIgnoreIfNull]
 			public LogId? LogId { get; set; }
 
@@ -774,7 +783,6 @@ namespace HordeServer.Jobs
 		const int MaxRetries = 2;
 
 		readonly IMongoCollection<JobDocument> _jobs;
-		readonly MongoIndex<JobDocument> _createTimeIndex;
 		readonly MongoIndex<JobDocument> _streamThenTemplateThenCreationTimeIndex;
 		readonly MongoIndex<JobDocument> _startedByBisectTaskIdIndex;
 		readonly ITelemetryWriter _telemetryWriter;
@@ -809,7 +817,7 @@ namespace HordeServer.Jobs
 			indexes.Add(keys => keys.Ascending(x => x.CommitOrder));
 			indexes.Add(keys => keys.Ascending(x => x.PreflightCommitName));
 			indexes.Add(keys => keys.Ascending(x => x.PreflightChange));
-			indexes.Add(_createTimeIndex = MongoIndex.Create<JobDocument>(keys => keys.Descending(x => x.CreateTimeUtc)));
+			indexes.Add(MongoIndex.Create<JobDocument>(keys => keys.Descending(x => x.CreateTimeUtc)));
 			indexes.Add(keys => keys.Ascending(x => x.StartedByUserId));
 			indexes.Add(keys => keys.Descending(x => x.SchedulePriority));
 			indexes.Add(_startedByBisectTaskIdIndex = MongoIndex.Create<JobDocument>(keys => keys.Descending(x => x.StartedByBisectTaskId), sparse: true));
@@ -900,7 +908,7 @@ namespace HordeServer.Jobs
 				codeCommitIdWithOrder = await _commitService.GetOrderedAsync(streamId, codeCommitId, cancellationToken);
 			}
 
-			JobDocument newJob = new JobDocument(jobId, streamId, templateRefId, templateHash, graph.Id, name, commitIdWithOrder, codeCommitIdWithOrder, options, DateTime.UtcNow);
+			JobDocument newJob = new JobDocument(jobId, streamId, templateRefId, templateHash, graph.Id, name, commitIdWithOrder, codeCommitIdWithOrder, options, DateTime.UtcNow, new List<string>());
 
 			JobAudit(jobId, "Job {JobId} added, stream {StreamId}, templateRefId {TemplateId}, name {Name}, commitId {CommitId}", jobId, streamId.Id, templateRefId, name, commitId);
 
@@ -2711,6 +2719,52 @@ namespace HordeServer.Jobs
 			return JobStepBatchState.Ready;
 		}
 
+		/// <inheritdoc/>
+		public async Task<IJob?> UpdateMetadataAsync(JobId jobId, List<string>? jobMetadata = null, Dictionary<JobStepId, List<string>>? stepMetadata = null, CancellationToken cancellationToken = default)
+		{
+			try
+			{
+				JobDocument? job = await _jobs.Find(x => x.Id == jobId).FirstOrDefaultAsync(cancellationToken);
+				job = Clone(job);
+
+				// Create the update 
+				UpdateDefinitionBuilder<JobDocument> updateBuilder = Builders<JobDocument>.Update;
+				List<UpdateDefinition<JobDocument>> updates = new List<UpdateDefinition<JobDocument>>();
+
+				if (job == null)
+				{
+					return null;
+				}
+
+				if (jobMetadata != null && jobMetadata.Count > 0)
+				{
+					updates.Add(updateBuilder.AddToSetEach(x => x.Metadata, jobMetadata));
+				}
+
+				if (stepMetadata != null)
+				{
+					foreach (KeyValuePair<JobStepId, List<string>> entry in stepMetadata)
+					{						
+						int batchIdx = job.Batches.FindIndex(x => x.Steps.FirstOrDefault(y => y.Id == entry.Key) != null);
+						if (batchIdx >= 0)
+						{
+							int stepIdx = job.Batches[batchIdx].Steps.FindIndex(x => x.Id == entry.Key);
+							updates.Add(updateBuilder.AddToSetEach(x => x.Batches[batchIdx].Steps[stepIdx].Metadata, entry.Value));
+						}
+					}						
+				}
+
+				job = await TryUpdateAsync(job, updates, cancellationToken);
+
+				return job == null ? null : await CreateJobObjectAsync(job, cancellationToken);
+			}
+			catch (Exception ex)
+			{
+				_logger.LogError(ex, "Failed to update job {JobId} meta data", jobId);
+				return null;
+			}			
+		}
+
 		/// <inheritdoc/>
 		public async Task UpgradeDocumentsAsync()
 		{
diff --git a/Engine/Source/Programs/Horde/Plugins/Build/HordeServer.Build/Jobs/JobService.cs b/Engine/Source/Programs/Horde/Plugins/Build/HordeServer.Build/Jobs/JobService.cs
--- a/Engine/Source/Programs/Horde/Plugins/Build/HordeServer.Build/Jobs/JobService.cs
+++ b/Engine/Source/Programs/Horde/Plugins/Build/HordeServer.Build/Jobs/JobService.cs
@@ -80,7 +80,6 @@ namespace HordeServer.Jobs
 		readonly IUserCollection _userCollection;
 		readonly INotificationTriggerCollection _triggerCollection;
 		readonly JobTaskSource _jobTaskSource;
-		readonly IStreamCollection _streamCollection;
 		readonly ITemplateCollectionInternal _templateCollection;
 		readonly IPerforceService _perforceService;
 		readonly IOptionsMonitor<BuildConfig> _buildConfig;
@@ -128,7 +127,7 @@ namespace HordeServer.Jobs
 		/// <summary>
 		/// Constructor
 		/// </summary>
-		public JobService(IJobCollection jobs, IGraphCollection graphs, IAgentCollection agents, IJobStepRefCollection jobStepRefs, IBisectTaskCollection bisectTasks, IJobTimingCollection jobTimings, IUserCollection userCollection, INotificationTriggerCollection triggerCollection, JobTaskSource jobTaskSource, IStreamCollection streamCollection, ITemplateCollectionInternal templateCollection, IPerforceService perforceService, Tracer tracer, IOptionsMonitor<BuildConfig> buildConfig, ILogger<JobService> logger, IAuditLog<JobId> auditLog)
+		public JobService(IJobCollection jobs, IGraphCollection graphs, IAgentCollection agents, IJobStepRefCollection jobStepRefs, IBisectTaskCollection bisectTasks, IJobTimingCollection jobTimings, IUserCollection userCollection, INotificationTriggerCollection triggerCollection, JobTaskSource jobTaskSource, ITemplateCollectionInternal templateCollection, IPerforceService perforceService, Tracer tracer, IOptionsMonitor<BuildConfig> buildConfig, ILogger<JobService> logger, IAuditLog<JobId> auditLog)
 		{
 			_jobs = jobs;
 			_graphs = graphs;
@@ -138,8 +137,7 @@ namespace HordeServer.Jobs
 			_jobTimings = jobTimings;
 			_userCollection = userCollection;
 			_triggerCollection = triggerCollection;
-			_jobTaskSource = jobTaskSource;
-			_streamCollection = streamCollection;
+			_jobTaskSource = jobTaskSource;			
 			_templateCollection = templateCollection;
 			_perforceService = perforceService;
 			_buildConfig = buildConfig;
@@ -1081,6 +1079,19 @@ namespace HordeServer.Jobs
 			return null;
 		}
 
+		/// <summary>
+		/// Update a job and step metadata
+		/// </summary>
+		/// <param name="jobId"></param>
+		/// <param name="jobMetadata"></param>
+		/// <param name="stepMetadata"></param>
+		/// <param name="cancellationToken"></param>
+		/// <returns></returns>
+		public async Task<IJob?> UpdateMetadataAsync(JobId jobId, List<string>? jobMetadata = null, Dictionary<JobStepId, List<string>>? stepMetadata = null, CancellationToken cancellationToken = default)
+		{
+			return await _jobs.UpdateMetadataAsync(jobId, jobMetadata, stepMetadata, cancellationToken);
+		}
+
 		/// <summary>
 		/// Submit the given change for a preflight
 		/// </summary>
@@ -1153,29 +1164,6 @@ namespace HordeServer.Jobs
 			}
 		}
 
-		/// <summary>
-		/// Clone a shelved changelist for running a preflight
-		/// </summary>
-		/// <param name="clusterName">Name of the Perforce cluster</param>
-		/// <param name="change">The changelist to clone</param>
-		/// <param name="cancellationToken">Cancellation token for the operation</param>
-		/// <returns></returns>
-		private async Task<CommitId> CloneShelvedChangeAsync(string clusterName, CommitId change, CancellationToken cancellationToken = default)
-		{
-			int clonedChange;
-			try
-			{
-				clonedChange = await _perforceService.DuplicateShelvedChangeAsync(clusterName, change.GetPerforceChange(), cancellationToken);
-				_logger.LogInformation("CL {Change} was duplicated into {ClonedChange}", change, clonedChange);
-			}
-			catch (Exception ex)
-			{
-				_logger.LogError(ex, "Unable to CL {Change} for preflight: {Message}", change, ex.Message);
-				throw;
-			}
-			return CommitId.FromPerforceChange(clonedChange);
-		}
-
 		/// <summary>
 		/// Deletes a shelved changelist, catching and logging any exceptions
 		/// </summary>
diff --git a/Engine/Source/Programs/Horde/Plugins/Build/HordeServer.Build/Jobs/JobStepRefCollection.cs b/Engine/Source/Programs/Horde/Plugins/Build/HordeServer.Build/Jobs/JobStepRefCollection.cs
--- a/Engine/Source/Programs/Horde/Plugins/Build/HordeServer.Build/Jobs/JobStepRefCollection.cs
+++ b/Engine/Source/Programs/Horde/Plugins/Build/HordeServer.Build/Jobs/JobStepRefCollection.cs
@@ -53,6 +53,9 @@ namespace HordeServer.Jobs
 			public JobStepState? State { get; set; }
 			public JobStepOutcome? Outcome { get; set; }
 
+			[BsonElement("md"), BsonIgnoreIfNull]
+			public List<string>? Metadata { get; set; }
+
 			[BsonIgnore]
 			public CommitIdWithOrder? LastSuccess
 			{
@@ -107,8 +110,9 @@ namespace HordeServer.Jobs
 			string IJobStepRef.NodeName => Name;
 			bool IJobStepRef.UpdateIssues => UpdateIssues ?? false;
 			IReadOnlyList<int>? IJobStepRef.IssueIds => IssueIds;
+			IReadOnlyList<string>? IJobStepRef.Metadata => Metadata;
 
-			public JobStepRef(JobStepRefId id, string jobName, string nodeName, StreamId streamId, TemplateId templateId, CommitIdWithOrder commitId, LogId? logId, PoolId? poolId, AgentId? agentId, JobStepState? state, JobStepOutcome? outcome, bool updateIssues, CommitIdWithOrder? lastSuccess, CommitIdWithOrder? lastWarning, float batchWaitTime, float batchInitTime, DateTime jobStartTimeUtc, DateTime startTimeUtc, DateTime? finishTimeUtc)
+			public JobStepRef(JobStepRefId id, string jobName, string nodeName, StreamId streamId, TemplateId templateId, CommitIdWithOrder commitId, LogId? logId, PoolId? poolId, AgentId? agentId, JobStepState? state, JobStepOutcome? outcome, bool updateIssues, CommitIdWithOrder? lastSuccess, CommitIdWithOrder? lastWarning, float batchWaitTime, float batchInitTime, DateTime jobStartTimeUtc, DateTime startTimeUtc, DateTime? finishTimeUtc, List<string>? stepMetadata = null)
 			{
 				Id = id;
 				JobName = jobName;
@@ -129,6 +133,7 @@ namespace HordeServer.Jobs
 				JobStartTimeUtc = jobStartTimeUtc;
 				StartTimeUtc = startTimeUtc;
 				FinishTimeUtc = finishTimeUtc;
+				Metadata = stepMetadata;
 			}
 		}
 
@@ -150,9 +155,9 @@ namespace HordeServer.Jobs
 		}
 
 		/// <inheritdoc/>
-		public async Task<IJobStepRef> InsertOrReplaceAsync(JobStepRefId id, string jobName, string stepName, StreamId streamId, TemplateId templateId, CommitIdWithOrder commitId, LogId? logId, PoolId? poolId, AgentId? agentId, JobStepState? state, JobStepOutcome? outcome, bool updateIssues, CommitIdWithOrder? lastSuccess, CommitIdWithOrder? lastWarning, float waitTime, float initTime, DateTime jobStartTimeUtc, DateTime startTimeUtc, DateTime? finishTimeUtc)
+		public async Task<IJobStepRef> InsertOrReplaceAsync(JobStepRefId id, string jobName, string stepName, StreamId streamId, TemplateId templateId, CommitIdWithOrder commitId, LogId? logId, PoolId? poolId, AgentId? agentId, JobStepState? state, JobStepOutcome? outcome, bool updateIssues, CommitIdWithOrder? lastSuccess, CommitIdWithOrder? lastWarning, float waitTime, float initTime, DateTime jobStartTimeUtc, DateTime startTimeUtc, DateTime? finishTimeUtc, List<string>? stepMetadata = null)
 		{
-			JobStepRef newJobStepRef = new JobStepRef(id, jobName, stepName, streamId, templateId, commitId, logId, poolId, agentId, state, outcome, updateIssues, lastSuccess, lastWarning, waitTime, initTime, jobStartTimeUtc, startTimeUtc, finishTimeUtc);
+			JobStepRef newJobStepRef = new JobStepRef(id, jobName, stepName, streamId, templateId, commitId, logId, poolId, agentId, state, outcome, updateIssues, lastSuccess, lastWarning, waitTime, initTime, jobStartTimeUtc, startTimeUtc, finishTimeUtc, stepMetadata);
 			await _jobStepRefs.ReplaceOneAsync(Builders<JobStepRef>.Filter.Eq(x => x.Id, newJobStepRef.Id), newJobStepRef, new ReplaceOptions { IsUpsert = true });
 
 			if (_buildConfig.CurrentValue.TryGetStream(streamId, out StreamConfig? streamConfig) && !streamConfig.TelemetryStoreId.IsEmpty)
diff --git a/Engine/Source/Programs/Horde/Plugins/Build/HordeServer.Build/Jobs/JobsController.cs b/Engine/Source/Programs/Horde/Plugins/Build/HordeServer.Build/Jobs/JobsController.cs
--- a/Engine/Source/Programs/Horde/Plugins/Build/HordeServer.Build/Jobs/JobsController.cs
+++ b/Engine/Source/Programs/Horde/Plugins/Build/HordeServer.Build/Jobs/JobsController.cs
@@ -386,6 +386,37 @@ namespace HordeServer.Jobs
 			return Ok();
 		}
 
+		/// <summary>
+		/// Updates a specific job.
+		/// </summary>
+		/// <param name="jobId">Id of the job to find</param>
+		/// <param name="request">Settings to update in the job</param>
+		/// <returns>Async task</returns>
+		[HttpPut]
+		[Route("/api/v1/jobs/{jobId}/metadata")]
+		public async Task<ActionResult> UpdateJobMetaDataAsync(JobId jobId, [FromBody] PutJobMetadataRequest request, CancellationToken cancellationToken = default)
+		{
+			IJob? job = await _jobService.GetJobAsync(jobId, cancellationToken);
+			if (job == null)
+			{
+				return NotFound(jobId);
+			}
+
+			StreamConfig? streamConfig;
+			if (!_buildConfig.Value.TryGetStream(job.StreamId, out streamConfig))
+			{
+				return NotFound(job.StreamId);
+			}
+			if (!streamConfig.Authorize(JobAclAction.UpdateJob, User))
+			{
+				return Forbid(JobAclAction.UpdateJob, jobId);
+			}
+
+			await _jobService.UpdateMetadataAsync(jobId, request.JobMetaData, request.StepMetaData, cancellationToken);
+			
+			return Ok();
+		}
+
 		/// <summary>
 		/// Updates notifications for a specific job.
 		/// </summary>
@@ -671,6 +702,7 @@ namespace HordeServer.Jobs
 			response.CanUpdate = canUpdate;
 			response.ParentJobId = job.ParentJobId.ToString();
 			response.ParentJobStepId = job.ParentJobStepId.ToString();
+			response.Metadata = job.Metadata.ToList();
 			return response;
 		}
 
@@ -809,6 +841,7 @@ namespace HordeServer.Jobs
 			response.FinishTime = step.FinishTimeUtc;
 			response.Reports = step.Reports?.ConvertAll(x => CreateGetReportResponse(x));
 			response.SpawnedJobs = step.SpawnedJobs?.ConvertAll(x => x.ToString());
+			response.Metadata = step.Metadata.ToList();
 
 			if (step.Properties != null && step.Properties.Count > 0)
 			{
diff --git a/Engine/Source/Programs/Horde/Plugins/Build/HordeServer.Build/Streams/StreamsController.cs b/Engine/Source/Programs/Horde/Plugins/Build/HordeServer.Build/Streams/StreamsController.cs
--- a/Engine/Source/Programs/Horde/Plugins/Build/HordeServer.Build/Streams/StreamsController.cs
+++ b/Engine/Source/Programs/Horde/Plugins/Build/HordeServer.Build/Streams/StreamsController.cs
@@ -330,6 +330,7 @@ namespace HordeServer.Streams
 			response.IssueIds = jobStepRef.IssueIds?.Select(id => id).ToList();
 			response.StartTime = jobStepRef.StartTimeUtc;
 			response.FinishTime = jobStepRef.FinishTimeUtc;
+			response.Metadata = jobStepRef.Metadata?.ToList();
 			return response;
 		}
 
diff --git a/Engine/Source/Programs/Shared/EpicGames.Horde/Jobs/IJob.cs b/Engine/Source/Programs/Shared/EpicGames.Horde/Jobs/IJob.cs
--- a/Engine/Source/Programs/Shared/EpicGames.Horde/Jobs/IJob.cs
+++ b/Engine/Source/Programs/Shared/EpicGames.Horde/Jobs/IJob.cs
@@ -254,6 +254,11 @@ namespace EpicGames.Horde.Jobs
 		/// </summary>
 		public DateTime UpdateTimeUtc { get; }
 
+		/// <summary>
+		/// The job meta data
+		/// </summary>
+		public List<string> Metadata { get; }
+
 		/// <summary>
 		/// Update counter for this document. Any updates should compare-and-swap based on the value of this counter, or increment it in the case of server-side updates.
 		/// </summary>
@@ -1492,6 +1497,11 @@ namespace EpicGames.Horde.Jobs
 		/// </summary>
 		public IReadOnlyNodeAnnotations Annotations { get; }
 
+		/// <summary>
+		/// Metadata for this node
+		/// </summary>
+		public IReadOnlyList<string> Metadata { get; }
+
 		/// <summary>
 		/// Current state of the job step. This is updated automatically when runs complete.
 		/// </summary>
diff --git a/Engine/Source/Programs/Shared/EpicGames.Horde/Jobs/IJobStepRef.cs b/Engine/Source/Programs/Shared/EpicGames.Horde/Jobs/IJobStepRef.cs
--- a/Engine/Source/Programs/Shared/EpicGames.Horde/Jobs/IJobStepRef.cs
+++ b/Engine/Source/Programs/Shared/EpicGames.Horde/Jobs/IJobStepRef.cs
@@ -68,15 +68,20 @@ namespace HordeServer.Jobs
 		public JobStepOutcome? Outcome { get; }
 
 		/// <summary>
-		/// Whether this step should update issues
+		/// Metadata for the step
 		/// </summary>
-		public bool UpdateIssues { get; }
+		public IReadOnlyList<string>? Metadata { get; }
 
 		/// <summary>
 		/// Issues ids affecting this job step
 		/// </summary>
 		public IReadOnlyList<int>? IssueIds { get; }
 
+		/// <summary>
+		/// Whether this step should update issues
+		/// </summary>
+		public bool UpdateIssues { get; }
+
 		/// <summary>
 		/// The last change that succeeded. Note that this is only set when the ref is updated; it is not necessarily consistent with steps run later.
 		/// </summary>
diff --git a/Engine/Source/Programs/Shared/EpicGames.Horde/Jobs/JobMessages.cs b/Engine/Source/Programs/Shared/EpicGames.Horde/Jobs/JobMessages.cs
--- a/Engine/Source/Programs/Shared/EpicGames.Horde/Jobs/JobMessages.cs
+++ b/Engine/Source/Programs/Shared/EpicGames.Horde/Jobs/JobMessages.cs
@@ -425,6 +425,22 @@ namespace EpicGames.Horde.Jobs
 		public List<string>? Arguments { get; set; }
 	}
 
+	/// <summary>
+	/// Job metadata put request
+	/// </summary>
+	public class PutJobMetadataRequest
+	{
+		/// <summary>
+		/// Meta data to append to the job
+		/// </summary>
+		public List<string>? JobMetaData { get; set; }
+
+		/// <summary>
+		/// Step meta data to append
+		/// </summary>
+		public Dictionary<JobStepId, List<string>>? StepMetaData { get; set; }
+	}
+
 	/// <summary>
 	/// Placement for a job report
 	/// </summary>
@@ -689,6 +705,11 @@ namespace EpicGames.Horde.Jobs
 		/// </summary>
 		public List<string>? Targets { get; set; }
 
+		/// <summary>
+		/// List of metadata for this job
+		/// </summary>
+		public List<string>? Metadata { get; set; }
+
 		/// <summary>
 		/// The last update time for this job
 		/// </summary>
@@ -949,6 +970,11 @@ namespace EpicGames.Horde.Jobs
 		/// </summary>
 		public IReadOnlyDictionary<string, string>? Annotations { get; set; }
 
+		/// <summary>
+		/// List of metadata for this step
+		/// </summary>
+		public List<string>? Metadata { get; set; }
+
 		/// <summary>
 		/// Current state of the job step. This is updated automatically when runs complete.
 		/// </summary>
@@ -1514,6 +1540,11 @@ namespace EpicGames.Horde.Jobs
 		/// </summary>
 		public List<int>? IssueIds { get; set; }
 
+		/// <summary>
+		/// The step meta data for this step
+		/// </summary>
+		public List<string>? Metadata { get; set; }
+
 		/// <summary>
 		/// Time at which the step started.
 		/// </summary>
