Horde Server: Fix controller parameter warning

diff --git a/Engine/Source/Programs/Horde/Plugins/Build/HordeServer.Build/Jobs/JobsController.cs b/Engine/Source/Programs/Horde/Plugins/Build/HordeServer.Build/Jobs/JobsController.cs
--- a/Engine/Source/Programs/Horde/Plugins/Build/HordeServer.Build/Jobs/JobsController.cs
+++ b/Engine/Source/Programs/Horde/Plugins/Build/HordeServer.Build/Jobs/JobsController.cs
@@ -391,6 +391,7 @@ namespace HordeServer.Jobs
 		/// </summary>
 		/// <param name="jobId">Id of the job to find</param>
 		/// <param name="request">Settings to update in the job</param>
+		/// <param name="cancellationToken">Cancellation token</param>
 		/// <returns>Async task</returns>
 		[HttpPut]
 		[Route("/api/v1/jobs/{jobId}/metadata")]
