Horde Server: Marshal metadata job step id by string

diff --git a/Engine/Source/Programs/Horde/Plugins/Build/HordeServer.Build.Tests/Jobs/JobCollectionTests.cs b/Engine/Source/Programs/Horde/Plugins/Build/HordeServer.Build.Tests/Jobs/JobCollectionTests.cs
--- a/Engine/Source/Programs/Horde/Plugins/Build/HordeServer.Build.Tests/Jobs/JobCollectionTests.cs
+++ b/Engine/Source/Programs/Horde/Plugins/Build/HordeServer.Build.Tests/Jobs/JobCollectionTests.cs
@@ -1028,7 +1028,7 @@ namespace HordeServer.Tests.Jobs
 
 			IJob? job = Deref(await JobCollection.GetAsync(jobId));
 
-			PutJobMetadataRequest metaRequest = new PutJobMetadataRequest() { JobMetaData = new List<string> { "TestMeta=42" }, StepMetaData = new Dictionary<JobStepId, List<string>> { { job.Batches[0].Steps[0].Id, new List<string> { "TestStepMeta=2025"} } } };
+			PutJobMetadataRequest metaRequest = new PutJobMetadataRequest() { JobMetaData = new List<string> { "TestMeta=42" }, StepMetaData = new Dictionary<string, List<string>> { { job.Batches[0].Steps[0].Id.ToString(), new List<string> { "TestStepMeta=2025"} } } };
 
 			await JobsController.UpdateJobMetaDataAsync(jobId, metaRequest);
 
diff --git a/Engine/Source/Programs/Horde/Plugins/Build/HordeServer.Build/Jobs/JobsController.cs b/Engine/Source/Programs/Horde/Plugins/Build/HordeServer.Build/Jobs/JobsController.cs
--- a/Engine/Source/Programs/Horde/Plugins/Build/HordeServer.Build/Jobs/JobsController.cs
+++ b/Engine/Source/Programs/Horde/Plugins/Build/HordeServer.Build/Jobs/JobsController.cs
@@ -413,7 +413,18 @@ namespace HordeServer.Jobs
 				return Forbid(JobAclAction.UpdateJob, jobId);
 			}
 
-			await _jobService.UpdateMetadataAsync(jobId, request.JobMetaData, request.StepMetaData, cancellationToken);
+			Dictionary<JobStepId, List<string>>? stepMetaData = null;
+
+			if (request.StepMetaData != null)
+			{
+				stepMetaData = new Dictionary<JobStepId, List<string>>();
+				foreach(KeyValuePair<string, List<string>> entry in request.StepMetaData)
+				{
+					stepMetaData[JobStepId.Parse(entry.Key)] = entry.Value;
+				}
+			}
+
+			await _jobService.UpdateMetadataAsync(jobId, request.JobMetaData, stepMetaData, cancellationToken);
 			
 			return Ok();
 		}
diff --git a/Engine/Source/Programs/Shared/EpicGames.Horde/Jobs/JobMessages.cs b/Engine/Source/Programs/Shared/EpicGames.Horde/Jobs/JobMessages.cs
--- a/Engine/Source/Programs/Shared/EpicGames.Horde/Jobs/JobMessages.cs
+++ b/Engine/Source/Programs/Shared/EpicGames.Horde/Jobs/JobMessages.cs
@@ -436,9 +436,9 @@ namespace EpicGames.Horde.Jobs
 		public List<string>? JobMetaData { get; set; }
 
 		/// <summary>
-		/// Step meta data to append
+		/// Step meta data to append, JobStepId => list of meta data strings
 		/// </summary>
-		public Dictionary<JobStepId, List<string>>? StepMetaData { get; set; }
+		public Dictionary<string, List<string>>? StepMetaData { get; set; }
 	}
 
 	/// <summary>
