# Introduce Test Automation Hub v2
- Implement Test Automation Hub v2 landing page
- Add TestDataV2 backend interface
- Implement an abstraction layers for handling test session results, metadata, and test reference
- Implement a registration system for support multi-versioning
- Filter data on the client side (to avoid repeating query)

#rnx

diff --git a/Engine/Source/Programs/Horde/HordeDashboard/plugins/registry.ts b/Engine/Source/Programs/Horde/HordeDashboard/plugins/registry.ts
--- a/Engine/Source/Programs/Horde/HordeDashboard/plugins/registry.ts
+++ b/Engine/Source/Programs/Horde/HordeDashboard/plugins/registry.ts
@@ -1,2 +1,4 @@
 
 import "./analytics"
+
+import "./test-automation"
diff --git a/Engine/Source/Programs/Horde/HordeDashboard/plugins/test-automation/api.ts b/Engine/Source/Programs/Horde/HordeDashboard/plugins/test-automation/api.ts
new file mode 100644
--- /dev/null
+++ b/Engine/Source/Programs/Horde/HordeDashboard/plugins/test-automation/api.ts
@@ -0,0 +1,441 @@
+import backend from "horde/backend";
+
+/// A test data response
+export type GetTestDataResponse = {
+
+	/// The unique id of the test data
+	id: string;
+
+	/// The key associated with this test data
+	key: string;
+
+	/// The change id related to the job
+	change: number;
+
+	/// The job id of this test data
+	jobId: string;
+
+	/// The step id of this test data
+	stepId: string;
+
+	/// The stream id of this test data
+	streamId: string;
+
+	/// The template ref id of the related job
+	templateRefId: string;
+}
+
+/// A test data
+export type TestData = GetTestDataResponse & {
+    data: object;
+}
+
+/// A test name request
+export type GetTestsRequest = {
+
+    /// The id of the tests to request information from
+    testIds: string[];
+}
+
+/// A test metadata request
+export type GetMetadataRequest = {
+
+    /// The ids of the test metadata to request information from
+    metadataIds?: string[];
+
+    /// The entries to search for
+    entries?: Record<string, string>;
+}
+
+/// A test tag request
+export type GetTestTagsRequest = {
+
+    /// The ids of the test tags to request information from
+    tagIds?: string[];
+
+    /// The tag names to search for
+    tagNames?: string[];
+}
+
+/// A test phase request
+export type GetTestPhasesRequest = {
+
+    /// The id of the tests to request information from
+    testIds?: string[];
+
+    /// The key of the tests to request information from
+    testKeys?: string[];
+}
+
+/// Stream ids to get test and meta collections
+export type GetStreamTestsRequest = {
+
+    /// The id of the streams to request information from
+    streamIds: string[];
+}
+
+/// Test session reference
+export type GetTestSessionsRequest = {
+
+    /// The ids of the streams to request information from
+    streamIds: string[];
+
+    /// The ids of the tests to request information from
+    testIds?: string[];
+
+    /// The ids of the metadata to request information from
+    metaIds?: string[];
+}
+
+/// Test phase session reference
+export type GetPhaseSessionsRequest = {
+
+    /// The id of the streams to request information from
+    streamIds: string[];
+
+    /// The id of the test phases to request information from
+    phaseIds?: string[];
+}
+
+/// Test name reference
+export type GetTestNameResponse = {
+
+    /// The id of the test
+    id: string;
+
+    /// The name of the test
+    name: string;
+
+    /// The key associated with this test
+    key: string;
+}
+
+/// Test metadata reference
+export type GetTestMetadataResponse = {
+
+    /// The id of the test metadata
+    id: string;
+
+    /// The key/value entries associated with the metadata
+    entries: Record<string, string>;
+}
+
+/// Test tag reference
+export type GetTestTagResponse = {
+
+    /// The id of the test tag
+    id: string;
+
+    /// The name of the test tag
+    name: string;
+}
+
+/// Describe the tests running in a stream
+export type GetTestSessionStreamResponse = {
+
+    /// The stream id
+    streamId: string;
+
+    /// Individual tests which run in the stream
+    tests: GetTestNameResponse[];
+
+    /// Test metadata that run in the stream
+    testMetadata: GetTestMetadataResponse[];
+
+    /// Test tags that run in the stream
+    testTags: GetTestTagResponse[];
+}
+
+/// Test phase reference
+export type GetPhaseResponse = {
+
+    /// The id of the test phase
+    id: string;
+
+    /// The name of the test phase
+    name: string;
+
+    /// The key associated with this test phase
+    key: string;
+}
+
+/// Descibe the phases running in a test
+export type GetTestPhaseResponse = {
+
+    /// The test id
+    testId: string;
+
+    /// The test name
+    testName: string;
+
+    /// The test key
+    testKey: string;
+
+    /// The test phases
+    testPhases: GetPhaseResponse[];
+}
+
+/// Test outcome
+export enum TestOutcome {
+
+    /// The test was successful
+    Success = "Success",
+
+    /// The test failed
+    Failure = "Failure",
+
+    /// The test was skipped
+    Skipped = "Skipped",
+
+    /// The test had an unspecified result
+    Unspecified = "Unspecified"
+}
+
+export enum TestPhaseOutcome {
+
+    // The test phase failed
+    Failed = "Failed",
+
+    /// The test phase was successful
+    Success = "Success",
+
+    /// The test phase was not run
+    NotRun = "NotRun",
+
+    /// The test phase was interrupted
+    Interrupted = "Interrupted",
+
+    /// The test phase was skipped
+    Skipped = "Skipped",
+
+    /// The test phase outcome is unknown
+    Unknown = "Unknown"
+}
+
+/// Describes test session outcome
+export type GetTestSessionResponse = {
+
+    /// The test session id
+    id: string;
+
+    /// The stream associated with the test session
+    streamId: string;
+
+    /// The metadata associated with the test session
+    metadataId: string;
+
+    /// The test data id associated with the test session
+    testDataId: string;
+
+    /// The commit id associated with the test session
+    commitId: string;
+
+    /// The commit order associated with the test session
+    commitOrder: number;
+
+    /// The tag ids associated with the test session
+    tagIds?: string[]
+
+    /// The duration of the test session
+    duration: string;
+
+    /// The date time at which the test session started
+    startDateTime: string;
+
+    /// The test name reference associated with the test session
+    nameRef: string;
+
+    /// The outcome of the test session
+    outcome: TestOutcome;
+
+    /// The total cound of phases
+    phasesTotalCount: number;
+
+    /// The count of succeeded phases
+    phasesSucceededCount: number;
+
+    /// The count of undefined phases
+    phasesUndefinedCount: number;
+
+    /// The count of failed phases
+    phasesFailedCount: number;
+
+    /// The job id associated with the test session
+    jobId: string;
+
+    /// The step job id associated with the test session
+    stepId: string;
+}
+
+/// Describes test phase session outcome
+export type GetTestPhaseSessionResponse = {
+
+    /// The phase session id
+    id: string;
+
+    /// The stream associated with the phase session
+    streamId: string;
+
+    /// The metadata associated with the phase session
+    metadataId: string;
+
+    /// The test session id associated with the phase session
+    sessionId: string;
+
+    /// The commit id associated with the phase session
+    commitId: string;
+
+    /// The commit order associated with the phase session
+    commitOrder: number;
+
+    /// The duration of the phase session
+    duration: string;
+
+    /// The date time at which the phase session started
+    startDateTime: string;
+
+    /// The test phase reference associated with the phase session
+    phaseRef: string;
+
+    /// The outcome of the phase session
+    outcome: TestPhaseOutcome;
+
+    /// The job id associated with the phase session
+    jobId: string;
+
+    /// The step job id associated with the phase session
+    stepId: string;
+
+    /// The path to the event stream
+    eventStreamPath?: string;
+
+    /// Whether the test phase encountered at least one warning
+    hasWarning?: boolean;
+
+    /// The error finerprint from the event stream
+    errorFingerprint?: string;
+}
+
+export async function getTestDataV2(id: string): Promise<TestData> {
+    return new Promise<TestData>((resolve, reject) => {
+        backend.fetch.get(`/api/v3/testdata/${id}`, {
+            params: {}
+        }).then((value) => {
+            resolve(value.data as TestData);
+        }).catch((reason) => {
+            reject(reason);
+        });
+    });
+}
+
+export async function getStreamTestsV2(streamIds: string[]): Promise<GetTestSessionStreamResponse[]> {
+    const request: GetStreamTestsRequest = { streamIds: streamIds };
+    return new Promise<GetTestSessionStreamResponse[]>((resolve, reject) => {
+        backend.fetch.post(`/api/v3/testdata/streams`, request).then((value) => {
+            resolve(value.data as GetTestSessionStreamResponse[]);
+        }).catch((reason) => {
+            reject(reason);
+        });
+    });
+}
+
+export async function getTestNamesV2(testIds: string[]): Promise<GetTestNameResponse[]> {
+    const request: GetTestsRequest = { testIds: testIds };
+    return new Promise<GetTestNameResponse[]>((resolve, reject) => {
+        backend.fetch.post(`/api/v3/testdata/tests`, request).then((value) => {
+            resolve(value.data as GetTestNameResponse[]);
+        }).catch((reason) => {
+            reject(reason);
+        });
+    });
+}
+
+export async function getTestMetadataV2(metaIds?: string[], entries?: Record<string, string>): Promise<GetTestMetadataResponse[]> {
+    const request: GetMetadataRequest = { metadataIds: metaIds, entries: entries };
+    return new Promise<GetTestMetadataResponse[]>((resolve, reject) => {
+        backend.fetch.post(`/api/v3/testdata/metadata`, request).then((value) => {
+            resolve(value.data as GetTestMetadataResponse[]);
+        }).catch((reason) => {
+            reject(reason);
+        });
+    });
+}
+
+export async function getTestTagsV2(tagIds?: string[], tagNames?: string[]): Promise<GetTestTagResponse[]> {
+    const request: GetTestTagsRequest = { tagIds: tagIds, tagNames: tagNames };
+    return new Promise<GetTestTagResponse[]>((resolve, reject) => {
+        backend.fetch.post(`/api/v3/testdata/tags`, request).then((value) => {
+            resolve(value.data as GetTestTagResponse[]);
+        }).catch((reason) => {
+            reject(reason);
+        });
+    });
+}
+
+export async function getTestPhasesV2(testIds?: string[], testKeys?: string[]): Promise<GetTestPhaseResponse[]> {
+    const request: GetTestPhasesRequest = { testIds: testIds, testKeys: testKeys };
+    return new Promise<GetTestPhaseResponse[]>((resolve, reject) => {
+        backend.fetch.post(`/api/v3/testdata/phases`, request).then((value) => {
+            resolve(value.data as GetTestPhaseResponse[]);
+        }).catch((reason) => {
+            reject(reason);
+        });
+    });
+}
+
+export async function getTestSessionsV2(streamIds: string[], testIds?: string[], metaIds?: string[], minCreateTime?: string, maxCreateTime?: string, minChange?: string, maxChange?: string): Promise<GetTestSessionResponse[]> {
+    const query: any = {};
+    if (minCreateTime) {
+        query.minCreateTime = minCreateTime;
+    }
+    if (maxCreateTime) {
+        query.maxCreateTime = maxCreateTime;
+    }
+    if (minChange) {
+        query.minChange = minChange;
+    }
+    if (maxChange) {
+        query.maxChange = maxChange;
+    }
+    const request: GetTestSessionsRequest = { streamIds: streamIds };
+    if (testIds) {
+        request.testIds = testIds;
+    }
+    if (metaIds) {
+        request.metaIds = metaIds;
+    }
+    return new Promise<GetTestSessionResponse[]>((resolve, reject) => {
+        backend.fetch.post(`/api/v3/testdata/session/tests`, request, { params: query }).then((value) => {
+            resolve(value.data as GetTestSessionResponse[]);
+        }).catch((reason) => {
+            reject(reason);
+        });
+    });
+}
+
+export async function getTestPhaseSessionsV2(streamIds: string[], phaseIds?: string[], minCreateTime?: string, maxCreateTime?: string, minChange?: string, maxChange?: string): Promise<GetTestPhaseSessionResponse[]> {
+    const query: any = {};
+    if (minCreateTime) {
+        query.minCreateTime = minCreateTime;
+    }
+    if (maxCreateTime) {
+        query.maxCreateTime = maxCreateTime;
+    }
+    if (minChange) {
+        query.minChange = minChange;
+    }
+    if (maxChange) {
+        query.maxChange = maxChange;
+    }
+    const request: GetPhaseSessionsRequest = { streamIds: streamIds };
+    if (phaseIds) {
+        request.phaseIds = phaseIds;
+    }
+    return new Promise<GetTestPhaseSessionResponse[]>((resolve, reject) => {
+        backend.fetch.post(`/api/v3/testdata/session/phases`, request, { params: query }).then((value) => {
+            resolve(value.data as GetTestPhaseSessionResponse[]);
+        }).catch((reason) => {
+            reject(reason);
+        });
+    });
+}
\ No newline at end of file
diff --git a/Engine/Source/Programs/Horde/HordeDashboard/plugins/test-automation/index.tsx b/Engine/Source/Programs/Horde/HordeDashboard/plugins/test-automation/index.tsx
new file mode 100644
--- /dev/null
+++ b/Engine/Source/Programs/Horde/HordeDashboard/plugins/test-automation/index.tsx
@@ -0,0 +1,13 @@
+import { registerHordePlugin, MountType } from "..";
+import { TestAutomationView } from "./testAutomationView";
+
+registerHordePlugin({
+   id: "test-automation",
+   routes: [{ path: "test-automation", element: <TestAutomationView /> }],
+   mount: {
+      type: MountType.TopNav,
+      context: "Tools",
+      text: "Test Automation Hub(v2)",
+      route: `/test-automation`
+   }
+})
\ No newline at end of file
diff --git a/Engine/Source/Programs/Horde/HordeDashboard/plugins/test-automation/testAutomationCommon.tsx b/Engine/Source/Programs/Horde/HordeDashboard/plugins/test-automation/testAutomationCommon.tsx
new file mode 100644
--- /dev/null
+++ b/Engine/Source/Programs/Horde/HordeDashboard/plugins/test-automation/testAutomationCommon.tsx
@@ -0,0 +1,137 @@
+import { mergeStyles, mergeStyleSets } from "@fluentui/react";
+import { TestSessionResult } from "./testData";
+import dashboard, { StatusColor } from "horde/backend/Dashboard";
+
+const styles = mergeStyleSets({
+
+   stripes: {
+      backgroundImage: 'repeating-linear-gradient(-45deg, rgba(255, 255, 255, .2) 25%, transparent 25%, transparent 50%, rgba(255, 255, 255, .2) 50%, rgba(255, 255, 255, .2) 75%, transparent 75%, transparent)',
+   }
+
+});
+
+export type StatusBarStack = {
+   value: number,
+   title?: string,
+   titleValue?: number,
+   color?: string,
+   onClick?: () => void,
+   stripes?: boolean,
+
+}
+
+export const StatusBar = (stack: StatusBarStack[], width: number, height: number, basecolor?: string, style?: any): JSX.Element => {
+
+   stack = stack.filter(s => s.value > 0);
+
+   const mainTitle = stack.map((item) => {
+      return item.titleValue === undefined ? `${Math.ceil(item.value)}% ${item.title}` : `${item.titleValue} ${item.title}`
+   }).join(' ');
+
+   return (
+      <div className={mergeStyles({ backgroundColor: basecolor, width: width, height: height, verticalAlign: 'middle', display: "flex" }, style)} title={mainTitle}>
+         {stack.map((item) => <span key={item.title!}
+            onClick={item.onClick}
+            className={item.stripes ? styles.stripes : undefined}
+            style={{
+               width: `${item.value}%`, height: '100%',
+               backgroundColor: item.color,
+               display: 'block',
+               cursor: item.onClick ? 'pointer' : 'inherit',
+               backgroundSize: `${height * 2}px ${height * 2}px`
+            }} />)}
+      </div>
+   );
+}
+
+export const SessionStatusBar = (session: TestSessionResult, width: number, height: number): JSX.Element | undefined => {
+   const sessionTotalCount = session.phaseSucceededCount + session.phaseFailedCount + session.phaseUnspecifiedCount;
+   const metaFailedFactor = Math.ceil(session.phaseFailedCount / (sessionTotalCount || 1) * 20) / 20;
+   const metaUnspecifiedFactor = Math.ceil(session.phaseUnspecifiedCount / (sessionTotalCount || 1) * 20) / 20;
+
+   const bDisplayBar = (metaFailedFactor + metaUnspecifiedFactor) > 0 && (metaFailedFactor + metaUnspecifiedFactor) < 1;
+   if (!bDisplayBar) return <div/>
+
+   const statusColors = dashboard.getStatusColors();
+
+   const stack: StatusBarStack[] = [
+      {
+          value: metaUnspecifiedFactor * 100,
+          title: "Unspecified",
+          color: statusColors.get(StatusColor.Unspecified)!,
+          stripes: true
+      },
+      {
+          value: metaFailedFactor * 100,
+          title: "Failure",
+          color: statusColors.get(StatusColor.Failure)!,
+          stripes: true
+      }
+  ];
+
+  return StatusBar(stack, width, height, statusColors.get(StatusColor.Success)!, { margin: '3px !important' });
+}
+
+export const StatusPie = (stack: StatusBarStack[], radius: number, basecolor?: string, style?: any): JSX.Element => {
+
+   const width = radius * 2;
+   const height = width;
+
+   stack = stack.filter(s => s.value > 0);
+   let offset = 0;
+   const stackOffsets = stack.map(s => {
+      const stepOffset = offset;
+      offset += s.value;
+      return stepOffset;
+   });
+
+   const mainTitle = stack.map((item) => {
+      return item.titleValue === undefined ? `${Math.ceil(item.value)}% ${item.title}` : `${item.titleValue} ${item.title}`
+   }).join(' ');
+
+   return (
+      <div className={mergeStyles({ width: width, height: height, verticalAlign: 'middle', display: "flex" }, style)} title={mainTitle}>
+         <svg width={width} height={height} viewBox={`0 0 64 64`}>
+            {stack.map((item, index) => <circle key={item.title!} r="25%" cx="50%" cy="50%" strokeWidth="50%" transform={`rotate(-90) translate(-64)`} fill="transparent"
+                                    stroke={item.color}
+                                    strokeDashoffset={`-${stackOffsets[index]}`}
+                                    strokeDasharray={`${item.value} 100`}
+                                    onClick={item.onClick} />)
+            }
+         </svg>
+      </div>
+   );
+}
+
+export const SessionStatusPie = (session: TestSessionResult, radius: number): JSX.Element | undefined => {
+   const sessionTotalCount = session.phaseSucceededCount + session.phaseFailedCount + session.phaseUnspecifiedCount;
+   const metaFailedFactor = Math.ceil(session.phaseFailedCount / (sessionTotalCount || 1) * 20) / 20;
+   const metaUnspecifiedFactor = Math.ceil(session.phaseUnspecifiedCount / (sessionTotalCount || 1) * 20) / 20;
+
+   const bDisplayBar = (metaFailedFactor + metaUnspecifiedFactor) > 0 && (metaFailedFactor + metaUnspecifiedFactor) < 1;
+   if (!bDisplayBar) return <div/>
+
+   const statusColors = dashboard.getStatusColors();
+
+   const stack: StatusBarStack[] = [
+      {
+          value: metaUnspecifiedFactor * 100,
+          title: "Unspecified",
+          color: statusColors.get(StatusColor.Unspecified)!,
+          stripes: true
+      },
+      {
+          value: metaFailedFactor * 100,
+          title: "Failure",
+          color: statusColors.get(StatusColor.Failure)!,
+          stripes: true
+      },
+      {
+          value: (1 - (metaFailedFactor + metaUnspecifiedFactor)) * 100,
+          title: "Passed",
+          color: statusColors.get(StatusColor.Success)!
+      }
+  ];
+
+  return StatusPie(stack, radius);
+}
diff --git a/Engine/Source/Programs/Horde/HordeDashboard/plugins/test-automation/testAutomationSummary.tsx b/Engine/Source/Programs/Horde/HordeDashboard/plugins/test-automation/testAutomationSummary.tsx
new file mode 100644
--- /dev/null
+++ b/Engine/Source/Programs/Horde/HordeDashboard/plugins/test-automation/testAutomationSummary.tsx
@@ -0,0 +1,477 @@
+import { DefaultButton, FontIcon, IconButton, mergeStyleSets, Modal, Stack, Text, HoverCard, HoverCardType, IPlainCardProps, DirectionalHint, Spinner, SpinnerSize, IContextualMenuItem, IContextualMenuProps } from "@fluentui/react";
+import { useState, useEffect } from "react";
+import { observer } from "mobx-react-lite";
+import { TestDataHandler, TestNameRef, MetadataRef, TestSessionStatus, TestSessionResult, TestStatus } from "./testData";
+import { TestOutcome } from './api';
+import dashboard, { StatusColor } from "horde/backend/Dashboard";
+import { getShortNiceTime } from "horde/base/utilities/timeUtils";
+import { getHordeStyling } from "horde/styles/Styles";
+import { projectStore } from 'horde/backend/ProjectStore';
+import { SessionStatusBar, StatusBar, StatusBarStack } from "./testAutomationCommon";
+import { TestHistoryGraph } from "./testHistoryGraph";
+
+type SessionCallback = (session: TestSessionResult) => void;
+type HistorySessionCallback = (test: TestNameRef, onClick: SessionCallback) => void;
+
+const TestHistoryGraphWidget: React.FC<{ test: TestNameRef, stream: string, handler: TestDataHandler, onClick: SessionCallback, commonMetaKeys?: string[]}> = ({ test, stream, handler, onClick, commonMetaKeys }) => {
+    const [container, setContainer] = useState<HTMLDivElement | null>(null);
+    const [state, setState] = useState<{ graph?: TestHistoryGraph }>({});
+
+    useEffect(() => {
+        const setup = async () => {
+            const graph = new TestHistoryGraph(test, stream, handler, commonMetaKeys, onClick);
+            setState({ graph: graph });
+        }
+        setup();
+    }, []);
+
+    if (!state.graph) {
+        return null;
+    }
+
+    if (container) {
+        try {
+            state.graph?.render(container);
+        } catch (err) {
+            console.error(err);
+        }
+
+    }
+
+    const { hordeClasses } = getHordeStyling();
+    const graph_container_id = `test_history_graph_container_${test.id}_${stream}`;
+
+    return <Stack className={hordeClasses.horde}>
+                <Stack style={{ paddingLeft: 8, paddingTop: 8 }}>
+                    <div id={graph_container_id} className="horde-no-darktheme" style={{ userSelect: "none", position: "relative" }} ref={(ref: HTMLDivElement) => setContainer(ref)}/>
+                </Stack>
+            </Stack>
+}
+
+
+const StreamSelector: React.FC<{ streams: string[], onClick: (stream: string) => void }> = ({ streams, onClick }) => {
+    const options: IContextualMenuItem[] = [];
+    streams.forEach(a => {
+        options.push({
+            key: `streams_${a}`,
+            text: projectStore.streamById(a)?.fullname ?? a,
+            onClick: () => onClick(a)
+        });
+    });
+
+    const menuProps: IContextualMenuProps = {
+        shouldFocusOnMount: true,
+        subMenuHoverDelay: 0,
+        items: options,
+    };
+
+    return <DefaultButton style={{ height: 22, fontSize: 11, whiteSpace: 'nowrap' }} menuProps={menuProps} text="Compare with"/>
+}
+
+export const TestHistoryView: React.FC<{ test: TestNameRef, handler: TestDataHandler, onClick: SessionCallback  }> = observer(({ test, handler, onClick }) => {
+
+    const { hordeClasses } = getHordeStyling();
+    const [streams, setStreams] = useState<{current: string, toCompare?: string}>({current: handler.stream!});
+
+    handler.subscribeToSubQueryLoading();
+
+    const availableStreams = handler.availableStreams.filter((s) => s !== streams.current);
+    let anyFailed = false;
+
+    // collecting metadata
+    let metadataSet = new Set(handler.getStatusStream(streams.current)!.tests.get(test)!.getMetadata());
+    if (!!streams.toCompare) {
+        const toCompareMetadataList = handler.getStatusStream(streams.toCompare!)?.tests.get(test)?.getMetadata();
+        if (!!toCompareMetadataList) {
+            metadataSet = metadataSet.union(new Set(toCompareMetadataList!));
+        }
+    }
+
+    const commonMetaKeys = MetadataRef.identifyCommonKeys(metadataSet.values().toArray());
+    const commonMetaValues = metadataSet.values().next().value?.getCommonValues(commonMetaKeys).join(" / ");
+
+    const scolors = dashboard.getStatusColors();
+    const streamViews = (!!availableStreams.length? [streams.current, streams.toCompare] : [streams.current]).map(stream => {
+
+        const streamName = projectStore.streamById(stream)?.fullname ?? stream;
+
+        let anyMetaFailed = false;
+        const status = stream && handler.getStatusStream(stream)?.tests.get(test);
+        if (!!status) {
+            for (const [, lastMetaSession] of status!.getMetadataLastSessions()) {
+                if (lastMetaSession?.outcome !== TestOutcome.Success) {
+                    anyMetaFailed = true;
+                    anyFailed = true;
+                }
+            }
+        }
+
+        return <Stack key={`test_history_view_${stream}_${test.id}`}>
+                { !!stream &&
+                    <Stack horizontal>
+                        <Stack className="horde-no-darktheme" style={{ paddingTop: 2, paddingRight: 3 }}>
+                            <FontIcon style={{ fontSize: 11, color: !!status? anyMetaFailed ? scolors.get(StatusColor.Failure)! : scolors.get(StatusColor.Success)! : scolors.get(StatusColor.Unspecified) }} iconName="Square" />
+                        </Stack>
+                        <Stack>
+                            <Text style={{ fontSize: 11 }}>{streamName}</Text>
+                        </Stack>
+                    </Stack>
+                }
+                { !!status &&
+                    <Stack style={{ paddingTop: 4 }}>
+                        <TestHistoryGraphWidget test={test} stream={stream} commonMetaKeys={commonMetaKeys} handler={handler} onClick={onClick} />
+                    </Stack>
+                }
+                { !status && !!stream &&
+                    <Stack horizontalAlign='center' style={{ padding: 4 }}>
+                        <Text style={{ fontSize: 12 }}>No Results</Text>
+                    </Stack>
+                }
+            </Stack>
+    });
+
+    let name = test.name;
+    if (commonMetaValues) {
+        name += ` - ${commonMetaValues}`
+    }
+
+    return <Stack>
+            <Stack className={hordeClasses.raised}>
+                <Stack style={{ paddingBottom: 12 }}>
+                    <Stack horizontal>
+                        <Stack className="horde-no-darktheme" style={{ paddingTop: 3, paddingRight: 4 }}>
+                            <FontIcon style={{ color: anyFailed ? scolors.get(StatusColor.Failure)! : scolors.get(StatusColor.Success)! }} iconName="Square" />
+                        </Stack>
+                        <Stack>
+                            <Text style={{ fontFamily: "Horde Open Sans Semibold" }} variant='medium'>{name}</Text>
+                        </Stack>
+                    </Stack>
+                </Stack>
+                <Stack style={{ paddingLeft: 12 }} tokens={{ childrenGap: 4 }}>
+                    {streamViews}
+                </Stack>
+                {handler.subQueryLoading &&
+                    <Stack horizontalAlign='center' style={{ padding: 12}} tokens={{ childrenGap: 10 }}>
+                        <Text style={{ fontSize: 12 }}>Loading Data</Text>
+                        <Spinner size={SpinnerSize.small} />
+                    </Stack>
+                }
+                <Stack horizontal>
+                    <Stack style={{ width: 170 }}>
+                        <StreamSelector streams={availableStreams} onClick={(targetStream: string) => {
+                            setStreams({...streams, toCompare: undefined});
+                            handler.query(targetStream, true).catch((reason) => console.error(reason)).finally(() => setStreams({...streams, toCompare: targetStream}))
+                        }} />
+                    </Stack>
+                </Stack>
+            </Stack>
+        </Stack>
+});
+
+const TestHistoryModal: React.FC<{ test: TestNameRef, handler: TestDataHandler, onDismiss: () => void, onDetailsClick: (test: TestSessionResult) => void }> = ({ test, handler, onDismiss, onDetailsClick }) => {
+    return <Modal isOpen={true} isBlocking={true} topOffsetFixed={true} styles={{ main: { padding: 8, width: 1134, hasBeenOpened: false, top: "80px", position: "absolute" } }} onDismiss={() => onDismiss()} >
+            <Stack grow>
+                <Stack horizontal verticalAlign="center">
+                    <Stack><Text style={{ paddingLeft: 8, fontSize: 16, fontWeight: 600 }}>{`${test.name} History`}</Text></Stack>
+                    <Stack grow />
+                    <Stack style={{ paddingBottom: 4 }}>
+                        <IconButton
+                            iconProps={{ iconName: 'Cancel' }}
+                            ariaLabel="Close popup modal"
+                            onClick={() => { onDismiss(); }}
+                        />
+                    </Stack>
+                </Stack>
+                <Stack tokens={{ childrenGap: 40 }} style={{ padding: 8, overflow: "auto", maxHeight: 'calc(100vh - 200px)'} }>
+                    <TestHistoryView test={test} handler={handler} onClick={onDetailsClick} />
+                </Stack>
+            </Stack>
+        </Modal>
+}
+
+const TestSummaryCard: React.FC<{ test: TestNameRef, sessions: [MetadataRef, TestSessionResult][], metaKeyMask: string[] | undefined, handler: TestDataHandler, onDetailsClick: SessionCallback, onHistoryClick: HistorySessionCallback }> = ({ test, sessions, metaKeyMask, handler, onDetailsClick, onHistoryClick }) => {
+    const [expanded, setExpanded] = useState(false);
+ 
+    const { hordeClasses, modeColors } = getHordeStyling();
+    const statusColors = dashboard.getStatusColors();
+ 
+    const styles = mergeStyleSets({
+       metaitem: {
+            selectors: {
+                ':hover': {
+                    filter: dashboard.darktheme ? "brightness(120%)" : "brightness(90%)"
+                }
+            }
+       }
+    });
+ 
+    const colorA = dashboard.darktheme ? "#181A1B" : "#e8e8e8";
+    const colorB = dashboard.darktheme ? "#242729" : "#f8f8f8";
+ 
+    const metaElements: JSX.Element[] = [];
+    const onRenderStatusCard = (session: TestSessionResult): JSX.Element => {
+        return <Stack key={`test_summary_hover_${test.id}`} style={{pointerEvents: 'none'}}>
+                    <Text style={{fontSize: 12, padding: '8px'}}>
+                        {!!session.phaseUnspecifiedCount && <span style={{color: statusColors.get(StatusColor.Unspecified)}}> Unspecified: {session.phaseUnspecifiedCount}</span>}
+                        {!!session.phaseFailedCount && <span style={{color: statusColors.get(StatusColor.Failure)}}> Failed: {session.phaseFailedCount}</span>}
+                        {!!session.phaseSucceededCount && <span style={{color: statusColors.get(StatusColor.Success)}}> Success: {session.phaseSucceededCount}</span>}
+                    </Text>
+                </Stack>
+    }
+ 
+    let testFailed = 0;
+    let testUnspecified = 0;
+    let testTotal = 0;
+ 
+    sessions.sort(([, asession], [, bsession]) => bsession.commitOrder - asession.commitOrder);
+    const latestCommit = sessions.reduce((latest, [, session]) => session.commitOrder > latest.order? {order: session.commitOrder, id: session.commitId} : latest, {order: 0, id: "NA"});
+ 
+    const commonKeys = MetadataRef.identifyCommonKeys(sessions.map(([meta,]) => meta));
+    const commonMetaString  = sessions.map(([meta,]) => meta).find(() => true)?.getCommonValues(commonKeys)?.join(" / ");
+    const maskMetaString  = sessions.map(([meta,]) => meta).find(() => true)?.getCommonValues(metaKeyMask)?.join(" / ");
+ 
+    sessions.forEach(([meta, session]) => {
+        const outcome = session.outcome;
+        const metaName = meta.getValuesExcept(commonKeys).join(" / ") || (meta.getValues().findLast(() => true) ?? test.name);
+    
+        const sessionTotalCount = session.phaseSucceededCount + session.phaseFailedCount + session.phaseUnspecifiedCount;
+        testUnspecified += session.phaseUnspecifiedCount;
+        testFailed += session.phaseFailedCount;
+        testTotal += sessionTotalCount;
+    
+        let color = statusColors.get(StatusColor.Success);
+        if (outcome === TestOutcome.Failure) {
+            color = statusColors.get(StatusColor.Failure);
+        } else if (outcome === TestOutcome.Unspecified) {
+            color = statusColors.get(StatusColor.Unspecified);
+        }
+    
+        let dateString = "";
+    
+        let date = handler.commitIdDates.get(session.id);      
+        if (!date) {
+            date = session.start;
+        }               
+        if (date) {
+            dateString = getShortNiceTime(date, false, false)
+        }
+    
+        let fontWeight: number | undefined;
+        if (session.commitId === latestCommit.id) {
+            fontWeight = 600;
+        }
+    
+        const StatusCardProps: IPlainCardProps = {
+            onRenderPlainCard: onRenderStatusCard,
+            renderData: session,
+            directionalHint: DirectionalHint.rightCenter,
+            calloutProps: { isBeakVisible: true }
+        }
+    
+        metaElements.push(
+            <Stack key={`test_summary_button_${session.id}`} className="horde-no-darktheme" style={{ cursor: "pointer", backgroundColor: metaElements.length % 2 ? colorA : colorB, paddingTop: 2, paddingBottom: 2 }}
+                onClick={(ev) => {
+                    ev.stopPropagation();
+                    onDetailsClick(session);
+                }}>
+                <HoverCard plainCardProps={StatusCardProps} type={HoverCardType.plain} cardOpenDelay={20} sticky={false}>
+                    <Stack className={styles.metaitem} horizontal verticalAlign="center">
+                        <Stack className="horde-no-darktheme" style={{ paddingLeft: 0, paddingTop: 1, paddingRight: 4 }}>
+                            <FontIcon style={{ fontSize: 11, color: color }} iconName="Square" />
+                        </Stack>
+                        <Stack horizontal>
+                            <Stack style={{ width: 163 }}>
+                                <Text variant="xSmall" style={{ fontWeight: fontWeight, textOverflow: 'ellipsis', overflow: 'hidden', whiteSpace: 'nowrap' }} title={metaName}>{metaName}</Text>
+                            </Stack>
+                            <Stack horizontal style={{ width: 50, paddingLeft: 3 }} verticalFill verticalAlign="center" tokens={{ childrenGap: 4 }}>
+                                <Text variant="xSmall" style={{ fontWeight: fontWeight }}>{session.commitId}</Text>
+                                {session.commitId === latestCommit.id && <FontIcon style={{ fontSize: 11, color: color }} iconName="Star" />}
+                            </Stack>
+                            <Stack style={{ width: 60 }} horizontalAlign="end">
+                                <Text variant="xSmall" style={{ fontWeight: fontWeight }}>{dateString}</Text>
+                            </Stack>
+                        </Stack>
+                    </Stack>
+                    {SessionStatusBar(session, 294, 8)}
+                </HoverCard>
+            </Stack>
+        );
+    });
+ 
+    const testFailedFactor = Math.ceil(testFailed / (testTotal || 1) * 20) / 20;
+    const testUnspecifiedFactor = Math.ceil(testUnspecified / (testTotal || 1) * 20) / 20;
+    const testStack: StatusBarStack[] = [
+        {
+            value: testUnspecifiedFactor * 100,
+            title: "Unspecified",
+            color: statusColors.get(StatusColor.Unspecified)!,
+            stripes: true
+        },
+        {
+            value: testFailedFactor * 100,
+            title: "Failure",
+            color: statusColors.get(StatusColor.Failure)!,
+            stripes: true
+        }
+    ]
+ 
+    return <Stack className={hordeClasses.raised} style={{ cursor: "pointer", height: "fit-content", backgroundColor: dashboard.darktheme ? "#242729" : modeColors.background, padding: 12, width: 332 }} onClick={() => setExpanded(!expanded)}>
+            <Stack horizontal verticalAlign="center" >
+                <Stack className="horde-no-darktheme" style={{ paddingTop: 1, paddingRight: 4 }}>
+                    <FontIcon style={{ fontSize: 13 }} iconName={expanded ? "ChevronDown" : "ChevronRight"} />
+                </Stack>
+                <Stack>
+                    <Text style={{maxWidth: 240, fontWeight: 600, textOverflow: 'ellipsis', overflow: 'hidden', whiteSpace: 'nowrap' }}  title={`${test.name} - ${commonMetaString}`}>{test.name}{maskMetaString && <span style={{fontSize: 11}}> - {maskMetaString}</span>}</Text>
+                </Stack>
+        
+                <Stack grow />
+        
+                <Stack horizontalAlign="end">
+                    <DefaultButton style={{ minWidth: 52, fontSize: 10, padding: 0, height: 18 }} onClick={(ev) => {
+                        ev.stopPropagation();
+                        onHistoryClick(test, onDetailsClick);
+                    }}>History</DefaultButton>
+                </Stack>
+            </Stack>
+            <Stack>
+                <Stack style={{ paddingTop: 4 }}>
+                    {StatusBar(testStack, 308, 10, statusColors.get(StatusColor.Success)!, { margin: '3px !important' })}
+                </Stack>
+        
+            </Stack>
+            {!!expanded &&
+                <Stack style={{ paddingLeft: 8 }} tokens={{ childrenGap: 4 }}>
+                    {metaElements}
+                </Stack>
+            }
+            </Stack>
+}
+
+const filterSessionsWithMetaMask = (testStatus: TestStatus, filteredTests: TestNameRef[], handler: TestDataHandler): [[TestNameRef, [MetadataRef, TestSessionStatus][]][], Map<string, string[]>] => {
+    const filteredMetadata = handler.filteredMetadata;
+
+    const testMeta: Map<string, [TestNameRef, MetadataRef[]][]> = new Map();
+
+    const failing: [TestNameRef, [MetadataRef, TestSessionStatus][]][] = [];
+    const passing: [TestNameRef, [MetadataRef, TestSessionStatus][]][] = [];
+
+    filteredTests.forEach(test => {
+        // sort by failed
+        const status = testStatus!.tests.get(test)!;
+        let anyFailed = false;
+
+        const metaRefs: MetadataRef[] = [];
+        const testRefMeta: [TestNameRef, MetadataRef[]][] = testMeta.get(test.name) ?? testMeta.set(test.name, []).get(test.name)!;
+        testRefMeta.push([test, metaRefs]);
+
+        const metaResults: [MetadataRef, TestSessionStatus][] = [];
+
+        filteredMetadata.forEach(meta => {
+            const metaStatus = status.sessions.get(meta);
+            if (!metaStatus) return;
+
+            const session = metaStatus.getLastSession();
+            if (!session) return;
+
+            const outcome = session.outcome;
+            if ( outcome === TestOutcome.Failure || outcome === TestOutcome.Unspecified ) {
+                anyFailed = true;
+            }
+
+            metaResults.push([meta, metaStatus]);
+            metaRefs.push(meta);
+        });
+
+        if (anyFailed) {
+            failing.push([test, metaResults]);
+        } else {
+            passing.push([test, metaResults]);
+        }
+    });
+
+    // Convoluted way to get unique meta identifier for each test with the same name
+    const testMetaKeys = new Map(testMeta.entries().filter(([, testRefs]) => testRefs.length > 1).map(([testName, testRefs]) => {
+        const testCommonKeys: string[][] = testRefs.map(([, metaRefs]) => MetadataRef.identifyCommonKeys(metaRefs));
+        const commonKeys: string[] = MetadataRef.identifyCommonKeys(testRefs.map(([, metaRefs]) => metaRefs).flat());
+        const uniqueKeys = testCommonKeys.reduce(
+            (current, testCommonKey) => current.filter((key) => testCommonKey.filter(k => !commonKeys.includes(k)).includes(key)),
+            testCommonKeys.shift() ?? []
+        );
+        return [testName, uniqueKeys];
+    }));
+
+    return [[...failing, ...passing], testMetaKeys];
+}
+
+export const TestSummary: React.FC<{ handler: TestDataHandler }> = ({ handler }) => {
+    const [historyState, setHistoryState] = useState<{test?: TestNameRef, onClick?: SessionCallback}>({});
+
+    const searchTestId = handler.getSearchParam('history');
+    if (!searchTestId) {
+        historyState.test = historyState.onClick = undefined;
+    }
+
+    const testStatus = handler.selectedStatusStream;
+    if (!testStatus?.tests.size) {
+        return null;
+    }
+
+    const filteredTests = handler.filteredTests;
+    if (!filteredTests.length) {
+        return null;
+    }
+
+    const [testSessions, testMetaKeys] = filterSessionsWithMetaMask(testStatus, filteredTests, handler);
+
+    const tests = testSessions.map(([test, sessions]) => {
+
+        const onDetailsClick = (session: TestSessionResult) => {
+            if (session.jobId && session.stepId) {
+                window.open(`/job/${session.jobId}?step=${session.stepId}`, '_blank');
+            }
+        }
+
+        if (searchTestId && test.id === searchTestId) {
+            historyState.test = test;
+            historyState.onClick = onDetailsClick;
+        }
+
+        return <TestSummaryCard key={`test_summary_${test.id}`}
+                    test={test}
+                    sessions={sessions.map(([meta, status]) => [meta, status.getLastSession()] as [MetadataRef, TestSessionResult])}
+                    metaKeyMask={testMetaKeys.get(test.name)} handler={handler}
+                    onHistoryClick={(test: TestNameRef, onClick: SessionCallback) => {
+                        setHistoryState({test: test, onClick: onClick});
+                        handler.setSearchParam('history', test.id);
+                    }}
+                    onDetailsClick={onDetailsClick}
+                />
+
+    });
+
+    const { hordeClasses } = getHordeStyling();
+
+    return <Stack className={hordeClasses.raised}>
+            {!!historyState.test &&
+                <TestHistoryModal
+                    test={historyState.test}
+                    handler={handler}
+                    onDismiss={() => {
+                        setHistoryState({});
+                        handler.removeSearchParam('history');
+                    }}
+                    onDetailsClick={historyState.onClick!}
+                />
+            }
+            <Stack style={{ paddingBottom: 12 }}>
+                <Stack>
+                    <Text style={{ fontFamily: "Horde Open Sans Semibold" }} variant='medium'>Latest Test Results</Text>
+                </Stack>
+            </Stack>
+            <Stack style={{ paddingLeft: 12 }} tokens={{ childrenGap: 4 }}>
+                <Stack wrap horizontal style={{ width: "100%" }} tokens={{ childrenGap: 12 }}>
+                    {tests}
+                </Stack>
+            </Stack>
+        </Stack>
+}
\ No newline at end of file
diff --git a/Engine/Source/Programs/Horde/HordeDashboard/plugins/test-automation/testAutomationView.tsx b/Engine/Source/Programs/Horde/HordeDashboard/plugins/test-automation/testAutomationView.tsx
new file mode 100644
--- /dev/null
+++ b/Engine/Source/Programs/Horde/HordeDashboard/plugins/test-automation/testAutomationView.tsx
@@ -0,0 +1,334 @@
+import { Stack, FontIcon, IComboBox, IComboBoxOption, ComboBox, IContextualMenuItem, IContextualMenuProps, Text, Spinner, SpinnerSize, Label, DefaultButton, SelectableOptionMenuItemType, TagItemSuggestion, ITag, TagPicker } from "@fluentui/react";
+import { Breadcrumbs } from "horde/components/Breadcrumbs";
+import { TopNav } from "horde/components/TopNav";
+import { getHordeStyling } from "horde/styles/Styles";
+import { observer } from "mobx-react-lite";
+import { useSearchParams } from "react-router-dom";
+import React, { useEffect, useState, useRef } from "react";
+import { projectStore } from 'horde/backend/ProjectStore';
+import { TestDataHandler, TestDataVersionRegistrar } from "./testData";
+import { TestSummary } from "./testAutomationSummary";
+
+import { TestDataV2 } from "./testDataV2Fetcher";
+
+const handler = new TestDataHandler();
+
+// register TestDataV2 fetcher
+TestDataVersionRegistrar.register(new TestDataV2);
+
+const MultiOptionChooser: React.FC<{options: IComboBoxOption[], initialSelection: string[], updateKeys: (selectedKeys: string[]) => void }> = ({ options, initialSelection, updateKeys }) => {
+    const [selection, setSelection] = useState<string[]>(initialSelection);
+    const optionsWithSelectAll = [...options];
+
+    useEffect(() => {
+        // keep input selection in sync
+        if (selection.toString() !== initialSelection.toString()) {
+            setSelection(initialSelection);
+        }
+    }, [initialSelection]);
+
+    if (selection.length === options.length) {
+        selection.push('selectAll');
+    }
+
+    optionsWithSelectAll.unshift({ key: 'selectAll', text: 'Select All', itemType: SelectableOptionMenuItemType.SelectAll });
+
+    return <ComboBox styles={{ root: { height: 30 } }} placeholder="None" selectedKey={selection} multiSelect options={optionsWithSelectAll}
+        onChange={(_, option) => {
+            if (option) {
+                if (option.itemType === SelectableOptionMenuItemType.SelectAll) {
+                    setSelection(option.selected? optionsWithSelectAll.map((o) => o.key as string) : []);
+                    return;
+                }
+                const key = option.key as string;
+                const index = selection.indexOf(key);
+                if (index === -1 && option.selected) {
+                    selection.push(key);
+                } else if (index >= 0 && !option.selected) {
+                    selection.splice(index, 1);
+                }
+                setSelection([...selection]);
+            }
+        }}
+        onMenuDismissed={() => {
+            updateKeys(selection.filter((k) => k !== 'selectAll'));
+        }} />
+};
+
+const StreamChooser: React.FC = observer(() => {
+
+    handler.subscribeToFilter();
+
+    const stream = handler.filterState.stream;
+
+    const options: IContextualMenuItem[] = [];
+    handler.availableStreams.forEach(a => {
+        options.push({
+            key: `streams_${a}`,
+            text: projectStore.streamById(a)?.fullname ?? a,
+            onClick: () => handler.selectStream(a)
+        });
+    });
+
+    const menuProps: IContextualMenuProps = {
+        shouldFocusOnMount: true,
+        subMenuHoverDelay: 0,
+        items: options,
+    };
+
+    return <Stack style={{ paddingTop: 12, paddingBottom: 4 }} horizontal>
+            <Label style={{ paddingTop: 5, paddingRight: 5 }}>Stream</Label>
+            <DefaultButton style={{ width: 270, textAlign: "left" }} menuProps={menuProps} text={stream ? projectStore.streamById(stream)?.fullname : "Select"} />
+        </Stack>   
+    });
+
+const onRenderSuggestionItem = (props: ITag) => {
+    return <TagItemSuggestion style={{height: 'auto', whiteSpace: 'normal', maxWidth: 500}}>{props.name}</TagItemSuggestion>
+}
+
+const getTextFromTag = (item: ITag) => item.name;
+
+const TestChooser: React.FC = observer(() => {
+    
+    handler.subscribeToFilter();
+
+    const ctests: Set<string> = new Set<string>(handler.filterState.tests ?? []);
+    const streamTests = handler.allStreamTestNames;
+
+    const [testsSelection, setTestsSelection] = useState<Set<string>>(ctests);
+
+    useEffect(() => {
+        // keep input selection in sync
+        if (testsSelection.size !== ctests.size || testsSelection.keys().some(s => !ctests.has(s))) {
+            setTestsSelection(ctests);
+        }
+    }, [handler.searchUpdated]);
+
+    if (!streamTests.length) {
+            return null;
+    }
+
+    const testTags: ITag[] = streamTests.map(t => ({key: `#${t.toLowerCase()}`, name: `#${t}`}))
+
+    const filterSelectedTests = (filterText: string, tagList: ITag[]): ITag[] => {
+        if (!filterText) return [];
+        const lowerText = filterText.toLowerCase();
+        const closestTags = testTags.filter((tag) => tag.key.toString().indexOf(lowerText) >= 0);
+        const isTagExist = closestTags.find((tag) => tag.key === lowerText);
+        if (!isTagExist) closestTags.splice(0, 0, {key: lowerText, name: filterText});
+        return closestTags;
+    }
+    const selectedItems: ITag[] = testsSelection.keys().toArray().map((key) => ({key: key.toLowerCase(), name: key} as ITag));
+
+    return <Stack style={{ paddingTop: 12, paddingBottom: 4 }} horizontal>
+            <Label style={{ paddingTop: 5, paddingRight: 5 }}>Tests</Label>
+            <Stack tokens={{ childrenGap: 6 }} style={{ paddingLeft: 12 }} grow>
+                <TagPicker
+                        onRenderSuggestionsItem={onRenderSuggestionItem}
+                        removeButtonAriaLabel="Remove"
+                        selectionAriaLabel="Selected tests"
+                        selectedItems={selectedItems}
+                        onResolveSuggestions={filterSelectedTests}
+                        getTextFromItem={getTextFromTag}
+                        onChange={(tags) => {
+                            if (!!tags) {
+                                const keys = tags!.map(getTextFromTag);
+                                keys.forEach((key) => handler.addTest(key));
+                                testsSelection.difference(new Set(keys)).values().forEach(key => handler.removeTest(key));
+                            } else {
+                                testsSelection.values().forEach(key => handler.removeTest(key));
+                            }
+                            setTestsSelection(new Set(handler.filterState.tests ?? []));
+                        }}
+                />
+                
+            </Stack>
+        </Stack>
+});
+
+const MetaKeyChooser: React.FC<{ onUpdateKey: (key: string) => void }> = observer(({ onUpdateKey }) => {
+
+    handler.subscribeToFilter();
+
+   const cmetakeys: Set<string> = new Set(handler.filterState.metadata?.keys() ?? []);
+   const streamMetadataKeys = handler.allStreamMetadataKeys;
+
+   if (!streamMetadataKeys.length) {
+        return null;
+   }
+
+   const options: IContextualMenuItem[] = [];
+   streamMetadataKeys.forEach(a => {
+        if (cmetakeys.has(a)) {
+            return;
+        }
+        options.push({
+            key: `metakey_${a}`, text: a, onClick: () => {
+                onUpdateKey(a);
+            }
+        });
+   });
+
+   const menuProps: IContextualMenuProps = {
+        shouldFocusOnMount: true,
+        subMenuHoverDelay: 0,
+        items: options,
+   };
+
+   return <Stack style={{ paddingTop: 12, paddingBottom: 4 }} horizontal>
+            <Label style={{ paddingTop: 5, paddingRight: 5, whiteSpace: 'nowrap' }}>Filter by</Label>
+            <DefaultButton style={{width: 270, textAlign: "left" }} menuProps={menuProps} text="Metadata" />
+        </Stack>   
+});
+
+const MetaValueChooser: React.FC<{ metakey: string, onRemove: (key: string) => void }> = observer(({ metakey, onRemove }) => {
+
+    handler.subscribeToFilter();
+
+   const cmetavalues: Set<string> = new Set(handler.filterState.metadata?.get(metakey) ?? []);
+   const metavalues = handler.getAllStreamMetadataValues(metakey);
+
+   if (!metavalues.length) {
+        return null;
+   }
+
+   const options: IComboBoxOption[] = metavalues.map(t => { return { key: t, text: t } });
+
+   return <Stack style={{ paddingTop: 12, paddingBottom: 4, }} horizontal>
+            <Label style={{ paddingTop: 5, paddingRight: 5 }}>{metakey}</Label>
+            <Stack style={{ padding: 1, width: 270 }}>
+                <MultiOptionChooser options={options} initialSelection={cmetavalues.keys().toArray()}
+                updateKeys={(keys) => {
+                    metavalues.forEach(t => {
+                    const selected = keys.includes(t);
+                    if (!selected && cmetavalues.has(t)) {
+                        handler.removeMetadata(metakey, t);
+                    }
+                    else if (selected && !cmetavalues.has(t)) {
+                        handler.addMetadata(metakey, t);
+                    }
+                    });
+                }} />
+        </Stack>
+        <DefaultButton style={{ minWidth: 0, fontSize: 11, padding: 5 }} title="remove filter"
+                onClick={() => {
+                    if (handler.filterState?.metadata?.keys().some((key) => key === metakey)) {
+                        handler.removeMetadata(metakey);
+                    }
+                    onRemove(metakey);
+                }}>
+            <FontIcon iconName="CalculatorMultiply" />
+        </DefaultButton>
+    </Stack>
+});
+
+
+const TestAutomationSidebarLeft: React.FC = observer(() => {
+
+    handler.subscribeToFilter();
+    const filterMetaKeys = handler.filterState.metadata?.keys().toArray().sort((a, b) => a.localeCompare(b)) ?? [];
+
+    const [metaKeys, setMetaKeys] = useState<string[]>(filterMetaKeys);
+    const { hordeClasses } = getHordeStyling();
+
+    const updateMetaKey = (key: string) => {
+        const idx = metaKeys.indexOf(key);
+        if (idx < 0) {
+            metaKeys.push(key);
+        } else {
+            metaKeys.splice(idx, 1);
+        }
+        setMetaKeys([...metaKeys]);
+    }
+
+    useEffect(() => {
+        // keep search filter in sync
+        if (metaKeys.toString() !== filterMetaKeys.toString()) {
+            setMetaKeys(filterMetaKeys);
+        }
+    }, [handler.searchUpdated]);
+
+   return <Stack style={{ width: 300, paddingRight: 18 }}>
+            <Stack className={hordeClasses.modal}>
+                <Stack key="chooser_stream">
+                    <StreamChooser />
+                </Stack>
+                <Stack key="chooser_test">
+                    <TestChooser />
+                </Stack>
+                {!!metaKeys.length && metaKeys.map((key) => <Stack key={`chooser_meta_${key}`}><MetaValueChooser metakey={key} onRemove={updateMetaKey} /></Stack>)}
+                <Stack>
+                    <MetaKeyChooser key="chooser_meta" onUpdateKey={updateMetaKey} />
+                </Stack>
+            </Stack>
+        </Stack>
+});
+
+const TestAutomationSummary: React.FC = () => {
+
+    return <Stack tokens={{ childrenGap: 24 }}>
+                <TestSummary handler={handler} />
+                {/** Additional or alternate view could be implemented here */}
+            </Stack>
+}
+
+const TestAutomationPanel: React.FC = observer(() => {
+
+    if (handler.queryLoading) {
+        return <Stack horizontalAlign='center' style={{ paddingTop: 24, width: "100%" }} tokens={{childrenGap: 8}}>
+                    <Text style={{ fontSize: 24 }}>Loading Data</Text>
+                    <Spinner size={SpinnerSize.large} />
+                </Stack>
+    }
+
+    handler.subscribeToSearch();
+
+    const status = handler.filteredTests;
+
+    return <Stack grow>
+                { !status.length &&
+                    <Stack grow horizontalAlign='center' style={{ paddingTop: 24 }}>
+                        <Text style={{ fontSize: 24 }}>No Results</Text>
+                    </Stack>
+                }
+                { !!status.length &&
+                    <TestAutomationSummary />
+                }
+            </Stack>
+
+});
+
+export const TestAutomationView: React.FC = () => {
+
+    const [searchParams, setSearchParams] = useSearchParams();
+    const [initialized, setInitialized] = useState(false);
+
+    useEffect(() => {
+        handler.initialize(searchParams.toString(), (search: string) => {setSearchParams(search)}).then(() => {setInitialized(true)});
+        return () => {
+            handler.clear();
+            setInitialized(false); // for dev env
+        };
+    }, []);
+
+    useEffect(() => {
+        initialized && handler.syncSearchParam(searchParams.toString());
+    }, [searchParams]);
+
+    const { hordeClasses, modeColors } = getHordeStyling();
+
+    return <Stack className={hordeClasses.horde} key="key_test_automation_hub">
+                <TopNav />
+                <Breadcrumbs items={[{ text: 'Test Automation Hub' }]} />
+                <Stack grow style={{ padding: 12, backgroundColor: modeColors.background, height: 'calc(100vh - 148px)', overflow: 'auto' } } horizontalAlign='center'>
+                    {initialized && !!handler.availableStreams.length &&
+                        <Stack horizontal grow>
+                            <TestAutomationSidebarLeft/>
+                            <Stack style={{ minWidth: 1024 }}>
+                                <TestAutomationPanel />
+                            </Stack>
+                        </Stack>
+                    }
+                </Stack>
+            </Stack>
+}
\ No newline at end of file
diff --git a/Engine/Source/Programs/Horde/HordeDashboard/plugins/test-automation/testData.ts b/Engine/Source/Programs/Horde/HordeDashboard/plugins/test-automation/testData.ts
new file mode 100644
--- /dev/null
+++ b/Engine/Source/Programs/Horde/HordeDashboard/plugins/test-automation/testData.ts
@@ -0,0 +1,1010 @@
+import { isNumber } from '@datadog/browser-core';
+import { action, makeObservable, observable, reaction } from 'mobx';
+import { TestOutcome } from './api';
+import { projectStore } from 'horde/backend/ProjectStore';
+
+const defaultQueryWeeks = 2;
+const filterRefreshMS = 200;
+
+type StreamId = string;
+type TestId = string;
+type MetaId = string;
+type TagId = string;
+type TestSessionId = string;
+
+// TestData Version registration
+interface TestVersionFetcher {
+
+    get version(): string;
+
+    /// fetch stream information and populate stream map
+    fetchStreamsInfo(streamMap: Map<string, TestStreamRef>): Promise<void>;
+
+    /// fetch test sessions and populate test status session
+    fetchTestSessions(streamRef: TestStreamRef, status: TestStatus, weeks: number, onAdd?: (session: TestSessionResult) => void): Promise<void>;
+
+}
+
+export class TestDataVersionRegistrar {
+
+    private static _manifest: TestVersionFetcher[] = [];
+
+    static register(fetcher: TestVersionFetcher) {
+        TestDataVersionRegistrar._manifest.push(fetcher);
+    }
+
+    static get manifest(): ReadonlyArray<TestVersionFetcher> {
+        return TestDataVersionRegistrar._manifest;
+    }
+
+}
+
+export class TestSessionResult {
+    constructor(id: string, start: Date, duration: string, commitId: string, commitOrder: number, nameRef: string, metadataId: string, streamId: string, outcome: TestOutcome, testDataId?: string) {
+        this.id = id;
+        this.start = start;
+        this.duration = duration;
+        this.commitId = commitId;
+        this.commitOrder = commitOrder;
+        this.nameRef = nameRef;
+        this.metadataId = metadataId;
+        this.streamId = streamId;
+        this.outcome = outcome;
+        this.testDataId = testDataId;
+        this.phaseSucceededCount = 0;
+        this.phaseFailedCount = 0;
+        this.phaseUnspecifiedCount = 0;
+    }
+
+    id: TestSessionId;
+    start: Date;
+    duration: string;
+    commitId: string;
+    commitOrder: number;
+    nameRef: TestId;
+    metadataId: MetaId;
+    streamId: StreamId
+    outcome: TestOutcome;
+    phaseSucceededCount: number;
+    phaseFailedCount: number;
+    phaseUnspecifiedCount: number;
+    jobId?: string;
+    stepId?: string;
+    testDataId?: string;
+}
+
+export class TestSessionStatus {
+    constructor(test: TestNameRef) {
+        this.test = test;
+        this.history = [];
+    }
+
+    addSession(session: TestSessionResult) {
+        this.history.push(session);
+    }
+
+    getLastSession(): TestSessionResult | undefined {
+        if (!this.history.length) {
+            return;
+        }
+        return this.history[0];
+    }
+
+    sortSessions() {
+        /// sort session history by commit order descending
+        this.history.sort((a, b) => b.commitOrder - a.commitOrder);
+    }
+
+    history: TestSessionResult[];
+    test: TestNameRef;
+}
+
+export type Metadata = { [Key in string]: string | undefined };
+export type MetadataFilter = Map<string, string[]>;
+
+// query states
+type TestDataState = {
+    stream?: string;
+    tests?: string[];
+    tags?: string[];
+    metadata?: MetadataFilter;
+    weeks?: number;
+}
+
+export class TestTagRef {
+    constructor(id: string, name: string, version: number) {
+        this.id = id;
+        this.name = name;
+        this.version = version;
+        this.lowerName = this.name.toLowerCase();
+    }
+
+    isMatch(toMatch: string): boolean {
+        return toMatch === this.lowerName;
+    }
+
+    id: TagId;
+    name: string;
+    version: number;
+
+    private lowerName: string;
+}
+
+export class TestNameRef {
+    constructor(id: string, key: string, name: string, version: number) {
+        this.id = id;
+        this.key = name;
+        this.name = name;
+        this.version = version;
+        this.lowerKey = this.key.toLowerCase();
+        this.lowerName = this.name.toLowerCase();
+    }
+
+    isMatch(toMatch: string): boolean {
+        return toMatch.charAt(0) === '#'? this.lowerName === toMatch.substring(1) || this.lowerKey === toMatch.substring(1) : this.lowerName.includes(toMatch) || this.lowerKey.includes(toMatch);
+    }
+
+    updateTags(tags: TestTagRef[]) {
+        const newTags = new Set<TestTagRef>(tags);
+        this.tagRefs = this.tagRefs? this.tagRefs.union(newTags) : newTags;
+    }
+
+    includesTags(toMatch: string[]): boolean {
+        return toMatch.every(tag => this.tagRefs?.values().some(r => r.isMatch(tag)));
+    }
+
+    id: TestId;
+    key: string;
+    name: string;
+    version: number;
+    tagRefs?: Set<TestTagRef>;
+
+    private lowerKey: string;
+    private lowerName: string;
+}
+
+export class MetadataRef {
+
+    constructor(id: string, entries: Metadata, version: number) {
+        this.id =  id;
+        this.version = version;
+        this.entries = entries;
+        this.orderedKeys = this.getKeys().sort((a, b) => {
+            const indexa = MetadataRef.keysOrderPriority.indexOf(a);
+            const indexb = MetadataRef.keysOrderPriority.indexOf(b);
+            if (indexa === -1) return indexb === -1? a.localeCompare(b) : 1;
+            if (indexb === -1) return -1;
+            return indexa - indexb;
+        });
+        this.orderedValues = this.orderedKeys.map((k) => this.entries[k]!);
+    }
+
+    getKeys(): string[] {
+        return Object.keys(this.entries);
+    }
+
+    getValues(): string[] {
+        return this.orderedKeys.map((k) => this.entries[k]!).filter((v) => !MetadataRef.valuesToIgnore.has(v));
+    }
+
+    getValuesExcept(skipKeys: string[]): string[] {
+        return this.orderedKeys.filter((k) => !skipKeys.includes(k)).map((k) => this.entries[k]!).filter((v) => !MetadataRef.valuesToIgnore.has(v));
+    }
+
+    getValue(key: string): string | undefined {
+        return this.entries[key];
+    }
+
+    getCommonValues(commonKeys?: string[]): string[] | undefined {
+        if (!commonKeys) {
+            return undefined;
+        }
+        return this.orderedKeys.filter((k) => commonKeys.includes(k)).map((k) => this.entries[k]!).filter((v) => !MetadataRef.valuesToIgnore.has(v));
+    }
+
+    isMatch(metaFilter: MetadataFilter): boolean {
+        for (const [key, values] of metaFilter.entries()) {
+            const thisValue = this.entries[key];
+            if (!thisValue) {
+                // missing key
+                return false;
+            }
+            const toValidate = values.some((v) => v.indexOf('+') > -1) ? [thisValue] : thisValue.split('+');
+            if (!toValidate.some((v) => values.includes(v))) {
+                return false;
+            }
+        }
+        return true;
+    }
+
+    id: MetaId;
+    version: number;
+    entries: Metadata;
+    orderedKeys: string[];
+    orderedValues: string[];
+
+    private static keysOrderPriority = ["Project", "BuildTarget", "Platform"]; // Should those be set from a config?
+    private static valuesToIgnore = new Set(["default", "Default"]);
+
+    private static globalMap: Map<string, MetadataRef> = new Map();
+
+    static identifyCommonKeys(items: MetadataRef[]): string[] {
+        if (!items.length) return [];
+        const first = items[0].entries;
+        const commonKeys = new Set(Object.keys(first));
+        if (items.length > 1) {
+            items.slice(1).forEach((item) => {
+                commonKeys.values().forEach((key) => {
+                    if (first[key] !== item.entries[key]) {
+                        commonKeys.delete(key);
+                    }
+                });
+            });
+        }
+        return Array.from(commonKeys);
+    }
+}
+
+export class TestMetaStatus {
+    constructor(test: TestNameRef) {
+        this.sessions = new Map();
+        this.test = test;
+    }
+
+    addSession(metaRef: MetadataRef, session: TestSessionResult) {
+        if (!this.sessions.has(metaRef)) {
+            this.sessions.set(metaRef, new TestSessionStatus(this.test));
+        }
+
+        this.sessions.get(metaRef)!.addSession(session);
+    }
+
+    getLastSession(metaRef: MetadataRef): TestSessionResult | undefined {
+        return this.sessions.get(metaRef)?.getLastSession();
+    }
+
+    getMetadata(): MetadataRef[] {
+        return this.sessions.keys().toArray();
+    }
+
+    includesMetadata(metaRef: MetadataRef): boolean {
+        return this.sessions.has(metaRef);
+    }
+
+    includesAnyMetadata(metaRefs: Set<MetadataRef>): boolean {
+        return metaRefs.keys().some((meta) => this.sessions.has(meta));
+    }
+
+    getMetadataLastSessions(): [MetadataRef, TestSessionResult][] {
+        return this.sessions.entries().filter(([meta, status]) => status.history.length > 0).map(([meta, status]) => [meta, status.getLastSession()!] as [MetadataRef, TestSessionResult]).toArray();
+    }
+
+    sortSessions() {
+        for (const status of this.sessions.values()) {
+            status.sortSessions();
+        }
+    }
+
+    sessions: Map<MetadataRef, TestSessionStatus>;
+    test: TestNameRef;
+}
+
+export class TestStatus {
+    constructor() {
+        this.tests = new Map();
+    }
+
+    addSession(testNameRef: TestNameRef, metaRef: MetadataRef, session: TestSessionResult) {
+        if (!this.tests.has(testNameRef)) {
+            this.tests.set(testNameRef, new TestMetaStatus(testNameRef));
+        }
+
+        this.tests.get(testNameRef)!.addSession(metaRef, session);
+    }
+
+    getLastSession(testNameRef: TestNameRef, metaRef: MetadataRef): TestSessionResult | undefined {
+        return this.tests.get(testNameRef)?.getLastSession(metaRef);
+    }
+
+    getMatchingTests(testNameRefs: Set<TestNameRef>, metaRefs: Set<MetadataRef>): TestNameRef[] {
+        return this.tests.entries().filter(([test, status]) => testNameRefs.has(test) && status.includesAnyMetadata(metaRefs)).map(([test, status]) => test).toArray();
+    }
+
+    getTestNames(): TestNameRef[] {
+        return this.tests.keys().toArray();
+    }
+
+    getTestNameMetadata(): [TestNameRef, TestMetaStatus][] {
+        return this.tests.entries().toArray();
+    }
+
+    sortSessions() {
+        for (const meta of this.tests.values()) {
+            meta.sortSessions();
+        }
+    }
+
+    tests: Map<TestNameRef, TestMetaStatus>;
+}
+
+export class TestStreamRef {
+
+    constructor(streamId: string) {
+        this.id = streamId;
+        this.tests = [];
+        this.meta = [];
+        this.tags = [];
+
+        this.testNames = new Set();
+        this.metadataMap = new Map();
+        this.tagNames = new Set();
+    }
+
+    addTest(item: TestNameRef) {
+        this.tests.push(item);
+        this.testNames.add(item.name);
+    }
+
+    addMetadata(item: MetadataRef) {
+        this.meta.push(item);
+        item.getKeys().forEach((key) => {
+            let meta = this.metadataMap.get(key);
+            if (!meta) {
+                meta = new Set<string>();
+                this.metadataMap.set(key, meta!);
+            }
+            meta!.add(item.getValue(key)!);
+        });
+    }
+
+    addTag(item: TestTagRef) {
+        this.tags.push(item);
+        this.tagNames.add(item.name);
+    }
+
+    getAllTestNames(): string[] {
+        return  this.testNames.keys().toArray();
+    }
+
+    getAllMetadataKeys(): string[] {
+        return  this.metadataMap.keys().toArray();
+    }
+
+    getAllMetadataValues(key: string): string[] | undefined {
+        return this.metadataMap.get(key)?.keys().toArray();
+    }
+
+    getAllTagNames(): string[] {
+        return  this.tagNames.keys().toArray();
+    }
+
+    getMatchingTests(toMatchNames?: string[], toMatchTags?: string[]): Set<TestNameRef> {
+        const lowerToMatchNames = toMatchNames?.map((i) => i.toLowerCase());
+        const lowerToMatchTags = toMatchTags?.map((i) => i.toLowerCase());
+        const matches: Set<TestNameRef> = new Set();
+        this.tests.forEach((test) => {
+            const isMatchNames = !lowerToMatchNames || lowerToMatchNames.some((key) => test.isMatch(key));
+            const isMatchTags = !lowerToMatchTags || test.includesTags(lowerToMatchTags);
+            if (isMatchNames && isMatchTags) {
+                matches.add(test);
+            }
+        })
+
+        return matches;
+    }
+
+    getMatchingMeta(toMatchMeta?: MetadataFilter): Set<MetadataRef> {
+        const matches: Set<MetadataRef> = new Set();
+        this.meta.forEach((meta) => {
+            if (!toMatchMeta || meta.isMatch(toMatchMeta)) {
+                matches.add(meta);
+            }
+        });
+
+        return matches;
+    }
+
+    id: StreamId;
+    tests: TestNameRef[];
+    meta: MetadataRef[];
+    tags: TestTagRef[];
+
+    private testNames: Set<string>;
+    private metadataMap: Map<string, Set<string>>;
+    private tagNames: Set<string>;
+}
+
+export class TestDataFilterSearchParams {
+
+    constructor(options?: any) {
+        this.search =  new URLSearchParams(options);
+    }
+
+    private encodeKey(target: string, key?: string | null): string {
+        if (!!key) {
+            return `${target}.${key}`;
+        }
+
+        return target;
+    }
+
+    private encodeValue(items: string | string[]): string {
+        if (typeof items === "string") {
+            return items;
+        }
+
+        return items.join('~');
+    }
+
+    private decodeKey(key: string): string[] {
+        return key.split('.', 2);
+    }
+
+    private decodeValue(value: string): string[] {
+        return value.split('~')
+    }
+
+    append(target: string, key: string | null, items: string | string[]) {
+        const actual_key = this.encodeKey(target, key);
+        const append_items = typeof items === "string"? [items] : items;
+        const merged_items = new Set<string>(
+            !this.search.has(actual_key)? append_items : [...this.decodeValue(this.search.get(actual_key)!), ...append_items]
+        );
+
+        this.search.set(actual_key, this.encodeValue(Array.from(merged_items)));
+    }
+
+    set(target: string, key: string | null, items?: string | string[]) {
+        if (!items || !items.length) {
+            this.delete(target, key);
+            return;
+        }
+
+        this.search.set(this.encodeKey(target, key), this.encodeValue(items));
+    }
+
+    delete(target: string, key?: string | null, items?: string | string[]) {
+        const actual_key = this.encodeKey(target, key);
+        if (!this.search.has(actual_key)) {
+            return;
+        }
+        if (!items) {
+            this.search.delete(actual_key)
+            return;
+        }
+        const current_items = new Set<string>(this.decodeValue(this.search.get(actual_key)!));
+        const remove_items = new Set<string>(typeof items === "string"? [items] : items);
+        this.search.set(actual_key, this.encodeValue(Array.from(current_items.difference(remove_items))));
+    }
+
+    get(target: string, key?: string | null): string[] | null {
+        const actual_key = this.encodeKey(target, key);
+        return this.search.has(actual_key)? this.decodeValue(this.search.get(actual_key)!) : null;
+    }
+
+    getGroup(target: string): Map<string, string[]> | null {
+        const entries = new Map<string, string[]>();
+        const targetKey = `${target}.`;
+
+        this.search.keys().forEach(longKey => {
+            if (longKey.startsWith(targetKey)) {
+                entries.set(longKey.substring(targetKey.length), this.decodeValue(this.search.get(longKey)!));
+            }
+        });
+
+        return entries;
+    }
+
+    deleteGroup(target: string) {
+        const targetKey = `${target}.`;
+
+        this.search.keys().forEach(longKey => {
+            if (longKey.startsWith(targetKey)) {
+                this.search.delete(longKey);
+            }
+        });
+    }
+
+    has(target: string, key?: string | null, items?: string | string[]): boolean {
+        const actual_key = this.encodeKey(target, key);
+        if (!!items) {
+            if (!this.search.has(actual_key)) {
+                return false;
+            }
+            const current_items = new Set<string>(this.decodeValue(this.search.get(actual_key)!));
+            return typeof items === "string"? current_items.has(items) : items.every((item) => current_items.has(item));
+        } else {
+            return this.search.has(actual_key);
+        }
+    }
+
+    toString(): string {
+        return this.search.toString();
+    }
+
+    private search: URLSearchParams;
+}
+
+export class TestDataHandler {
+
+    constructor() {
+        TestDataHandler.instance = this;
+        makeObservable(this);
+        // setup a reaction on filterUpdated to debounce the filter refresh
+        reaction(
+            () => this.filterUpdated,
+            () => this.updateFilter(),
+            { delay: filterRefreshMS }
+        );
+    }
+
+    async initialize(search: string, searchCallback?: (search: string) => void) {
+        this.filterState = this.filterFromSearch(search);
+        this.searchCallback = searchCallback;
+        this.checkSearchCallback(search);
+        await this.load();
+    }
+
+    clear() {
+        this.searchCallback = undefined;
+        this.filterState = {};
+        this.search = new TestDataFilterSearchParams();
+        this.testStreams = new Map();
+        this.commitIdDates = new Map();
+        this.testStatus = new Map();
+        this.loaded = false;
+    }
+
+    syncSearchParam(search: string) {
+        if (search != this.search.toString()) {
+            this.filterState = this.filterFromSearch(search);
+            this.setSearchUpdated();
+        }
+    }
+
+    private checkSearchCallback(previous: string) {
+        if (this.searchCallback && previous != this.search.toString()) {
+            this.searchCallback(this.search.toString());
+        }
+    }
+
+    setSearchParam(key: string, value: string) {
+        const previous = this.search.toString();
+        this.search.set(key, null, value);
+        this.checkSearchCallback(previous);
+    }
+
+    removeSearchParam(key: string) {
+        const previous = this.search.toString();
+        this.search.delete(key);
+        this.checkSearchCallback(previous);
+    }
+
+    getSearchParam(key: string): string | string[] | null | undefined {
+        const values = this.search.get(key);
+        if (values?.length === 1) return values[0];
+        return values;
+    }
+
+    setQueryWeeks(weeks: number) {
+        if (weeks !== this.filterState.weeks) {
+            this.filterState.weeks = weeks;
+            if (weeks && weeks > defaultQueryWeeks) {
+                // scrap the cache
+                this.testStatus = new Map();
+            }
+            this.setFilterUpdated();
+        }
+    }
+
+    addMetadata(key: string, value: string) {
+
+        const state = this.filterState;
+        const meta = state.metadata?.get(key);
+
+        if (!!meta && meta!.includes(value)) {
+            return;
+        }
+
+        if (!state.metadata) {
+            state.metadata = new Map();
+        }
+
+        if (!meta) {
+            state.metadata!.set(key, [value]);
+        } else {
+            state.metadata!.get(key)!.push(value);
+        }
+
+        this.setFilterUpdated();
+    }
+
+    removeMetadata(key: string, value?: string) {
+
+        const state = this.filterState;
+        const meta = state.metadata?.get(key);
+        if (!!value) {
+            const idx = meta?.indexOf(value!);
+
+            if (!meta || !isNumber(idx) || idx! < 0) {
+                return;
+            }
+
+            meta!.splice(idx!, 1);
+        }
+
+        if (!value || !meta!.length) {
+            state.metadata!.delete(key);
+        }
+
+        this.setFilterUpdated();
+    }
+
+    addTest(test: string) {
+
+        const state = this.filterState;
+
+        if (state.tests) {
+            if (state.tests.includes(test)) {
+                return;
+            }
+        }
+
+        if (!state.tests) {
+            state.tests = [];
+        }
+
+        state.tests.push(test);
+
+        this.setFilterUpdated();
+    }
+
+    removeTest(test: string) {
+
+        const state = this.filterState;
+
+        const idx = state.tests?.indexOf(test);
+        if (!isNumber(idx) || idx! < 0) {
+            return;
+        }
+
+        state.tests!.splice(idx!, 1);
+
+        if (!state.tests!.length) {
+            state.tests = undefined;
+        }
+
+        this.setFilterUpdated();
+    }
+
+    selectStream(streamId?: StreamId) {
+
+        const state = this.filterState;
+
+        if (state.stream === streamId) {
+            return;
+        }
+
+        state.stream = streamId;
+
+        this.setFilterUpdated();
+    }
+
+    /// All test names from current stream
+    get allStreamTestNames(): string[] {
+
+        const state = this.filterState;
+        if (!state.stream) {
+            return [];
+        }
+
+        const testStreamRef = this.testStreams.get(state.stream);
+        if (!testStreamRef) {
+            return [];
+        }
+
+        return testStreamRef.getAllTestNames().sort((a, b) => a.localeCompare(b));
+    }
+
+    /// All test tags from current stream
+    get allStreamTestTags(): string[] {
+
+        const state = this.filterState;
+        if (!state.stream) {
+            return [];
+        }
+
+        const testStreamRef = this.testStreams.get(state.stream);
+        if (!testStreamRef) {
+            return [];
+        }
+
+        return testStreamRef.getAllTagNames().sort((a, b) => a.localeCompare(b));
+    }
+
+    /// All metadata keys from current stream
+    get allStreamMetadataKeys(): string[] {
+
+        const state = this.filterState;
+        if (!state.stream) {
+            return [];
+        }
+
+        const testStreamRef = this.testStreams.get(state.stream);
+        if (!testStreamRef) {
+            return [];
+        }
+
+        return testStreamRef.getAllMetadataKeys().sort((a, b) => a.localeCompare(b));
+    }
+
+    /// All metadata values for a key from current stream
+    getAllStreamMetadataValues(key: string): string[] {
+
+        const state = this.filterState;
+        if (!state.stream) {
+            return [];
+        }
+
+        const testStreamRef = this.testStreams.get(state.stream);
+        if (!testStreamRef) {
+            return [];
+        }
+
+        return testStreamRef.getAllMetadataValues(key)?.sort((a, b) => a.localeCompare(b)) ?? [];
+    }
+
+    /// Tests that match current filters
+    get filteredTests(): TestNameRef[] {
+
+        const state = this.filterState;
+        if (!state.stream) {
+            return [];
+        }
+
+        const testStreamRef = this.testStreams.get(state.stream);
+        if (!testStreamRef) {
+            return [];
+        }
+
+        const testStatus = this.testStatus.get(state.stream);
+        if (!testStatus) {
+            return [];
+        }
+
+        const metaRefs = testStreamRef!.getMatchingMeta(state.metadata);
+        const testNameRefs = testStreamRef!.getMatchingTests(state.tests, state.tags);
+
+        return testStatus!.getMatchingTests(testNameRefs, metaRefs).sort((a, b) => a.name.localeCompare(b.name));
+    }
+
+    /// Metadata that match current filters
+    get filteredMetadata(): MetadataRef[] {
+
+        const state = this.filterState;
+        if (!state.stream) {
+            return [];
+        }
+
+        const testStreamRef = this.testStreams.get(state.stream);
+        if (!testStreamRef) {
+            return [];
+        }
+
+        return testStreamRef!.getMatchingMeta(state.metadata).keys().toArray().sort((a, b) => a.orderedValues.join().localeCompare(b.orderedValues.join()));
+    }
+
+    private updateFilter() {
+        if (this.checkUpdateFilter()) {
+            this.setSearchUpdated();
+            if (this.searchCallback) {
+                this.searchCallback(this.search.toString());
+            }
+        }
+    }
+
+    private checkUpdateFilter(): boolean {
+        const state = this.filterState;
+        state.weeks = state.weeks ? state.weeks : undefined;
+        state.stream = state.stream ? state.stream : undefined;
+        state.metadata = state.metadata && state.metadata.size > 0 ? state.metadata : undefined
+        state.tests = state.tests?.sort((a, b) => a.localeCompare(b));
+
+        let needUpdate = false;
+        const search = this.search;
+
+        needUpdate = needUpdate || search.get('weeks')?.[0] !== state.weeks?.toString();
+        if (needUpdate) {
+            search.set("weeks", null, state.weeks?.toString());
+        }
+
+        needUpdate = needUpdate || search.get('stream')?.[0] !== state.stream;
+        if (needUpdate) {
+            search.set("stream", null, state.stream);
+        }
+
+        needUpdate = needUpdate || search.get('tests')?.toString() !== state.tests?.toString();
+        if (needUpdate) {
+            search.set("tests", null, state.tests);
+        }
+
+        const metaKeys = new Set<string>(search.getGroup("meta")?.keys());
+        metaKeys.union(new Set(state.metadata?.keys())).keys().forEach(key => {
+            needUpdate = needUpdate || search.get('meta', key)?.toString() !== state.metadata?.get(key)?.toString();
+            search.set("meta", key, state.metadata?.get(key));
+        });
+
+        if (needUpdate) {
+            this.query(state.stream);
+            return true;
+        }
+
+        return false;
+    }
+
+    private filterFromSearch(searchInfo: string): TestDataState {
+        const search = new TestDataFilterSearchParams(searchInfo);
+        this.search = search;
+
+        const state: TestDataState = {};
+
+        const stream = search.get("stream") ?? undefined;
+        const metadata = search.getGroup("meta");
+        const tests = search.get("tests") ?? undefined;
+        const weeks = search.get("weeks") ?? undefined;
+
+        state.stream = stream && stream.length > 0? stream[0] : undefined;
+        state.metadata = metadata && metadata.size > 0 ? metadata : undefined;
+        state.tests = tests;
+        state.weeks = weeks && weeks.length > 0 ? parseInt(weeks[0]) : defaultQueryWeeks;
+        if (typeof (state.weeks) !== "number") {
+            state.weeks = defaultQueryWeeks;
+        }
+
+        return state;
+    }
+
+    async query(streamId?: StreamId, isSubQuery: boolean = false) {
+
+        if (!streamId || !this.testStreams.has(streamId) || this.testStatus.has(streamId)) {
+            isSubQuery? this.setSubQueryLoading(false): this.setQueryLoading(false);
+            return;
+        }
+
+        isSubQuery? this.setSubQueryLoading(true): this.setQueryLoading(true);
+
+        const state = this.filterState;
+        const commitIdDates = this.commitIdDates
+        const testStreamRef: TestStreamRef = this.testStreams.get(streamId)!;
+        const testStatus = new TestStatus();
+        this.testStatus.set(streamId, testStatus);
+
+        const aggregateCommitDates = (session: TestSessionResult) => {
+            /// Update uniform commit dates
+            const date = commitIdDates.get(session.commitId);
+            const rdate = session.start;
+            // Use the latest date for summary, otherwise gets confusing when a CL spans days
+            if (!date || date.getTime() < rdate.getTime()) {
+                commitIdDates.set(session.commitId, rdate);
+            }
+        }
+
+        const queryWeeks = state.weeks && state.weeks > defaultQueryWeeks ? state.weeks : defaultQueryWeeks;
+
+        for (const fetcherVersion of TestDataVersionRegistrar.manifest) {
+            console.log(`Fetching from ${fetcherVersion.version} test sessions`)
+            await fetcherVersion.fetchTestSessions(testStreamRef, testStatus, queryWeeks, aggregateCommitDates);
+        }
+
+        /// sort session history by commit order descending
+        testStatus.sortSessions();
+
+        if (this.loaded) {
+            isSubQuery? this.setSubQueryLoading(false): this.setQueryLoading(false);
+            !isSubQuery && this.setSearchUpdated();
+        }
+    }
+
+    // load immutable data for this view session
+    private async load() {
+
+        const allStreams = projectStore.projects.map(p => p.streams).flat().filter(s => !!s).map(s => s!.id);
+        const testStreams = this.testStreams = new Map(allStreams.map((s) => [s, new TestStreamRef(s)]));
+
+        for (const fetcherVersion of TestDataVersionRegistrar.manifest) {
+            console.log(`Fetching from ${fetcherVersion.version} stream info`)
+            await fetcherVersion.fetchStreamsInfo(testStreams);
+        }
+
+        if (!this.loaded) {
+            this.loaded = true;
+            this.query(this.filterState.stream);
+        }
+    }
+
+    @action
+    setFilterUpdated() {
+        this.filterUpdated++;
+    }
+
+    @action
+    setSearchUpdated() {
+        this.searchUpdated++;
+    }
+
+    @action
+    setQueryLoading(value: boolean) {
+        this.queryLoading = value;
+    }
+
+    @action
+    setSubQueryLoading(value: boolean) {
+        this.subQueryLoading = value;
+    }
+
+    get stream(): StreamId | undefined {
+        return this.filterState.stream;
+    }
+
+    get availableStreams(): StreamId[] {
+        return  this.testStreams.keys().toArray();
+    }
+
+    get availableStatusStreams(): StreamId[] {
+        return this.testStatus.keys().toArray();
+    }
+
+    getStatusStream(streamId: string): TestStatus | undefined {
+        return this.testStatus.get(streamId);
+    }
+
+    get selectedStatusStream(): TestStatus | undefined {
+        if (!this.filterState.stream) {
+            return;
+        }
+        return this.getStatusStream(this.filterState.stream!);
+    }
+
+    @observable
+    filterUpdated: number = 0;
+
+    @observable
+    searchUpdated: number = 0;
+
+    @observable
+    queryLoading: boolean = false;
+
+    @observable
+    subQueryLoading: boolean = false;
+
+    subscribeToFilter() {
+        if (this.filterUpdated) { }
+    }
+
+    subscribeToSearch() {
+        if (this.searchUpdated) { }
+    }
+
+    subscribeToQueryLoading() {
+        if (this.queryLoading) { }
+    }
+
+    subscribeToSubQueryLoading() {
+        if (this.subQueryLoading) { }
+    }
+
+    private searchCallback?: (search: string) => void;
+    private search: TestDataFilterSearchParams = new TestDataFilterSearchParams();
+    filterState: TestDataState = {};
+
+    testStreams: Map<StreamId, TestStreamRef> = new Map();
+
+    // quantized common commits to land at same time in the views
+    commitIdDates = new Map<string, Date>();
+    
+    private testStatus: Map<StreamId, TestStatus> = new Map();
+
+    loaded = false;
+
+    static instance: TestDataHandler;
+}
diff --git a/Engine/Source/Programs/Horde/HordeDashboard/plugins/test-automation/testDataV2Fetcher.ts b/Engine/Source/Programs/Horde/HordeDashboard/plugins/test-automation/testDataV2Fetcher.ts
new file mode 100644
--- /dev/null
+++ b/Engine/Source/Programs/Horde/HordeDashboard/plugins/test-automation/testDataV2Fetcher.ts
@@ -0,0 +1,110 @@
+import { getStreamTestsV2, GetTestMetadataResponse, GetTestNameResponse, GetTestSessionResponse, getTestSessionsV2, GetTestTagResponse } from "./api";
+import { TestStreamRef, TestStatus, TestSessionResult, TestNameRef, MetadataRef, TestTagRef } from "./testData"
+
+const minuteInWeek = 10080;
+
+export class TestDataV2 {
+    get version() { return "TestData Version 2" }
+
+    async fetchStreamsInfo(testStreams: Map<string, TestStreamRef>) {
+
+        const results = await getStreamTestsV2(testStreams.keys().toArray());
+        results.forEach((streamInfo) => {
+            const testStream = testStreams.get(streamInfo.streamId)!;
+
+            /// Tests
+            streamInfo.tests.forEach((test) => {
+                testStream.addTest(this.createTestNameRef(test));
+            });
+
+            /// Metadata
+            streamInfo.testMetadata.forEach((entry) => {
+                testStream.addMetadata(this.createMetadataRef(entry));
+            });
+
+            /// Tags
+            streamInfo.testTags.forEach((tag) => {
+                testStream.addTag(this.createTagRef(tag));
+            });
+            
+        });
+   
+    }
+
+    async fetchTestSessions(streamRef: TestStreamRef, status: TestStatus, weeks: number, onAdd?: (session: TestSessionResult) => void) {
+
+        const streamId = streamRef.id;
+        const tests: Map<string, TestNameRef> = new Map(streamRef.tests.filter((t) => t.version >= 2).map((t) => [t.id, t]));
+        const meta: Map<string, MetadataRef> = new Map(streamRef.meta.filter((m) => m.version >= 2).map((m) => [m.id, m]));
+
+        const minQueryDate = new Date(new Date().valueOf() - (minuteInWeek * weeks * 60000));
+        const maxQueryDate = new Date();
+
+        if (tests.size) {
+            const responses = await getTestSessionsV2([streamId], Array.from(tests.keys()),  Array.from(meta.keys()), minQueryDate.toISOString(), maxQueryDate.toISOString());
+            responses.forEach((session) => {
+                if (session.streamId !== streamId) {
+                    return;
+                }
+                const testSession = this.createTestSession(session);
+                const testNameRef = tests.get(testSession.nameRef);
+                const metadataRef = meta.get(testSession.metadataId);
+                if (!!testNameRef && !!metadataRef) {
+                    status.addSession(testNameRef, metadataRef, testSession);
+                    if (onAdd) onAdd(testSession);
+                }
+            });
+        }
+
+    }
+
+    private testNameRefMap: Map<string, TestNameRef> = new Map();
+    private metadataRefMap: Map<string, MetadataRef> = new Map();
+    private testTageRefMap: Map<string, TestTagRef> = new Map();
+
+    private createTestNameRef(test: GetTestNameResponse): TestNameRef {
+        const globalKey = test.id;
+        let item = this.testNameRefMap.get(globalKey);
+        if (!item) {
+            item = new TestNameRef(test.id, test.key, test.name, 2);
+            this.testNameRefMap.set(globalKey, item);
+        }
+
+        return item;
+    }
+
+    private createMetadataRef(item: GetTestMetadataResponse): MetadataRef {
+        const globalKey = item.id;
+        let meta = this.metadataRefMap.get(globalKey);
+        if (!meta) {
+            meta = new MetadataRef(item.id, item.entries, 2);
+            this.metadataRefMap.set(globalKey, meta);
+        }
+
+        return meta;
+    }
+
+    private createTagRef(item: GetTestTagResponse): TestTagRef {
+        const globalKey = item.id;
+        let tag = this.testTageRefMap.get(globalKey);
+        if (!tag) {
+            tag = new TestTagRef(item.id, item.name, 2);
+            this.testTageRefMap.set(globalKey, tag);
+        }
+
+        return tag;
+    }
+
+    private createTestSession(test: GetTestSessionResponse): TestSessionResult {
+        const start = new Date(test.startDateTime);
+        const result = new TestSessionResult(test.id, start, test.duration, test.commitId, test.commitOrder, test.nameRef, test.metadataId, test.streamId, test.outcome);
+        result.phaseSucceededCount = test.phasesSucceededCount;
+        result.phaseFailedCount = test.phasesFailedCount;
+        result.phaseUnspecifiedCount = test.phasesUndefinedCount;
+        result.jobId = test.jobId;
+        result.stepId = test.stepId;
+
+        return result;
+    }
+}
+
diff --git a/Engine/Source/Programs/Horde/HordeDashboard/plugins/test-automation/testHistoryGraph.ts b/Engine/Source/Programs/Horde/HordeDashboard/plugins/test-automation/testHistoryGraph.ts
new file mode 100644
--- /dev/null
+++ b/Engine/Source/Programs/Horde/HordeDashboard/plugins/test-automation/testHistoryGraph.ts
@@ -0,0 +1,406 @@
+import * as d3 from "d3";
+import moment from 'moment';
+import { TestDataHandler, TestNameRef, MetadataRef, TestMetaStatus, TestSessionResult } from './testData';
+import dashboard, { StatusColor } from 'horde/backend/Dashboard';
+import { getHumanTime, getShortNiceTime, msecToElapsed } from 'horde/base/utilities/timeUtils';
+import { getHordeStyling } from 'horde/styles/Styles';
+import { SessionStatusBar } from "./testAutomationCommon";
+import { renderToString } from "react-dom/server"
+import { TestOutcome } from "./api";
+
+// Handle bad "@types/d3" types, fix if addressed upstream
+const _d3 = d3 as any;
+
+type SelectionType = d3.Selection<SVGGElement, unknown, null, undefined>;
+type DivSelectionType = d3.Selection<HTMLDivElement, unknown, null, undefined>;
+
+
+export class TestHistoryGraph {
+
+    constructor(test: TestNameRef, stream: string, handler: TestDataHandler, commonMetaKeys?: string[], onClickCallback?: (session: TestSessionResult) => void) {
+        this.test = test;
+
+        this.handler = handler;
+        this.onClickCallback = onClickCallback;
+        this.commonMetaKeys = commonMetaKeys;
+
+        this.status = handler.getStatusStream(stream)!.tests.get(test)!;
+
+        this.margin = { top: 0, right: 32, bottom: 0, left: 160 };
+        this.clipId = `test_history_clip_path_${test.key}_${stream}`;
+    }
+
+    initData() {
+        this.metaRefs = this.handler.filteredMetadata.filter((m) => this.status.includesMetadata(m));
+        const commonMetaKeys = this.commonMetaKeys = this.commonMetaKeys ?? MetadataRef.identifyCommonKeys(this.metaRefs);
+
+        this.sessions = [];
+
+        this.metaRefs.forEach(meta => {
+            const metaLastSession = this.status.getLastSession(meta);
+            this.metaStatus.set(meta, !metaLastSession? TestOutcome.Unspecified : metaLastSession.outcome);
+
+            const elements: string[] = meta.getValuesExcept(commonMetaKeys);
+            this.metaNames.set(meta, elements.join(" / "));
+
+            this.sessions.push(...this.status.sessions.get(meta)!.history);
+        });
+    }
+
+    render(container: HTMLDivElement) {
+
+        if (this.hasRendered && !this.forceRender) {
+            return;
+        }
+
+        const { modeColors } = getHordeStyling();
+
+        this.clear(container);
+
+        this.hasRendered = true;
+        this.forceRender = false;
+
+        this.initData();
+
+        const handler = this.handler;
+        const sessions = this.sessions;
+        const width = 1000
+
+        const scolors = dashboard.getStatusColors();
+        const colors = new Map<TestOutcome, string>([
+            [TestOutcome.Success, scolors.get(StatusColor.Success)!],
+            [TestOutcome.Failure, scolors.get(StatusColor.Failure)!],
+            [TestOutcome.Unspecified, scolors.get(StatusColor.Unspecified)!],
+            [TestOutcome.Skipped, scolors.get(StatusColor.Skipped)!]
+        ]);
+
+        const X = _d3.map(sessions, (r) => handler.commitIdDates.get(r.commitId)!.getTime() / 1000);
+        const Y = _d3.map(sessions, (r) => this.metaRefs.findIndex((m) => m.id === r.metadataId));
+
+        const xDomain = d3.extent(handler.commitIdDates.values() as any, (d:any) => d.getTime() / 1000);
+        const yDomain = new _d3.InternSet(Y as any);
+
+        const I = d3.range(X.length);
+
+        const yPadding = 1;
+        const height = Math.ceil((yDomain.size + yPadding) * 16) + this.margin.top + this.margin.bottom;
+
+        const xRange = [this.margin.left, width - this.margin.right];
+        const yRange = [this.margin.top, height - this.margin.bottom];
+
+        const xScale = _d3.scaleTime(xDomain, xRange);
+        const yScale = _d3.scalePoint(yDomain, yRange).round(true).padding(yPadding);
+
+        const svg = d3.select(container)
+            .append("svg")
+            .attr("width", width)
+            .attr("height", height + 24)
+            .attr("viewBox", [0, 0, width, height + 24] as any)
+
+        // top axis
+        const xAxis = (g: SelectionType) => {
+
+            const dateMin = new Date(xDomain[0]! * 1000);
+            const dateMax = new Date(xDomain[1]! * 1000);
+
+            let ticks: number[] = [];
+            for (const date of d3.timeDays(dateMin, dateMax, 1).reverse()) {
+                ticks.push(date.getTime() / 1000);
+            }
+
+            if (ticks.length > 14) {
+                let nticks = [...ticks];
+                // remove first and last, will be readded 
+                const first = nticks.shift()!;
+                const last = nticks.pop()!;
+
+                const n = Math.floor(nticks.length / 12);
+
+                const rticks: number[] = [];
+                for (let i = 0; i < nticks.length; i = i + n) {
+                rticks.push(nticks[i]);
+                }
+
+                rticks.unshift(first);
+                rticks.push(last);
+                ticks = rticks;
+
+            }
+
+            g.attr("transform", `translate(0,16)`)
+                .style("font-family", "Horde Open Sans Regular")
+                .style("font-size", "9px")
+                .call(d3.axisTop(xScale)
+                .tickValues(ticks)
+                .tickFormat(d => getHumanTime(new Date((d as number) * 1000)))
+                .tickSizeOuter(0))
+                .call(g => g.select(".domain").remove())
+                .call(g => g.selectAll(".tick line")
+                    .attr("stroke-opacity", dashboard.darktheme ? 0.35 : 0.25).clone()
+                    .attr("stroke", dashboard.darktheme ? "#6D6C6B" : "#4D4C4B")
+                    .attr("y2", height - this.margin.bottom)
+                )
+        }
+
+        svg.append("g").call(xAxis)
+
+        const g = svg.append("g")
+            .selectAll()
+            .data(_d3.group(I, i => Y[i]))
+            .join("g")
+            .attr("class", "meta")
+            .attr("id", ([y]: any) => `meta${y as any}`)
+            .attr("transform", ([y]: any) => `translate(0,${(yScale(y) as any) + 16})`);
+
+        g.append("line")
+            .attr("stroke", (y: any) => {
+                const status = this.metaStatus.get(this.metaRefs[y])!;
+                if (status === TestOutcome.Failure) {
+                    return colors[this.metaStatus.get(this.metaRefs[y])!]
+                }
+                return dashboard.darktheme ? "#6D6C6B" : "#4D4C4B";
+            })
+            .attr("stroke-width", 1)
+            .attr("stroke-linecap", 4)
+            .attr("stroke-opacity", dashboard.darktheme ? 0.35 : 0.25)
+            .attr("x1", ([, I]: any) => this.margin.left)
+            .attr("x2", ([, I]: any) => width - this.margin.right);
+
+
+        g.append("text")
+            .attr("text-anchor", "start")
+            .style("alignment-baseline", "left")
+            .style("font-family", "Horde Open Sans Regular")
+            .style("font-size", 10)
+            .attr("dy", "0.15em") // center stream name
+            .attr("x", () => 0)
+            .attr("fill", ([y]: any) => {
+                const status = this.metaStatus.get(this.metaRefs[y])!;
+                if (!dashboard.darktheme && status === TestOutcome.Failure) {
+                    return colors[this.metaStatus.get(this.metaRefs[y])!]
+                }
+                return dashboard.darktheme ? "#E0E0E0" : "#2D3F5F";
+            })
+            .text(([y]: any) => this.metaNames.get(this.metaRefs[y])!);
+
+        g.selectAll("circle")
+            .data(([, I]: any) => I)
+            .enter()
+            .append("g")
+            .attr("id", i => `circle${sessions[i as any].id}`)
+            .attr("transform", i => `translate(${xScale(X[i as any])},0)`)
+            .append("g")
+            .attr("class", "circleScale")
+            .attr("id", i => `circleScale${sessions[i as any].id}`)
+            .each((d, i, nodes) => SessionPieChart(nodes[i], sessions[d as any], 5));
+
+        const tooltip = this.tooltip = d3.select(container)
+            .append("div")
+            .attr("id", "tooltip")
+            .style("display", "none")
+            .style("background-color", modeColors.background)
+            .style("border", "solid")
+            .style("border-width", "1px")
+            .style("border-radius", "3px")
+            .style("border-color", dashboard.darktheme ? "#413F3D" : "#2D3F5F")
+            .style("padding", "8px")
+            .style("position", "absolute")
+            .style("pointer-events", "none");
+
+        const closestData = (x: number, y: number): TestSessionResult | undefined => {
+
+            if (x < this.margin.left - 16) {
+                return undefined;
+            }
+
+            y -= 16;
+
+            const metaRefs = this.metaRefs;
+            let closest = sessions.reduce((best, session, i) => {
+                let absy = Math.abs(yScale(metaRefs.findIndex((m) => m.id === session.metadataId) as any)! - y)
+                const timeStamp = handler.commitIdDates.get(session.commitId)!.getTime() / 1000;
+                const sx = xScale(timeStamp);
+                let absx = Math.abs(sx - x)
+
+                const lengthSqr = absy * absy + absx * absx;
+
+                if (lengthSqr < best.value) {
+                    return { index: i, value: lengthSqr };
+                }
+                else {
+                    return best;
+                }
+            }, { index: 0, value: Number.MAX_SAFE_INTEGER });
+
+            if (closest) {
+                return sessions[closest.index];
+            }
+
+            return undefined;
+
+        }
+
+        const handleMouseMove = (event: any) => {
+
+            const mouseX = _d3.pointer(event)[0];
+            const mouseY = _d3.pointer(event)[1];
+
+            const closest = closestData(mouseX, mouseY);
+            if (!closest) {
+                handleMouseLeave(undefined);
+                return;
+            }
+
+            if (this.tooltipSessionId !== closest.id) {
+                this.tooltipSessionId = closest.id;
+                const metaIndex = this.metaRefs.findIndex((m) => m.id === closest.metadataId);
+
+                svg.selectAll(".circleScale").attr("transform", `scale(1)`);
+                svg.select(`#circleScale${closest.id}`).attr("transform", `scale(2.5)`);
+                svg.select(`#circle${closest.id}`).raise();
+                svg.select(`#meta${metaIndex}`).raise();
+
+                const cmeta = this.metaRefs[metaIndex].getValuesExcept(this.commonMetaKeys!).join(" / ");
+                const date = getShortNiceTime(closest.start, true, true);
+                let desc = "";
+                if (!!cmeta) desc += `${cmeta} <br/>`;
+                desc += `Commit ${closest.commitId} <br/>`;
+                desc += `Duration ${msecToElapsed(moment.duration(closest.duration).asMilliseconds(), true, true)} <br/>`;
+                desc += `${date} <br/>`;
+                desc += renderToString(SessionStatusBar(closest, 150, 8));
+
+                const timeStamp = handler.commitIdDates.get(closest.commitId)!.getTime() / 1000;
+                const tx = xScale(timeStamp);
+                const ty = yScale(metaIndex as any)!
+
+                this.updateTooltip(true, tx, ty, desc);
+            }
+        }
+
+        const handleMouseLeave = (event: any) => {
+            svg.selectAll(".circleScale").attr("transform", `scale(1)`);
+            tooltip.style("display", "none");
+            this.tooltipSessionId = undefined;
+        }
+
+        const handleMouseClick = (event: any) => {
+
+            const mouseX = _d3.pointer(event)[0];
+            const mouseY = _d3.pointer(event)[1];
+
+            const closest = closestData(mouseX, mouseY);
+            if (closest && this.onClickCallback) {
+                this.onClickCallback(closest);
+            }
+        }
+
+        svg.on("mousemove", (event) => handleMouseMove(event))
+        svg.on("mouseleave", (event) => handleMouseLeave(event))
+        svg.on("click", (event) => handleMouseClick(event))
+
+    }
+
+    updateTooltip(show: boolean, x?: number, y?: number, html?: string) {
+        if (!this.tooltip) {
+            return;
+        }
+
+        x = x ?? 0;
+        y = y ?? 0;
+
+        this.tooltip
+            .style("display", "block")
+            .html(html ?? "")
+            .style("position", `absolute`)
+            .style("width", `max-content`)
+            .style("top", (y - 78) + "px")
+            .style("left", `${x}px`)
+            .style("transform", "translateX(-108%)")
+            .style("font-family", "Horde Open Sans Semibold")
+            .style("font-size", "10px")
+            .style("line-height", "16px")
+            .style("shapeRendering", "crispEdges")
+            .style("stroke", "none")
+
+    }
+
+    clear(container: HTMLDivElement) {
+        d3.select(container).selectAll("*").remove();
+    }
+
+    // test ref
+    test: TestNameRef;
+
+    commonMetaKeys?: string[];
+    onClickCallback?: (session: TestSessionResult) => void;
+
+    handler: TestDataHandler;
+    margin: { top: number, right: number, bottom: number, left: number }
+
+    hasRendered = false;
+    forceRender = false;
+
+    status: TestMetaStatus;
+    metaRefs: MetadataRef[] = [];
+    sessions: TestSessionResult[] = [];
+    // metaRef => status string
+    metaStatus: Map<MetadataRef, string> = new Map();
+    // metaRef => meta string
+    metaNames: Map<MetadataRef, string> = new Map();
+
+    clipId: string;
+
+    tooltip?: DivSelectionType;
+    tooltipSessionId?: string;
+}
+
+type PieChartStack = {value: number, name: string, color: string, title?: boolean};
+
+const PieChart = (container: SVGGElement, data: PieChartStack[], radius: number, innerRadius: number = 0, padAngle: number = 0) => {
+    const pie = d3.pie<PieChartStack>().padAngle(padAngle).value((d) => d.value).sort(null);
+    const arc = d3.arc().innerRadius(innerRadius).outerRadius(radius);
+    const color = d3.scaleOrdinal()
+        .domain(data.map(d => d.name))
+        .range(data.map(d => d.color));
+    
+    const g = d3.select(container)
+        .selectAll("pie")
+        .data(pie(data))
+        .join("path")
+            .attr("class", "pie")
+            .attr("fill", d => color(d.data.name) as any)
+            .attr("d", arc as any)
+        .append("title")
+            .text(d => d.data.title? `${Math.ceil(d.data.value)}% ${d.data.name}` : "");
+
+    return g;
+  }
+
+  const SessionPieChart = (container: SVGGElement, session: TestSessionResult, radius: number, innerRadius:number = 0, padAngle: number = 0) => {
+
+    const sessionTotalCount = session.phaseSucceededCount + session.phaseFailedCount + session.phaseUnspecifiedCount;
+    const metaFailedFactor = Math.ceil(session.phaseFailedCount / (sessionTotalCount || 1) * 10) / 10;
+    const metaUnspecifiedFactor = Math.ceil(session.phaseUnspecifiedCount / (sessionTotalCount || 1) * 10) / 10;
+ 
+    const statusColors = dashboard.getStatusColors();
+
+    const stack: PieChartStack[] = [
+       {
+           value: metaUnspecifiedFactor * 100,
+           name: "Unspecified",
+           color: statusColors.get(StatusColor.Unspecified)!,
+       },
+       {
+           value: metaFailedFactor * 100,
+           name: "Failure",
+           color: statusColors.get(StatusColor.Failure)!,
+       },
+       {
+           value: (1 - (metaFailedFactor + metaUnspecifiedFactor)) * 100,
+           name: "Passed",
+           color: statusColors.get(StatusColor.Success)!
+       }
+    ];
+
+    return PieChart(container, stack, radius, innerRadius, padAngle);
+  }
