Horde: Ensure agent registration token is used in dedicated mode

diff --git a/Engine/Source/Programs/Horde/HordeAgent/Services/GrpcService.cs b/Engine/Source/Programs/Horde/HordeAgent/Services/GrpcService.cs
--- a/Engine/Source/Programs/Horde/HordeAgent/Services/GrpcService.cs
+++ b/Engine/Source/Programs/Horde/HordeAgent/Services/GrpcService.cs
@@ -42,15 +42,6 @@ namespace HordeAgent.Services
 			_loggerFactory = loggerFactory;
 		}
 
-		/// <summary>
-		/// Create a GRPC channel with a default token
-		/// </summary>
-		/// <returns>New grpc channel</returns>
-		public Task<GrpcChannel> CreateGrpcChannelAsync(CancellationToken cancellationToken)
-		{
-			return CreateGrpcChannelAsync(_serverProfile.Token, cancellationToken);
-		}
-
 		/// <summary>
 		/// Create a GRPC channel with the given bearer token
 		/// </summary>
diff --git a/Engine/Source/Programs/Horde/HordeAgent/Services/SessionFactory.cs b/Engine/Source/Programs/Horde/HordeAgent/Services/SessionFactory.cs
--- a/Engine/Source/Programs/Horde/HordeAgent/Services/SessionFactory.cs
+++ b/Engine/Source/Programs/Horde/HordeAgent/Services/SessionFactory.cs
@@ -167,12 +167,17 @@ namespace HordeAgent.Services
 			{
 				logger.LogInformation("Cached server <-> agent registration info loaded from {File}. AgentId={AgentId}", registrationFile, registrationInfo.Id);
 			}
+			
+			// Try using the token returned from agent registration.
+			// For a pre-shared token, this new registration token is to be used (dedicated agent)
+			// A workstation agent never receives a registration as it will always authenticate as the current user
+			string? accessToken = !String.IsNullOrEmpty(registrationInfo.Token) ? registrationInfo.Token : serverProfile.GetAuthToken();
 
 			// Create the session
 			statusService.SetDescription(AgentStatusMessage.ConnectingToServer);
 
 			RpcCreateSessionResponse createSessionResponse;
-			await using (HordeClient sessionClient = hordeClientFactory.Create(accessToken: serverProfile.GetAuthToken()))
+			await using (HordeClient sessionClient = hordeClientFactory.Create(accessToken: accessToken))
 			{
 				HordeRpc.HordeRpcClient rpcClient = await sessionClient.CreateGrpcClientAsync<HordeRpc.HordeRpcClient>(cancellationToken);
 
