Fix hitching when hovering large matrix of historic data
Add history filtering by metadata
Adjust styling of dropdowns

#rnx

diff --git a/Engine/Source/Programs/Horde/HordeDashboard/plugins/test-automation/testAutomationCommon.tsx b/Engine/Source/Programs/Horde/HordeDashboard/plugins/test-automation/testAutomationCommon.tsx
--- a/Engine/Source/Programs/Horde/HordeDashboard/plugins/test-automation/testAutomationCommon.tsx
+++ b/Engine/Source/Programs/Horde/HordeDashboard/plugins/test-automation/testAutomationCommon.tsx
@@ -1,137 +1,185 @@
-import { mergeStyles, mergeStyleSets } from "@fluentui/react";
+import { mergeStyles, mergeStyleSets, Text, ComboBox, IComboBoxOption, SelectableOptionMenuItemType } from "@fluentui/react";
 import { TestSessionResult } from "./testData";
 import dashboard, { StatusColor } from "horde/backend/Dashboard";
+import { useEffect, useState } from "react";
+
+export const MultiOptionChooser: React.FC<{options: IComboBoxOption[], initialSelection: string[], updateKeys: (selectedKeys: string[]) => void, placeholder?: string, style?: any }> = ({ options, initialSelection, updateKeys, placeholder, style }) => {
+    const [selection, setSelection] = useState<string[]>(initialSelection);
+    const optionsWithSelectAll = [...options];
+
+    useEffect(() => {
+        // keep input selection in sync
+        if (selection.toString() !== initialSelection.toString()) {
+            setSelection(initialSelection);
+        }
+    }, [initialSelection]);
+
+    if (selection.length === options.length) {
+        selection.push('selectAll');
+    }
+
+    optionsWithSelectAll.unshift({ key: 'selectAll', text: 'Select All', itemType: SelectableOptionMenuItemType.SelectAll });
+
+    return <ComboBox styles={{ root: style }} placeholder={placeholder ?? "None"} selectedKey={selection} multiSelect options={optionsWithSelectAll}
+        onChange={(_, option) => {
+            if (option) {
+                if (option.itemType === SelectableOptionMenuItemType.SelectAll) {
+                    setSelection(option.selected? optionsWithSelectAll.map((o) => o.key as string) : []);
+                    return;
+                }
+                const key = option.key as string;
+                const index = selection.indexOf(key);
+                if (index === -1 && option.selected) {
+                    selection.push(key);
+                } else if (index >= 0 && !option.selected) {
+                    selection.splice(index, 1);
+                }
+                setSelection([...selection]);
+            }
+        }}
+        onMenuDismissed={() => {
+            updateKeys(selection.filter((k) => k !== 'selectAll'));
+        }} />
+};
 
 const styles = mergeStyleSets({
 
-   stripes: {
-      backgroundImage: 'repeating-linear-gradient(-45deg, rgba(255, 255, 255, .2) 25%, transparent 25%, transparent 50%, rgba(255, 255, 255, .2) 50%, rgba(255, 255, 255, .2) 75%, transparent 75%, transparent)',
-   }
+    stripes: {
+        backgroundImage: 'repeating-linear-gradient(-45deg, rgba(255, 255, 255, .2) 25%, transparent 25%, transparent 50%, rgba(255, 255, 255, .2) 50%, rgba(255, 255, 255, .2) 75%, transparent 75%, transparent)',
+    }
 
 });
 
 export type StatusBarStack = {
-   value: number,
-   title?: string,
-   titleValue?: number,
-   color?: string,
-   onClick?: () => void,
-   stripes?: boolean,
-
+    value: number,
+    title?: string,
+    titleValue?: number,
+    color?: string,
+    onClick?: () => void,
+    stripes?: boolean,
 }
 
 export const StatusBar = (stack: StatusBarStack[], width: number, height: number, basecolor?: string, style?: any): JSX.Element => {
 
-   stack = stack.filter(s => s.value > 0);
-
-   const mainTitle = stack.map((item) => {
-      return item.titleValue === undefined ? `${Math.ceil(item.value)}% ${item.title}` : `${item.titleValue} ${item.title}`
-   }).join(' ');
-
-   return (
-      <div className={mergeStyles({ backgroundColor: basecolor, width: width, height: height, verticalAlign: 'middle', display: "flex" }, style)} title={mainTitle}>
-         {stack.map((item) => <span key={item.title!}
-            onClick={item.onClick}
-            className={item.stripes ? styles.stripes : undefined}
-            style={{
-               width: `${item.value}%`, height: '100%',
-               backgroundColor: item.color,
-               display: 'block',
-               cursor: item.onClick ? 'pointer' : 'inherit',
-               backgroundSize: `${height * 2}px ${height * 2}px`
-            }} />)}
-      </div>
-   );
+    stack = stack.filter(s => s.value > 0);
+
+    const mainTitle = stack.map((item) => {
+        return item.titleValue === undefined ? `${Math.ceil(item.value)}% ${item.title}` : `${item.titleValue} ${item.title}`
+    }).join(' ');
+
+    return (
+        <div className={mergeStyles({ backgroundColor: basecolor, width: width, height: height, verticalAlign: 'middle', display: "flex" }, style)} title={mainTitle}>
+            {stack.map((item) => <span key={item.title!}
+                onClick={item.onClick}
+                className={item.stripes ? styles.stripes : undefined}
+                style={{
+                width: `${item.value}%`, height: '100%',
+                backgroundColor: item.color,
+                display: 'block',
+                cursor: item.onClick ? 'pointer' : 'inherit',
+                backgroundSize: `${height * 2}px ${height * 2}px`
+                }} />)}
+        </div>
+    );
 }
 
 export const SessionStatusBar = (session: TestSessionResult, width: number, height: number): JSX.Element | undefined => {
-   const sessionTotalCount = session.phaseSucceededCount + session.phaseFailedCount + session.phaseUnspecifiedCount;
-   const metaFailedFactor = Math.ceil(session.phaseFailedCount / (sessionTotalCount || 1) * 20) / 20;
-   const metaUnspecifiedFactor = Math.ceil(session.phaseUnspecifiedCount / (sessionTotalCount || 1) * 20) / 20;
-
-   const bDisplayBar = (metaFailedFactor + metaUnspecifiedFactor) > 0 && (metaFailedFactor + metaUnspecifiedFactor) < 1;
-   if (!bDisplayBar) return <div/>
-
-   const statusColors = dashboard.getStatusColors();
-
-   const stack: StatusBarStack[] = [
-      {
-          value: metaUnspecifiedFactor * 100,
-          title: "Unspecified",
-          color: statusColors.get(StatusColor.Unspecified)!,
-          stripes: true
-      },
-      {
-          value: metaFailedFactor * 100,
-          title: "Failure",
-          color: statusColors.get(StatusColor.Failure)!,
-          stripes: true
-      }
-  ];
-
-  return StatusBar(stack, width, height, statusColors.get(StatusColor.Success)!, { margin: '3px !important' });
+    const sessionTotalCount = session.phaseSucceededCount + session.phaseFailedCount + session.phaseUnspecifiedCount;
+    const metaFailedFactor = Math.ceil(session.phaseFailedCount / (sessionTotalCount || 1) * 20) / 20;
+    const metaUnspecifiedFactor = Math.ceil(session.phaseUnspecifiedCount / (sessionTotalCount || 1) * 20) / 20;
+
+    const bDisplayBar = (metaFailedFactor + metaUnspecifiedFactor) > 0 && (metaFailedFactor + metaUnspecifiedFactor) < 1;
+    if (!bDisplayBar) return undefined
+
+    const statusColors = dashboard.getStatusColors();
+
+    const stack: StatusBarStack[] = [
+        {
+            value: metaUnspecifiedFactor * 100,
+            title: "Unspecified",
+            color: statusColors.get(StatusColor.Unspecified)!,
+            stripes: true
+        },
+        {
+            value: metaFailedFactor * 100,
+            title: "Failure",
+            color: statusColors.get(StatusColor.Failure)!,
+            stripes: true
+        }
+    ];
+
+    return StatusBar(stack, width, height, statusColors.get(StatusColor.Success)!, { margin: '3px !important' });
 }
 
 export const StatusPie = (stack: StatusBarStack[], radius: number, basecolor?: string, style?: any): JSX.Element => {
 
-   const width = radius * 2;
-   const height = width;
-
-   stack = stack.filter(s => s.value > 0);
-   let offset = 0;
-   const stackOffsets = stack.map(s => {
-      const stepOffset = offset;
-      offset += s.value;
-      return stepOffset;
-   });
-
-   const mainTitle = stack.map((item) => {
-      return item.titleValue === undefined ? `${Math.ceil(item.value)}% ${item.title}` : `${item.titleValue} ${item.title}`
-   }).join(' ');
-
-   return (
-      <div className={mergeStyles({ width: width, height: height, verticalAlign: 'middle', display: "flex" }, style)} title={mainTitle}>
-         <svg width={width} height={height} viewBox={`0 0 64 64`}>
-            {stack.map((item, index) => <circle key={item.title!} r="25%" cx="50%" cy="50%" strokeWidth="50%" transform={`rotate(-90) translate(-64)`} fill="transparent"
-                                    stroke={item.color}
-                                    strokeDashoffset={`-${stackOffsets[index]}`}
-                                    strokeDasharray={`${item.value} 100`}
-                                    onClick={item.onClick} />)
-            }
-         </svg>
-      </div>
-   );
+    const width = radius * 2;
+    const height = width;
+
+    stack = stack.filter(s => s.value > 0);
+    let offset = 0;
+    const stackOffsets = stack.map(s => {
+        const stepOffset = offset;
+        offset += s.value;
+        return stepOffset;
+    });
+
+    const mainTitle = stack.map((item) => {
+        return item.titleValue === undefined ? `${Math.ceil(item.value)}% ${item.title}` : `${item.titleValue} ${item.title}`
+    }).join(' ');
+
+    return (
+        <div className={mergeStyles({ width: width, height: height, verticalAlign: 'middle', display: "flex" }, style)} title={mainTitle}>
+            <svg width={width} height={height} viewBox={`0 0 64 64`}>
+                {stack.map((item, index) => <circle key={item.title!} r="25%" cx="50%" cy="50%" strokeWidth="50%" transform={`rotate(-90) translate(-64)`} fill="transparent"
+                                        stroke={item.color}
+                                        strokeDashoffset={`-${stackOffsets[index]}`}
+                                        strokeDasharray={`${item.value} 100`}
+                                        onClick={item.onClick} />)
+                }
+            </svg>
+        </div>
+    );
 }
 
 export const SessionStatusPie = (session: TestSessionResult, radius: number): JSX.Element | undefined => {
-   const sessionTotalCount = session.phaseSucceededCount + session.phaseFailedCount + session.phaseUnspecifiedCount;
-   const metaFailedFactor = Math.ceil(session.phaseFailedCount / (sessionTotalCount || 1) * 20) / 20;
-   const metaUnspecifiedFactor = Math.ceil(session.phaseUnspecifiedCount / (sessionTotalCount || 1) * 20) / 20;
-
-   const bDisplayBar = (metaFailedFactor + metaUnspecifiedFactor) > 0 && (metaFailedFactor + metaUnspecifiedFactor) < 1;
-   if (!bDisplayBar) return <div/>
-
-   const statusColors = dashboard.getStatusColors();
-
-   const stack: StatusBarStack[] = [
-      {
-          value: metaUnspecifiedFactor * 100,
-          title: "Unspecified",
-          color: statusColors.get(StatusColor.Unspecified)!,
-          stripes: true
-      },
-      {
-          value: metaFailedFactor * 100,
-          title: "Failure",
-          color: statusColors.get(StatusColor.Failure)!,
-          stripes: true
-      },
-      {
-          value: (1 - (metaFailedFactor + metaUnspecifiedFactor)) * 100,
-          title: "Passed",
-          color: statusColors.get(StatusColor.Success)!
-      }
-  ];
-
-  return StatusPie(stack, radius);
+    const sessionTotalCount = session.phaseSucceededCount + session.phaseFailedCount + session.phaseUnspecifiedCount;
+    const metaFailedFactor = Math.ceil(session.phaseFailedCount / (sessionTotalCount || 1) * 20) / 20;
+    const metaUnspecifiedFactor = Math.ceil(session.phaseUnspecifiedCount / (sessionTotalCount || 1) * 20) / 20;
+
+    const bDisplayBar = (metaFailedFactor + metaUnspecifiedFactor) > 0 && (metaFailedFactor + metaUnspecifiedFactor) < 1;
+    if (!bDisplayBar) return;
+
+    const statusColors = dashboard.getStatusColors();
+
+    const stack: StatusBarStack[] = [
+        {
+            value: metaUnspecifiedFactor * 100,
+            title: "Unspecified",
+            color: statusColors.get(StatusColor.Unspecified)!,
+            stripes: true
+        },
+        {
+            value: metaFailedFactor * 100,
+            title: "Failure",
+            color: statusColors.get(StatusColor.Failure)!,
+            stripes: true
+        },
+        {
+            value: (1 - (metaFailedFactor + metaUnspecifiedFactor)) * 100,
+            title: "Passed",
+            color: statusColors.get(StatusColor.Success)!
+        }
+    ];
+
+    return StatusPie(stack, radius);
+}
+
+export const SessionValues = (session: TestSessionResult, style?: any): JSX.Element => {
+    const statusColors = dashboard.getStatusColors();
+    return <Text style={style}>
+                {!!session.phaseUnspecifiedCount && <span style={{color: statusColors.get(StatusColor.Unspecified)}}> Unspecified: {session.phaseUnspecifiedCount}</span>}
+                {!!session.phaseFailedCount && <span style={{color: statusColors.get(StatusColor.Failure)}}> Failed: {session.phaseFailedCount}</span>}
+                {!!session.phaseSucceededCount && <span style={{color: statusColors.get(StatusColor.Success)}}> Success: {session.phaseSucceededCount}</span>}
+            </Text>
 }
diff --git a/Engine/Source/Programs/Horde/HordeDashboard/plugins/test-automation/testAutomationSummary.tsx b/Engine/Source/Programs/Horde/HordeDashboard/plugins/test-automation/testAutomationSummary.tsx
--- a/Engine/Source/Programs/Horde/HordeDashboard/plugins/test-automation/testAutomationSummary.tsx
+++ b/Engine/Source/Programs/Horde/HordeDashboard/plugins/test-automation/testAutomationSummary.tsx
@@ -1,4 +1,4 @@
-import { DefaultButton, FontIcon, IconButton, mergeStyleSets, Modal, Stack, Text, HoverCard, HoverCardType, IPlainCardProps, DirectionalHint, Spinner, SpinnerSize, IContextualMenuItem, IContextualMenuProps } from "@fluentui/react";
+import { DefaultButton, FontIcon, IconButton, mergeStyleSets, Modal, Stack, Text, HoverCard, HoverCardType, IPlainCardProps, DirectionalHint, Spinner, SpinnerSize, IContextualMenuItem, IContextualMenuProps, IComboBoxOption } from "@fluentui/react";
 import { useState, useEffect } from "react";
 import { observer } from "mobx-react-lite";
 import { TestDataHandler, TestNameRef, MetadataRef, TestSessionStatus, TestSessionResult, TestStatus } from "./testData";
@@ -7,13 +7,13 @@ import dashboard, { StatusColor } from "horde/backend/Dashboard";
 import { getShortNiceTime } from "horde/base/utilities/timeUtils";
 import { getHordeStyling } from "horde/styles/Styles";
 import { projectStore } from 'horde/backend/ProjectStore';
-import { SessionStatusBar, StatusBar, StatusBarStack } from "./testAutomationCommon";
+import { SessionStatusBar, SessionValues, StatusBar, StatusBarStack, MultiOptionChooser } from "./testAutomationCommon";
 import { TestHistoryGraph } from "./testHistoryGraph";
 
 type SessionCallback = (session: TestSessionResult) => void;
 type HistorySessionCallback = (test: TestNameRef, onClick: SessionCallback) => void;
 
-const TestHistoryGraphWidget: React.FC<{ test: TestNameRef, stream: string, handler: TestDataHandler, onClick: SessionCallback, commonMetaKeys?: string[]}> = ({ test, stream, handler, onClick, commonMetaKeys }) => {
+const TestHistoryGraphWidget: React.FC<{ test: TestNameRef, stream: string, handler: TestDataHandler, onClick: SessionCallback, commonMetaKeys?: string[], selectedMetaKeys?: string[]}> = ({ test, stream, handler, onClick, commonMetaKeys, selectedMetaKeys }) => {
     const [container, setContainer] = useState<HTMLDivElement | null>(null);
     const [state, setState] = useState<{ graph?: TestHistoryGraph }>({});
 
@@ -25,17 +25,21 @@ const TestHistoryGraphWidget: React.FC<{ test: TestNameRef, stream: string, hand
         setup();
     }, []);
 
+    useEffect(() => {
+        // refresh on meta selection update
+        container && state.graph?.render(container, true, selectedMetaKeys);
+    }, [selectedMetaKeys])
+
     if (!state.graph) {
         return null;
     }
 
     if (container) {
         try {
-            state.graph?.render(container);
+            state.graph?.render(container, false, selectedMetaKeys);
         } catch (err) {
             console.error(err);
         }
-
     }
 
     const { hordeClasses } = getHordeStyling();
@@ -48,14 +52,14 @@ const TestHistoryGraphWidget: React.FC<{ test: TestNameRef, stream: string, hand
             </Stack>
 }
 
-
-const StreamSelector: React.FC<{ streams: string[], onClick: (stream: string) => void }> = ({ streams, onClick }) => {
+const StreamSelector: React.FC<{ streams: string[], onClick: (stream: string) => void, selected?: string }> = ({ streams, onClick, selected }) => {
     const options: IContextualMenuItem[] = [];
     streams.forEach(a => {
         options.push({
             key: `streams_${a}`,
             text: projectStore.streamById(a)?.fullname ?? a,
-            onClick: () => onClick(a)
+            onClick: () => onClick(a),
+            disabled: a === selected
         });
     });
 
@@ -65,13 +69,14 @@ const StreamSelector: React.FC<{ streams: string[], onClick: (stream: string) =>
         items: options,
     };
 
-    return <DefaultButton style={{ height: 22, fontSize: 11, whiteSpace: 'nowrap' }} menuProps={menuProps} text="Compare with"/>
+    return <DefaultButton style={{ height: 22, fontSize: 11, whiteSpace: 'nowrap', borderWidth: 0 }} menuProps={menuProps} text={selected? projectStore.streamById(selected)?.fullname ?? selected : "Compare stream"}/>
 }
 
 export const TestHistoryView: React.FC<{ test: TestNameRef, handler: TestDataHandler, onClick: SessionCallback  }> = observer(({ test, handler, onClick }) => {
 
     const { hordeClasses } = getHordeStyling();
     const [streams, setStreams] = useState<{current: string, toCompare?: string}>({current: handler.stream!});
+    const [selectedMeta, setSelectedMeta] = useState<Set<string>>(new Set());
 
     handler.subscribeToSubQueryLoading();
 
@@ -90,6 +95,12 @@ export const TestHistoryView: React.FC<{ test: TestNameRef, handler: TestDataHan
     const commonMetaKeys = MetadataRef.identifyCommonKeys(metadataSet.values().toArray());
     const commonMetaValues = metadataSet.values().next().value?.getCommonValues(commonMetaKeys).join(" / ");
 
+    const filteredMetadata = new Set(handler.filteredMetadata);
+    const metaOptions: IComboBoxOption[] = metadataSet.keys()
+        .filter(meta => filteredMetadata.has(meta))
+        .map((meta) => ({key: meta.id, text: meta.getValuesExcept(commonMetaKeys).join(' / ')}))
+        .toArray().sort((a, b) => a.text.localeCompare(b.text));
+
     const scolors = dashboard.getStatusColors();
     const streamViews = (!!availableStreams.length? [streams.current, streams.toCompare] : [streams.current]).map(stream => {
 
@@ -119,7 +130,7 @@ export const TestHistoryView: React.FC<{ test: TestNameRef, handler: TestDataHan
                 }
                 { !!status &&
                     <Stack style={{ paddingTop: 4 }}>
-                        <TestHistoryGraphWidget test={test} stream={stream} commonMetaKeys={commonMetaKeys} handler={handler} onClick={onClick} />
+                        <TestHistoryGraphWidget test={test} stream={stream} commonMetaKeys={commonMetaKeys} selectedMetaKeys={selectedMeta.size? selectedMeta.keys().toArray() : undefined} handler={handler} onClick={onClick} />
                     </Stack>
                 }
                 { !status && !!stream &&
@@ -137,7 +148,7 @@ export const TestHistoryView: React.FC<{ test: TestNameRef, handler: TestDataHan
 
     return <Stack>
             <Stack className={hordeClasses.raised}>
-                <Stack style={{ paddingBottom: 12 }}>
+                <Stack>
                     <Stack horizontal>
                         <Stack className="horde-no-darktheme" style={{ paddingTop: 3, paddingRight: 4 }}>
                             <FontIcon style={{ color: anyFailed ? scolors.get(StatusColor.Failure)! : scolors.get(StatusColor.Success)! }} iconName="Square" />
@@ -147,6 +158,19 @@ export const TestHistoryView: React.FC<{ test: TestNameRef, handler: TestDataHan
                         </Stack>
                     </Stack>
                 </Stack>
+                <Stack horizontal tokens={{childrenGap: 15}} style={{ paddingBottom: 12, paddingTop: 5 }}>
+                    <Stack style={{ width: 170 }}>
+                        <MultiOptionChooser style={{height: 22, fontSize: 11}} placeholder="Filter by metadata"
+                            options={metaOptions} initialSelection={selectedMeta.keys().toArray()}
+                            updateKeys={(keys) => setSelectedMeta(new Set(keys))} />
+                    </Stack>
+                    <Stack>
+                        <StreamSelector streams={availableStreams} selected={streams.toCompare} onClick={(targetStream: string) => {
+                            setStreams({...streams, toCompare: undefined});
+                            handler.query(targetStream, true).catch((reason) => console.error(reason)).finally(() => setStreams({...streams, toCompare: targetStream}))
+                        }} />
+                    </Stack>
+                </Stack>
                 <Stack style={{ paddingLeft: 12 }} tokens={{ childrenGap: 4 }}>
                     {streamViews}
                 </Stack>
@@ -156,14 +180,6 @@ export const TestHistoryView: React.FC<{ test: TestNameRef, handler: TestDataHan
                         <Spinner size={SpinnerSize.small} />
                     </Stack>
                 }
-                <Stack horizontal>
-                    <Stack style={{ width: 170 }}>
-                        <StreamSelector streams={availableStreams} onClick={(targetStream: string) => {
-                            setStreams({...streams, toCompare: undefined});
-                            handler.query(targetStream, true).catch((reason) => console.error(reason)).finally(() => setStreams({...streams, toCompare: targetStream}))
-                        }} />
-                    </Stack>
-                </Stack>
             </Stack>
         </Stack>
 });
@@ -211,11 +227,7 @@ const TestSummaryCard: React.FC<{ test: TestNameRef, sessions: [MetadataRef, Tes
     const metaElements: JSX.Element[] = [];
     const onRenderStatusCard = (session: TestSessionResult): JSX.Element => {
         return <Stack key={`test_summary_hover_${test.id}`} style={{pointerEvents: 'none'}}>
-                    <Text style={{fontSize: 12, padding: '8px'}}>
-                        {!!session.phaseUnspecifiedCount && <span style={{color: statusColors.get(StatusColor.Unspecified)}}> Unspecified: {session.phaseUnspecifiedCount}</span>}
-                        {!!session.phaseFailedCount && <span style={{color: statusColors.get(StatusColor.Failure)}}> Failed: {session.phaseFailedCount}</span>}
-                        {!!session.phaseSucceededCount && <span style={{color: statusColors.get(StatusColor.Success)}}> Success: {session.phaseSucceededCount}</span>}
-                    </Text>
+                    {SessionValues(session, {fontSize: 12, padding: '8px'})}
                 </Stack>
     }
  
diff --git a/Engine/Source/Programs/Horde/HordeDashboard/plugins/test-automation/testAutomationView.tsx b/Engine/Source/Programs/Horde/HordeDashboard/plugins/test-automation/testAutomationView.tsx
--- a/Engine/Source/Programs/Horde/HordeDashboard/plugins/test-automation/testAutomationView.tsx
+++ b/Engine/Source/Programs/Horde/HordeDashboard/plugins/test-automation/testAutomationView.tsx
@@ -1,4 +1,4 @@
-import { Stack, FontIcon, IComboBox, IComboBoxOption, ComboBox, IContextualMenuItem, IContextualMenuProps, Text, Spinner, SpinnerSize, Label, DefaultButton, SelectableOptionMenuItemType, TagItemSuggestion, ITag, TagPicker } from "@fluentui/react";
+import { Stack, FontIcon, IComboBox, IComboBoxOption, ComboBox, IContextualMenuItem, IContextualMenuProps, Text, Spinner, SpinnerSize, Label, DefaultButton, TagItemSuggestion, ITag, TagPicker } from "@fluentui/react";
 import { Breadcrumbs } from "horde/components/Breadcrumbs";
 import { TopNav } from "horde/components/TopNav";
 import { getHordeStyling } from "horde/styles/Styles";
@@ -8,6 +8,7 @@ import React, { useEffect, useState, useRef } from "react";
 import { projectStore } from 'horde/backend/ProjectStore';
 import { TestDataHandler, TestDataVersionRegistrar } from "./testData";
 import { TestSummary } from "./testAutomationSummary";
+import { MultiOptionChooser } from "./testAutomationCommon";
 
 import { TestDataV2 } from "./testDataV2Fetcher";
 
@@ -16,45 +17,6 @@ const handler = new TestDataHandler();
 // register TestDataV2 fetcher
 TestDataVersionRegistrar.register(new TestDataV2);
 
-const MultiOptionChooser: React.FC<{options: IComboBoxOption[], initialSelection: string[], updateKeys: (selectedKeys: string[]) => void }> = ({ options, initialSelection, updateKeys }) => {
-    const [selection, setSelection] = useState<string[]>(initialSelection);
-    const optionsWithSelectAll = [...options];
-
-    useEffect(() => {
-        // keep input selection in sync
-        if (selection.toString() !== initialSelection.toString()) {
-            setSelection(initialSelection);
-        }
-    }, [initialSelection]);
-
-    if (selection.length === options.length) {
-        selection.push('selectAll');
-    }
-
-    optionsWithSelectAll.unshift({ key: 'selectAll', text: 'Select All', itemType: SelectableOptionMenuItemType.SelectAll });
-
-    return <ComboBox styles={{ root: { height: 30 } }} placeholder="None" selectedKey={selection} multiSelect options={optionsWithSelectAll}
-        onChange={(_, option) => {
-            if (option) {
-                if (option.itemType === SelectableOptionMenuItemType.SelectAll) {
-                    setSelection(option.selected? optionsWithSelectAll.map((o) => o.key as string) : []);
-                    return;
-                }
-                const key = option.key as string;
-                const index = selection.indexOf(key);
-                if (index === -1 && option.selected) {
-                    selection.push(key);
-                } else if (index >= 0 && !option.selected) {
-                    selection.splice(index, 1);
-                }
-                setSelection([...selection]);
-            }
-        }}
-        onMenuDismissed={() => {
-            updateKeys(selection.filter((k) => k !== 'selectAll'));
-        }} />
-};
-
 const StreamChooser: React.FC = observer(() => {
 
     handler.subscribeToFilter();
@@ -78,7 +40,7 @@ const StreamChooser: React.FC = observer(() => {
 
     return <Stack style={{ paddingTop: 12, paddingBottom: 4 }} horizontal>
             <Label style={{ paddingTop: 5, paddingRight: 5 }}>Stream</Label>
-            <DefaultButton style={{ width: 270, textAlign: "left" }} menuProps={menuProps} text={stream ? projectStore.streamById(stream)?.fullname : "Select"} />
+            <DefaultButton style={{ width: 270, textAlign: "left", height: 30 }} menuProps={menuProps} text={stream ? projectStore.streamById(stream)?.fullname : "Select"} />
         </Stack>   
     });
 
@@ -124,6 +86,7 @@ const TestChooser: React.FC = observer(() => {
             <Label style={{ paddingTop: 5, paddingRight: 5 }}>Tests</Label>
             <Stack tokens={{ childrenGap: 6 }} style={{ paddingLeft: 12 }} grow>
                 <TagPicker
+                        styles={{ input: { height: 28 }, itemsWrapper: { height: 28 } }}
                         onRenderSuggestionsItem={onRenderSuggestionItem}
                         removeButtonAriaLabel="Remove"
                         selectionAriaLabel="Selected tests"
@@ -177,7 +140,7 @@ const MetaKeyChooser: React.FC<{ onUpdateKey: (key: string) => void }> = observe
 
    return <Stack style={{ paddingTop: 12, paddingBottom: 4 }} horizontal>
             <Label style={{ paddingTop: 5, paddingRight: 5, whiteSpace: 'nowrap' }}>Filter by</Label>
-            <DefaultButton style={{width: 270, textAlign: "left" }} menuProps={menuProps} text="Metadata" />
+            <DefaultButton style={{width: 270, textAlign: "left", height: 30 }} menuProps={menuProps} text="Metadata" />
         </Stack>   
 });
 
@@ -197,7 +160,7 @@ const MetaValueChooser: React.FC<{ metakey: string, onRemove: (key: string) => v
    return <Stack style={{ paddingTop: 12, paddingBottom: 4, }} horizontal>
             <Label style={{ paddingTop: 5, paddingRight: 5 }}>{metakey}</Label>
             <Stack style={{ padding: 1, width: 270 }}>
-                <MultiOptionChooser options={options} initialSelection={cmetavalues.keys().toArray()}
+                <MultiOptionChooser style={{height: 28}} options={options} initialSelection={cmetavalues.keys().toArray()}
                 updateKeys={(keys) => {
                     metavalues.forEach(t => {
                     const selected = keys.includes(t);
@@ -209,17 +172,17 @@ const MetaValueChooser: React.FC<{ metakey: string, onRemove: (key: string) => v
                     }
                     });
                 }} />
+            </Stack>
+            <DefaultButton style={{ minWidth: 0, fontSize: 11, padding: 5, height: 30 }} title="remove filter"
+                    onClick={() => {
+                        if (handler.filterState?.metadata?.keys().some((key) => key === metakey)) {
+                            handler.removeMetadata(metakey);
+                        }
+                        onRemove(metakey);
+                    }}>
+                <FontIcon iconName="CalculatorMultiply" />
+            </DefaultButton>
         </Stack>
-        <DefaultButton style={{ minWidth: 0, fontSize: 11, padding: 5 }} title="remove filter"
-                onClick={() => {
-                    if (handler.filterState?.metadata?.keys().some((key) => key === metakey)) {
-                        handler.removeMetadata(metakey);
-                    }
-                    onRemove(metakey);
-                }}>
-            <FontIcon iconName="CalculatorMultiply" />
-        </DefaultButton>
-    </Stack>
 });
 
 
@@ -283,6 +246,10 @@ const TestAutomationPanel: React.FC = observer(() => {
 
     handler.subscribeToSearch();
 
+    if (!handler.selectedStatusStream) {
+        return null;
+    }
+
     const status = handler.filteredTests;
 
     return <Stack grow>
diff --git a/Engine/Source/Programs/Horde/HordeDashboard/plugins/test-automation/testData.ts b/Engine/Source/Programs/Horde/HordeDashboard/plugins/test-automation/testData.ts
--- a/Engine/Source/Programs/Horde/HordeDashboard/plugins/test-automation/testData.ts
+++ b/Engine/Source/Programs/Horde/HordeDashboard/plugins/test-automation/testData.ts
@@ -30,7 +30,11 @@ export class TestDataVersionRegistrar {
     private static _manifest: TestVersionFetcher[] = [];
 
     static register(fetcher: TestVersionFetcher) {
-        TestDataVersionRegistrar._manifest.push(fetcher);
+        if (TestDataVersionRegistrar._manifest.find(f => f.version === fetcher.version)) {
+            console.warn(`${fetcher.version} already registered`);
+        } else {
+            TestDataVersionRegistrar._manifest.push(fetcher);
+        }
     }
 
     static get manifest(): ReadonlyArray<TestVersionFetcher> {
diff --git a/Engine/Source/Programs/Horde/HordeDashboard/plugins/test-automation/testHistoryGraph.ts b/Engine/Source/Programs/Horde/HordeDashboard/plugins/test-automation/testHistoryGraph.ts
--- a/Engine/Source/Programs/Horde/HordeDashboard/plugins/test-automation/testHistoryGraph.ts
+++ b/Engine/Source/Programs/Horde/HordeDashboard/plugins/test-automation/testHistoryGraph.ts
@@ -4,7 +4,7 @@ import { TestDataHandler, TestNameRef, MetadataRef, TestMetaStatus, TestSessionR
 import dashboard, { StatusColor } from 'horde/backend/Dashboard';
 import { getHumanTime, getShortNiceTime, msecToElapsed } from 'horde/base/utilities/timeUtils';
 import { getHordeStyling } from 'horde/styles/Styles';
-import { SessionStatusBar } from "./testAutomationCommon";
+import { SessionStatusBar, SessionValues } from "./testAutomationCommon";
 import { renderToString } from "react-dom/server"
 import { TestOutcome } from "./api";
 
@@ -30,13 +30,14 @@ export class TestHistoryGraph {
         this.clipId = `test_history_clip_path_${test.key}_${stream}`;
     }
 
-    initData() {
-        this.metaRefs = this.handler.filteredMetadata.filter((m) => this.status.includesMetadata(m));
+    initData(selectedMetaKeys?: string[]) {
+        this.metaRefs = this.handler.filteredMetadata.filter((m) => this.status.includesMetadata(m) && (!selectedMetaKeys || selectedMetaKeys.includes(m.id)));
         const commonMetaKeys = this.commonMetaKeys = this.commonMetaKeys ?? MetadataRef.identifyCommonKeys(this.metaRefs);
 
         this.sessions = [];
+        this.metaInfo = new Map();
 
-        this.metaRefs.forEach(meta => {
+        this.metaRefs.forEach((meta, index) => {
             const metaLastSession = this.status.getLastSession(meta);
             this.metaStatus.set(meta, !metaLastSession? TestOutcome.Unspecified : metaLastSession.outcome);
 
@@ -44,12 +45,13 @@ export class TestHistoryGraph {
             this.metaNames.set(meta, elements.join(" / "));
 
             this.sessions.push(...this.status.sessions.get(meta)!.history);
+            this.metaInfo.set(meta.id, {y: 0 , index: index});
         });
     }
 
-    render(container: HTMLDivElement) {
+    render(container: HTMLDivElement, forceRender?: boolean, selectedMetaKeys?: string[]) {
 
-        if (this.hasRendered && !this.forceRender) {
+        if (this.hasRendered && !forceRender) {
             return;
         }
 
@@ -58,9 +60,8 @@ export class TestHistoryGraph {
         this.clear(container);
 
         this.hasRendered = true;
-        this.forceRender = false;
 
-        this.initData();
+        this.initData(selectedMetaKeys);
 
         const handler = this.handler;
         const sessions = this.sessions;
@@ -75,7 +76,7 @@ export class TestHistoryGraph {
         ]);
 
         const X = _d3.map(sessions, (r) => handler.commitIdDates.get(r.commitId)!.getTime() / 1000);
-        const Y = _d3.map(sessions, (r) => this.metaRefs.findIndex((m) => m.id === r.metadataId));
+        const Y = _d3.map(sessions, (r) => this.metaInfo.get(r.metadataId)!.index);
 
         const xDomain = d3.extent(handler.commitIdDates.values() as any, (d:any) => d.getTime() / 1000);
         const yDomain = new _d3.InternSet(Y as any);
@@ -203,10 +204,13 @@ export class TestHistoryGraph {
             .style("border-width", "1px")
             .style("border-radius", "3px")
             .style("border-color", dashboard.darktheme ? "#413F3D" : "#2D3F5F")
-            .style("padding", "8px")
+            .style("padding", "6px")
             .style("position", "absolute")
             .style("pointer-events", "none");
 
+        // populate metaInfo y value
+        this.metaInfo.forEach(meta => meta.y = yScale(meta.index as any)!)
+
         const closestData = (x: number, y: number): TestSessionResult | undefined => {
 
             if (x < this.margin.left - 16) {
@@ -215,12 +219,10 @@ export class TestHistoryGraph {
 
             y -= 16;
 
-            const metaRefs = this.metaRefs;
-            let closest = sessions.reduce((best, session, i) => {
-                let absy = Math.abs(yScale(metaRefs.findIndex((m) => m.id === session.metadataId) as any)! - y)
-                const timeStamp = handler.commitIdDates.get(session.commitId)!.getTime() / 1000;
-                const sx = xScale(timeStamp);
-                let absx = Math.abs(sx - x)
+            const closest = sessions.reduce((best, session, i) => {
+                const absy = Math.abs(this.metaInfo.get(session.metadataId)!.y - y);
+                const sessionTime = handler.commitIdDates.get(session.commitId)!.getTime() / 1000;
+                const absx = Math.abs(xScale(sessionTime) - x);
 
                 const lengthSqr = absy * absy + absx * absx;
 
@@ -253,7 +255,7 @@ export class TestHistoryGraph {
 
             if (this.tooltipSessionId !== closest.id) {
                 this.tooltipSessionId = closest.id;
-                const metaIndex = this.metaRefs.findIndex((m) => m.id === closest.metadataId);
+                const metaIndex = this.metaInfo.get(closest.metadataId)!.index;
 
                 svg.selectAll(".circleScale").attr("transform", `scale(1)`);
                 svg.select(`#circleScale${closest.id}`).attr("transform", `scale(2.5)`);
@@ -263,11 +265,12 @@ export class TestHistoryGraph {
                 const cmeta = this.metaRefs[metaIndex].getValuesExcept(this.commonMetaKeys!).join(" / ");
                 const date = getShortNiceTime(closest.start, true, true);
                 let desc = "";
-                if (!!cmeta) desc += `${cmeta} <br/>`;
+                if (!!cmeta) desc += `# ${cmeta} <br/>`;
                 desc += `Commit ${closest.commitId} <br/>`;
                 desc += `Duration ${msecToElapsed(moment.duration(closest.duration).asMilliseconds(), true, true)} <br/>`;
                 desc += `${date} <br/>`;
-                desc += renderToString(SessionStatusBar(closest, 150, 8));
+                if (closest.phaseSucceededCount + closest.phaseFailedCount + closest.phaseUnspecifiedCount > 1)
+                    desc += renderToString(SessionValues(closest, {fontSize: 9}));
 
                 const timeStamp = handler.commitIdDates.get(closest.commitId)!.getTime() / 1000;
                 const tx = xScale(timeStamp);
@@ -313,9 +316,9 @@ export class TestHistoryGraph {
             .html(html ?? "")
             .style("position", `absolute`)
             .style("width", `max-content`)
-            .style("top", (y - 78) + "px")
+            .style("top", `${y}px`)
             .style("left", `${x}px`)
-            .style("transform", "translateX(-108%)")
+            .style("transform", "translate(-108%, -97%)")
             .style("font-family", "Horde Open Sans Semibold")
             .style("font-size", "10px")
             .style("line-height", "16px")
@@ -338,10 +341,10 @@ export class TestHistoryGraph {
     margin: { top: number, right: number, bottom: number, left: number }
 
     hasRendered = false;
-    forceRender = false;
 
     status: TestMetaStatus;
     metaRefs: MetadataRef[] = [];
+    metaInfo: Map<string, {y: number, index: number}> = new Map();
     sessions: TestSessionResult[] = [];
     // metaRef => status string
     metaStatus: Map<MetadataRef, string> = new Map();
@@ -374,33 +377,33 @@ const PieChart = (container: SVGGElement, data: PieChartStack[], radius: number,
             .text(d => d.data.title? `${Math.ceil(d.data.value)}% ${d.data.name}` : "");
 
     return g;
-  }
+}
 
-  const SessionPieChart = (container: SVGGElement, session: TestSessionResult, radius: number, innerRadius:number = 0, padAngle: number = 0) => {
+const SessionPieChart = (container: SVGGElement, session: TestSessionResult, radius: number, innerRadius:number = 0, padAngle: number = 0) => {
 
     const sessionTotalCount = session.phaseSucceededCount + session.phaseFailedCount + session.phaseUnspecifiedCount;
     const metaFailedFactor = Math.ceil(session.phaseFailedCount / (sessionTotalCount || 1) * 10) / 10;
     const metaUnspecifiedFactor = Math.ceil(session.phaseUnspecifiedCount / (sessionTotalCount || 1) * 10) / 10;
- 
+
     const statusColors = dashboard.getStatusColors();
 
     const stack: PieChartStack[] = [
-       {
-           value: metaUnspecifiedFactor * 100,
-           name: "Unspecified",
-           color: statusColors.get(StatusColor.Unspecified)!,
-       },
-       {
-           value: metaFailedFactor * 100,
-           name: "Failure",
-           color: statusColors.get(StatusColor.Failure)!,
-       },
-       {
-           value: (1 - (metaFailedFactor + metaUnspecifiedFactor)) * 100,
-           name: "Passed",
-           color: statusColors.get(StatusColor.Success)!
-       }
+        {
+            value: metaUnspecifiedFactor * 100,
+            name: "Unspecified",
+            color: statusColors.get(StatusColor.Unspecified)!,
+        },
+        {
+            value: metaFailedFactor * 100,
+            name: "Failure",
+            color: statusColors.get(StatusColor.Failure)!,
+        },
+        {
+            value: (1 - (metaFailedFactor + metaUnspecifiedFactor)) * 100,
+            name: "Passed",
+            color: statusColors.get(StatusColor.Success)!
+        }
     ];
 
     return PieChart(container, stack, radius, innerRadius, padAngle);
-  }
+}
