Horde: Fix issue with default incremental pcb artifact
#wubmittool bypass
#skipci

diff --git a/Engine/Source/Programs/Horde/HordeServer/Defaults/ue5-release-5.6.stream.json b/Engine/Source/Programs/Horde/HordeServer/Defaults/ue5-release-5.6.stream.json
--- a/Engine/Source/Programs/Horde/HordeServer/Defaults/ue5-release-5.6.stream.json
+++ b/Engine/Source/Programs/Horde/HordeServer/Defaults/ue5-release-5.6.stream.json
@@ -254,7 +254,7 @@
 			"arguments":
 			[
 				"-Script=Engine/Build/Graph/Examples/BuildEditorAndTools.xml",
-				"-Target=Stage for UGS",
+				"-Target=Submit PCBs",
 				"-set:UProjectPath=$(ProjectPath)",
 				"-set:EditorTarget=$(EditorTarget)",
 				"-set:PreferredAgent=IncrementalWin64;CompileWin64;Win64"
