Horde: Make it so non-admin users can see, though not edit notices

diff --git a/Engine/Source/Programs/Horde/HordeDashboard/src/components/NoticeView.tsx b/Engine/Source/Programs/Horde/HordeDashboard/src/components/NoticeView.tsx
--- a/Engine/Source/Programs/Horde/HordeDashboard/src/components/NoticeView.tsx
+++ b/Engine/Source/Programs/Horde/HordeDashboard/src/components/NoticeView.tsx
@@ -37,7 +37,7 @@ const NoticePanel: React.FC = observer(() => {
 
       if (column?.name === "Message") {
 
-         if (item.createdByUser) {
+         if (item.createdByUser && !!dashboard.user?.dashboardFeatures?.showNoticeEditor) {
             return <div onClick={() => { setState({ showEditor: true, editNotice: item }) }} style={{ "cursor": "pointer" }}>
                <Stack horizontal tokens={{ childrenGap: 24 }}>
                   <Icon style={{ paddingTop: 2 }} iconName="Edit" />
@@ -100,7 +100,7 @@ const NoticePanel: React.FC = observer(() => {
       <Stack style={{ paddingBottom: 12 }}>
          <Stack verticalAlign="center">
             <Stack horizontalAlign="end">
-               <PrimaryButton styles={{ root: { fontFamily: "Horde Open Sans SemiBold" } }} onClick={() => setState({ showEditor: true, editNotice: userNotice })}>Set Notice</PrimaryButton>
+               { !!dashboard.user?.dashboardFeatures?.showNoticeEditor &&<PrimaryButton styles={{ root: { fontFamily: "Horde Open Sans SemiBold" } }} onClick={() => setState({ showEditor: true, editNotice: userNotice })}>Set Notice</PrimaryButton> }
             </Stack>
             {!allNotices.length && <Stack horizontalAlign="center">
                <Text variant="mediumPlus">No Service Notices Found</Text>
diff --git a/Engine/Source/Programs/Horde/HordeDashboard/src/components/TopNav.tsx b/Engine/Source/Programs/Horde/HordeDashboard/src/components/TopNav.tsx
--- a/Engine/Source/Programs/Horde/HordeDashboard/src/components/TopNav.tsx
+++ b/Engine/Source/Programs/Horde/HordeDashboard/src/components/TopNav.tsx
@@ -604,14 +604,11 @@ export const TopNav: React.FC<{ suppressServer?: boolean }> = observer(({ suppre
          });
       }
 
-
-      if (features?.showNoticeEditor !== false) {
-         monitoringItems.push({
-            key: "admin_notices",
-            text: "Notices",
-            link: `/notices`
-         });
-      }
+      monitoringItems.push({
+         key: "admin_notices",
+         text: "Notices",
+         link: `/notices`
+      });
 
       if (features?.showPerforceServers !== false) {
          monitoringItems.push({
