Horde Storage: Downgrade a file missing exception to a warning as will retry and should not fail step

diff --git a/Engine/Source/Programs/Horde/Drivers/JobDriver/Utility/TempStorage.cs b/Engine/Source/Programs/Horde/Drivers/JobDriver/Utility/TempStorage.cs
--- a/Engine/Source/Programs/Horde/Drivers/JobDriver/Utility/TempStorage.cs
+++ b/Engine/Source/Programs/Horde/Drivers/JobDriver/Utility/TempStorage.cs
@@ -815,7 +815,7 @@ namespace JobDriver.Utility
 					}
 					catch (FileNotFoundException ex) when (attempt == 0)
 					{
-						logger.LogError(ex, "File not found: {Message}", ex.Message);
+						logger.LogWarning(ex, "File not found: {Message}", ex.Message);
 						logger.LogInformation("Running extract again with verbose logging.");
 					}
 				}
