Horde : Add missing properties to agent

Add RAM size for all platforms
Add core count for ARM based Macs

diff --git a/Engine/Source/Programs/Horde/HordeAgent/Services/CapabilitiesService.cs b/Engine/Source/Programs/Horde/HordeAgent/Services/CapabilitiesService.cs
--- a/Engine/Source/Programs/Horde/HordeAgent/Services/CapabilitiesService.cs
+++ b/Engine/Source/Programs/Horde/HordeAgent/Services/CapabilitiesService.cs
@@ -204,7 +204,9 @@ namespace HordeAgent.Services
 					if (match.Success)
 					{
 						long totalCapacity = Int64.Parse(match.Groups[1].Value) * 1024;
-						capabilities.Resources.Add(KnownResourceNames.Ram, (int)(totalCapacity / (1024 * 1024 * 1024)));
+						int ramGb = (int)(totalCapacity / (1024 * 1024 * 1024));
+						capabilities.Resources.Add(KnownResourceNames.Ram, ramGb);
+						capabilities.Properties.Add($"{KnownResourceNames.Ram}={ramGb}");
 					}
 				}
 
@@ -266,10 +268,29 @@ namespace HordeAgent.Services
 					numLogicalCores = (int)(numLogicalCores * _settings.CurrentValue.CpuMultiplier);
 					capabilities.Resources.Add(KnownResourceNames.LogicalCores, numLogicalCores);
 
-					XmlNode? cpuCountNode = hardwareNode.SelectSingleNode("key[. = 'number_processors']/following-sibling::integer");
-					if (cpuCountNode != null)
+					if (RuntimeInformation.OSArchitecture == Architecture.Arm64)
 					{
-						capabilities.Properties.Add($"PhysicalCores={cpuCountNode.InnerText}");
+						XmlNode? cpuCountNode = hardwareNode.SelectSingleNode("key[. = 'number_processors']/following-sibling::string");
+						if (cpuCountNode != null)
+						{
+							string[] parts = cpuCountNode.InnerText.Split([':'], StringSplitOptions.RemoveEmptyEntries);
+							if (parts.Length == 3)
+							{
+								if (Int32.TryParse(parts[1], out int performanceCores) && Int32.TryParse(parts[2], out int efficiencyCores))
+								{
+									capabilities.Properties.Add($"PhysicalCores={performanceCores}");
+									capabilities.Properties.Add($"EfficiencyCores={efficiencyCores}");
+								}
+							}
+						}
+					}
+					else
+					{
+						XmlNode? cpuCountNode = hardwareNode.SelectSingleNode("key[. = 'number_processors']/following-sibling::integer");
+						if (cpuCountNode != null)
+						{
+							capabilities.Properties.Add($"PhysicalCores={cpuCountNode.InnerText}");
+						}
 					}
 
 					XmlNode? memoryNode = hardwareNode.SelectSingleNode("key[. = 'physical_memory']/following-sibling::string");
@@ -279,6 +300,7 @@ namespace HordeAgent.Services
 						if (parts.Length == 2 && parts[1] == "GB" && Int32.TryParse(parts[0], out int ramGb))
 						{
 							capabilities.Resources.Add(KnownResourceNames.Ram, ramGb);
+							capabilities.Properties.Add($"{KnownResourceNames.Ram}={ramGb}");
 						}
 					}
 				}
diff --git a/Engine/Source/Programs/Horde/HordeAgent/Services/WindowsCapabilities.cs b/Engine/Source/Programs/Horde/HordeAgent/Services/WindowsCapabilities.cs
--- a/Engine/Source/Programs/Horde/HordeAgent/Services/WindowsCapabilities.cs
+++ b/Engine/Source/Programs/Horde/HordeAgent/Services/WindowsCapabilities.cs
@@ -49,7 +49,9 @@ public class WindowsCapabilities(IOptionsMonitor<AgentSettings> settings, ILogge
 		MEMORYSTATUSEX memStatus = new();
 		if (GlobalMemoryStatusEx(memStatus))
 		{
-			SetResource(caps, KnownResourceNames.Ram, (int)(memStatus.ullTotalPhys / (1024 * 1024 * 1024)));
+			int ramGb = (int)(memStatus.ullTotalPhys / (1024 * 1024 * 1024));
+			SetResource(caps, KnownResourceNames.Ram, ramGb);
+			SetProp(caps, KnownResourceNames.Ram, ramGb);
 		}
 		else
 		{
