Change type of test phase session and test session duration from TimeSpan to float to avoid duration to be serialize into string

#rnx

diff --git a/Engine/Source/Programs/Horde/Plugins/Build/HordeServer.Build/Jobs/TestData/TestDataController.cs b/Engine/Source/Programs/Horde/Plugins/Build/HordeServer.Build/Jobs/TestData/TestDataController.cs
--- a/Engine/Source/Programs/Horde/Plugins/Build/HordeServer.Build/Jobs/TestData/TestDataController.cs
+++ b/Engine/Source/Programs/Horde/Plugins/Build/HordeServer.Build/Jobs/TestData/TestDataController.cs
@@ -745,7 +745,7 @@ namespace HordeServer.Jobs.TestData
 					StreamId = testData.StreamId.ToString(),
 					JobId = testData.JobId.ToString(),
 					StepId = testData.StepId.ToString(),
-					Duration = testData.Duration,
+					Duration = testData.Duration.TotalSeconds,
 					StartDateTime = testData.StartDateTime,
 					CommitId = testData.BuildCommitId.ToString(),
 					CommitOrder = testData.BuildCommitId.Order,
@@ -813,7 +813,7 @@ namespace HordeServer.Jobs.TestData
 					StreamId = testData.StreamId.ToString(),
 					JobId = testData.JobId.ToString(),
 					StepId = testData.StepId.ToString(),
-					Duration = testData.Duration,
+					Duration = testData.Duration.TotalSeconds,
 					StartDateTime = testData.StartDateTime,
 					CommitId = testData.BuildCommitId.ToString(),
 					CommitOrder = testData.BuildCommitId.Order,
diff --git a/Engine/Source/Programs/Shared/EpicGames.Horde/Jobs/TestData/TestDataMessages.cs b/Engine/Source/Programs/Shared/EpicGames.Horde/Jobs/TestData/TestDataMessages.cs
--- a/Engine/Source/Programs/Shared/EpicGames.Horde/Jobs/TestData/TestDataMessages.cs
+++ b/Engine/Source/Programs/Shared/EpicGames.Horde/Jobs/TestData/TestDataMessages.cs
@@ -721,7 +721,7 @@ namespace EpicGames.Horde.Jobs.TestData
 		/// <summary>
 		/// The duration of the test session
 		/// </summary>
-		public TimeSpan Duration { get; set; }
+		public double Duration { get; set; }
 
 		/// <summary>
 		/// The date time at which the test session started
@@ -807,7 +807,7 @@ namespace EpicGames.Horde.Jobs.TestData
 		/// <summary>
 		/// The duration of the phase session
 		/// </summary>
-		public TimeSpan Duration { get; set; }
+		public double Duration { get; set; }
 
 		/// <summary>
 		/// The date time at which the phase session started
