Horde: Fix publishing step for Horde installer considerations
#skipci

diff --git a/Engine/Source/Programs/Horde/HordeServer/Defaults/ue5-release-5.6.stream.json b/Engine/Source/Programs/Horde/HordeServer/Defaults/ue5-release-5.6.stream.json
--- a/Engine/Source/Programs/Horde/HordeServer/Defaults/ue5-release-5.6.stream.json
+++ b/Engine/Source/Programs/Horde/HordeServer/Defaults/ue5-release-5.6.stream.json
@@ -272,7 +272,9 @@
 				"-set:ProjectFile=$(ProjectPath)",
 				"-set:TargetName=$(GameTarget)",
 				"-set:EditorTargetName=$(EditorTarget)",
-				"-set:WithDeviceReservation=false"
+				"-set:WithDeviceReservation=false",
+				"-set:PublishBuilds=false",
+				"-set:PublishHordeArtifacts=true"				
 			],
 			"parameters":
 			[
