Updated ini example in the Horde Analytics tutorial documentation.

#rnx

diff --git a/Engine/Source/Programs/Horde/Docs/Tutorials/Analytics.md b/Engine/Source/Programs/Horde/Docs/Tutorials/Analytics.md
--- a/Engine/Source/Programs/Horde/Docs/Tutorials/Analytics.md
+++ b/Engine/Source/Programs/Horde/Docs/Tutorials/Analytics.md
@@ -25,7 +25,7 @@ make sure it is enabled. It should be enabled by default.
 the following lines:
 
     ```ini
-    [StudioTelemetry.Provider.HordeAnalytics]
+    [StudioTelemetry.Provider.Horde]
     Name=HordeAnalytics
     ProviderModule=AnalyticsET
     UsageType=EditorAndClient
