Horde Dashboard: Fix autosubmit link when swarm url isn't defined
#skipci

diff --git a/Engine/Source/Programs/Horde/HordeDashboard/src/components/AutoSubmit.tsx b/Engine/Source/Programs/Horde/HordeDashboard/src/components/AutoSubmit.tsx
--- a/Engine/Source/Programs/Horde/HordeDashboard/src/components/AutoSubmit.tsx
+++ b/Engine/Source/Programs/Horde/HordeDashboard/src/components/AutoSubmit.tsx
@@ -24,12 +24,16 @@ export const AutosubmitInfo: React.FC<{ jobDetails: JobDetailsV2 }> = observer((
    }
 
    if (job.autoSubmitChange) {
-      const url = `${dashboard.swarmUrl}/change/${job.autoSubmitChange}`;
-      return <Stack style={{paddingTop: 24}}><Text>Automatically submitted in <a href={url} target="_blank" rel="noreferrer" onClick={ev => ev?.stopPropagation()}>{`CL ${job.autoSubmitChange}`}</a></Text></Stack>;
+      if (dashboard.swarmUrl) {
+         const url = `${dashboard.swarmUrl}/change/${job.autoSubmitChange}`;
+         return <Stack style={{ paddingTop: 24 }}><Text>Automatically submitted in <a href={url} target="_blank" rel="noreferrer" onClick={ev => ev?.stopPropagation()}>{`CL ${job.autoSubmitChange}`}</a></Text></Stack>
+      } else {
+         return <Stack style={{ paddingTop: 24 }}><Text>Automatically submitted in {`CL ${job.autoSubmitChange}`}</Text></Stack>
+      }
    }
 
    if (job.autoSubmitMessage) {
-      return <Stack style={{ paddingTop: 24, whiteSpace:"pre" }}><Text>{`Unable to submit change: ${job.autoSubmitMessage}`}</Text></Stack>;
+      return <Stack style={{ paddingTop: 24, whiteSpace: "pre" }}><Text>{`Unable to submit change: ${job.autoSubmitMessage}`}</Text></Stack>
    }
 
    if (job.state !== JobState.Complete) {
@@ -37,7 +41,7 @@ export const AutosubmitInfo: React.FC<{ jobDetails: JobDetailsV2 }> = observer((
       if (job.preflightChange) {
          return <Stack style={{ paddingTop: 24 }}><Checkbox label={`Automatically submit changelist ${job.preflightChange} upon preflight success`}
             checked={state.autoSubmit}
-            disabled={state.inflight || (job.startedByUserInfo?.id !== dashboard.userId) }
+            disabled={state.inflight || (job.startedByUserInfo?.id !== dashboard.userId)}
             onChange={(ev, checked) => {
 
                const value = !!checked;
@@ -45,7 +49,7 @@ export const AutosubmitInfo: React.FC<{ jobDetails: JobDetailsV2 }> = observer((
                backend.updateJob(job.id, { autoSubmit: !!checked }).then(result => {
                   // messing with job object here, not great
                   job.autoSubmit = value;
-               setState({ autoSubmit: value, inflight: false })
+                  setState({ autoSubmit: value, inflight: false })
                }).catch(reason => {
 
                   ErrorHandler.set({
diff --git a/Engine/Source/Programs/Horde/HordeDashboard/src/components/jobDetailsV2/JobDetailPreflight.tsx b/Engine/Source/Programs/Horde/HordeDashboard/src/components/jobDetailsV2/JobDetailPreflight.tsx
--- a/Engine/Source/Programs/Horde/HordeDashboard/src/components/jobDetailsV2/JobDetailPreflight.tsx
+++ b/Engine/Source/Programs/Horde/HordeDashboard/src/components/jobDetailsV2/JobDetailPreflight.tsx
@@ -28,7 +28,7 @@ class PreflightDataView extends JobDataView {
 
    }
 
-   order = 0;   
+   order = 0;
 
 }
 
@@ -45,21 +45,25 @@ const AutosubmitInfo: React.FC<{ jobDetails: JobDetailsV2 }> = observer(({ jobDe
 
    if (jobData.abortedByUserInfo) {
       return null;
-   }
+   }   
 
    if (jobData.autoSubmitChange) {
-      const url = `${dashboard.swarmUrl}/change/${jobData.autoSubmitChange}`;
-      return <Stack tokens={{ childrenGap: 8 }} style={{ paddingBottom: 12, paddingTop: 18 }}>
-         <Stack style={{  }}>
-            <Text>Automatically submitted in <a href={url} target="_blank" rel="noreferrer" onClick={ev => ev?.stopPropagation()}>{`CL ${jobData.autoSubmitChange}`}</a></Text>
+      if (dashboard.swarmUrl) {
+         const url = `${dashboard.swarmUrl}/change/${jobData.autoSubmitChange}`;
+         return <Stack tokens={{ childrenGap: 8 }} style={{ paddingBottom: 12, paddingTop: 18 }}>
+            <Stack style={{}}>
+               <Text>Automatically submitted in <a href={url} target="_blank" rel="noreferrer" onClick={ev => ev?.stopPropagation()}>{`CL ${jobData.autoSubmitChange}`}</a></Text>
+            </Stack>
          </Stack>
-      </Stack>
+      } else {
+         return <Stack style={{ paddingTop: 24 }}><Text>Automatically submitted in {`CL ${jobData.autoSubmitChange}`}</Text></Stack>
+      }
    }
 
    if (jobData.autoSubmitMessage) {
       return <Stack tokens={{ childrenGap: 8 }} style={{ paddingBottom: 12, paddingTop: 18 }}>
-         <Stack style={{  }}>
-            <Stack style={{ paddingTop: 24, whiteSpace: "pre" }}><Text>{`Unable to submit change: ${jobData.autoSubmitMessage}`}</Text></Stack>;
+         <Stack style={{}}>
+            <Stack style={{ paddingTop: 24, whiteSpace: "pre" }}><Text>{`Unable to submit change: ${jobData.autoSubmitMessage}`}</Text></Stack>
          </Stack>
       </Stack>
    }
