Updating the UBA documentation.

diff --git a/Engine/Source/Programs/Horde/Docs/Tutorials/RemoteCompilation.md b/Engine/Source/Programs/Horde/Docs/Tutorials/RemoteCompilation.md
--- a/Engine/Source/Programs/Horde/Docs/Tutorials/RemoteCompilation.md
+++ b/Engine/Source/Programs/Horde/Docs/Tutorials/RemoteCompilation.md
@@ -28,33 +28,57 @@ spreading the workload over multiple machines.
 ## Steps
 
 1. On the machine initiating the build, ensure your UE project is synced and builds locally.
-2. Configure UnrealBuildTool to use your Horde Server by updating
-   `Engine/Saved/UnrealBuildTool/BuildConfiguration.xml` with the following:
 
-   ```xml
-   <?xml version="1.0" encoding="utf-8" ?>
-   <Configuration xmlns="https://www.unrealengine.com/BuildConfiguration">
+2. Configure UnrealBuildTool
 
-       <BuildConfiguration>
-           <!-- Enable support for UnrealBuildAccelerator -->
-           <bAllowUBAExecutor>true</bAllowUBAExecutor>
-       </BuildConfiguration>
+  First setup your Uba.Provider settings, we recommend you do this in a file within perforce (so they can be shared for all users):
+    - `Engine/Config/BaseEngine.ini` for settings to be used globally
+    - `Config/DefaultEngine.ini` (within a project directory) for settings to be used for a specific project
 
-       <Horde>
-           <!-- Address of the Horde server -->
-           <Server>http://{{ SERVER_HOST_NAME }}:13340</Server>
+  An example of how to configure the Uba.Provider settings can be found here:
+    ```[Uba.Provider.Horde.Example]
+	ServerUrl="[YOUR HORDE URL]"
+	Pool=[YOUR HORDE POOL, COULD BE UNSET]
+	Enabled=True
+    ```
 
-           <!-- Pool of machines to offload work to. Horde configures Win-UE5 by default. -->
-           <WindowsPool>Win-UE5</WindowsPool>
-       </Horde>
+  Secondly configure the UbaControllers, this is done in a local file you need to manually create:
+    - For Windows => `%LocalAppData%/Unreal Engine/Engine/Config/UserEngine.ini`
+    - For Linux => `$HOME/.config/Epic/Unreal Engine/Engine/Config/UserEngine.ini`
+    - For Mac => `$HOME/Library/Application Support/Epic/Unreal Engine/Engine/Config/UserEngine.ini`
 
-       <UnrealBuildAccelerator>
-           <!-- Enable for visualizing UBA's progress (optional) -->
-           <bLaunchVisualizer>true</bLaunchVisualizer>
-       </UnrealBuildAccelerator>
+  An example of how to configure this file can be found below:
+	```[UbaController]
+	+Providers=Uba.Provider.Horde.Example
+	```
 
-   </Configuration>
-   ```
+  Note, you can alternatively configure UnrealBuildTool within `Engine/Saved/UnrealBuildTool/BuildConfiguration.xml` but this is not recommended and is a legacy approach
+
+  An example of how to configure this file:
+     ```xml
+     <?xml version="1.0" encoding="utf-8" ?>
+     <Configuration xmlns="https://www.unrealengine.com/BuildConfiguration">
+
+         <BuildConfiguration>
+             <!-- Enable support for UnrealBuildAccelerator -->
+             <bAllowUBAExecutor>true</bAllowUBAExecutor>
+         </BuildConfiguration>
+
+         <Horde>
+             <!-- Address of the Horde server -->
+             <Server>http://{{ SERVER_HOST_NAME }}:13340</Server>
+
+             <!-- Pool of machines to offload work to. Horde configures Win-UE5 by default. -->
+             <WindowsPool>Win-UE5</WindowsPool>
+         </Horde>
+
+         <UnrealBuildAccelerator>
+             <!-- Enable for visualizing UBA's progress (optional) -->
+             <bLaunchVisualizer>true</bLaunchVisualizer>
+         </UnrealBuildAccelerator>
+
+     </Configuration>
+     ```
 
    Replace `SERVER_HOST_NAME` with the address associated with your Horde server installation.
 
@@ -127,4 +151,4 @@ To see what claims are available to your users, open `/api/v1/user/claims` as a
         }
       }
     }
-    ```
+    ```
\ No newline at end of file
