Horde Dashboard: allow option for UserSelect to autofocus

diff --git a/Engine/Source/Programs/Horde/HordeDashboard/src/components/UserSelect.tsx b/Engine/Source/Programs/Horde/HordeDashboard/src/components/UserSelect.tsx
--- a/Engine/Source/Programs/Horde/HordeDashboard/src/components/UserSelect.tsx
+++ b/Engine/Source/Programs/Horde/HordeDashboard/src/components/UserSelect.tsx
@@ -1,4 +1,4 @@
-import { CompactPeoplePicker, IBasePickerSuggestionsProps, IPersonaProps, Label, Stack } from "@fluentui/react";
+import { CompactPeoplePicker, IBasePickerSuggestionsProps, IInputProps, IPersonaProps, Label, Stack } from "@fluentui/react";
 import { observer } from "mobx-react-lite";
 import React, { useState } from "react";
 import { useWindowSize } from "../base/utilities/hooks";
@@ -11,7 +11,7 @@ import { getHordeStyling } from "../styles/Styles";
 
 const loadedUsers: Map<string, GetUserResponse> = new Map();
 
-export const UserSelect: React.FC<{ handleSelection: (userId: string | undefined) => void, userIdHints?: string[], noResultsFoundText?: string, defaultUser?: GetUserResponse }> = observer(({ handleSelection, userIdHints, noResultsFoundText, defaultUser }) => {
+export const UserSelect: React.FC<{ handleSelection: (userId: string | undefined) => void, userIdHints?: string[], noResultsFoundText?: string, defaultUser?: GetUserResponse, autoFocus?: boolean }> = observer(({ handleSelection, userIdHints, noResultsFoundText, defaultUser, autoFocus }) => {
 
    const [state, setState] = useState<{ initial?: boolean }>({ initial: true });
 
@@ -30,6 +30,10 @@ export const UserSelect: React.FC<{ handleSelection: (userId: string | undefined
       showRemoveButtons: false
    };
 
+   const inputProps: IInputProps = { 
+      autoFocus: autoFocus
+   };
+
    // subscribe
    if (userCache.updated) { }
 
@@ -106,7 +110,7 @@ export const UserSelect: React.FC<{ handleSelection: (userId: string | undefined
          pickerSuggestionsProps={suggestionProps}
          className={'ms-PeoplePicker'}
          defaultSelectedItems={defaultUser ? [{ text: defaultUser.name, imageUrl: defaultUser.image24, optionalText: defaultUser.id }] : []}
-
+         inputProps={inputProps}
          onChange={(items) => {
 
             if (items && items?.length) {
