Horde:  Do not include steps with warnings in failed steps for job retries
#skipci

diff --git a/Engine/Source/Programs/Horde/HordeDashboard/src/components/jobDetailsV2/JobOperationsBar.tsx b/Engine/Source/Programs/Horde/HordeDashboard/src/components/jobDetailsV2/JobOperationsBar.tsx
--- a/Engine/Source/Programs/Horde/HordeDashboard/src/components/jobDetailsV2/JobOperationsBar.tsx
+++ b/Engine/Source/Programs/Horde/HordeDashboard/src/components/jobDetailsV2/JobOperationsBar.tsx
@@ -78,8 +78,8 @@ export const JobOperations: React.FC<{ jobDetails: JobDetailsV2 }> = observer(({
          runningSteps.add(s.name);
          return false;
       }
-
-      if (s.state === JobStepState.Completed && s.outcome == JobStepOutcome.Success) 
+      
+      if (s.state === JobStepState.Completed && (s.outcome === JobStepOutcome.Success || s.outcome === JobStepOutcome.Warnings))
       {
          return false;
       }
