Update Remote Shader Compilation Horde docs with the new ini configuration.

diff --git a/Engine/Source/Programs/Horde/Docs/Tutorials/RemoteShaderCompilation.md b/Engine/Source/Programs/Horde/Docs/Tutorials/RemoteShaderCompilation.md
--- a/Engine/Source/Programs/Horde/Docs/Tutorials/RemoteShaderCompilation.md
+++ b/Engine/Source/Programs/Horde/Docs/Tutorials/RemoteShaderCompilation.md
@@ -25,20 +25,37 @@ Once you've verified that it's working, this same functionality will be availabl
    Engine\Build\BatchFiles\RunUBT.bat LyraEditor ShaderCompileWorker Win64 Development -Project=..\..\Samples\Games\Lyra\Lyra.uproject
    ```
 
-2. Configure `Samples\Games\Lyra\Config\DefaultEngine.ini` with Horde and UBA settings:
+2. Configure the Uba settings
+
+   First setup your Uba.Provider settings, we recommend you do this in a file within perforce (so they can be shared for all users):
+     * `Engine\Config\BaseEngine.ini` for settings to be used globally
+     * `Samples\Games\Lyra\Config\DefaultEngine.ini` for settings to be used for a specific project (in this case Lyra)
+  
+   An example of how to configure the Uba.Provider settings can be found here:
 
    ```ini
-   [UbaController]
-   Enabled=true
-   Horde=(Pool=Win-UE5,MaxCores=500)
-   
-   [Horde]
+   [Uba.Provider.Horde.Example]
    ServerUrl="http://{{ SERVER_HOST_NAME }}:13340/"
+   Enabled=true
+   Pool=[YOUR HORDE POOL, COULD BE UNSET]
+   MaxCores=500
    ```
 
    `MaxCores=500` describes how many CPU cores this UBA session should try to allocate for work.
    Adjust as needed for fair use and size of your Horde setup.
 
+   Secondly configure the UbaControllers, this is done in a local file you need to manually create:
+     * For Windows => `%LocalAppData%/Unreal Engine/Engine/Config/UserEngine.ini`
+     * For Linux => `$HOME/.config/Epic/Unreal Engine/Engine/Config/UserEngine.ini`
+     * For Mac => `$HOME/Library/Application Support/Epic/Unreal Engine/Engine/Config/UserEngine.ini`
+
+   An example of how to configure this file can be found below:
+
+   ```ini
+   [UbaController]
+   +Providers=Uba.Provider.Horde.Example
+   ```
+
 3. Cook the data, which will compile shaders among many other asset transformations
 
    ```
