Horde: Prevent per-cluster Perforce errors from aborting changelist replication loop

With an invalid Perforce cluster or no available servers for that cluster, such an exception would bubble all the way up to the tick method. This lead to the replication loop restarting.

diff --git a/Engine/Source/Programs/Horde/Plugins/Build/HordeServer.Build/Perforce/PerforceServiceCache.cs b/Engine/Source/Programs/Horde/Plugins/Build/HordeServer.Build/Perforce/PerforceServiceCache.cs
--- a/Engine/Source/Programs/Horde/Plugins/Build/HordeServer.Build/Perforce/PerforceServiceCache.cs
+++ b/Engine/Source/Programs/Horde/Plugins/Build/HordeServer.Build/Perforce/PerforceServiceCache.cs
@@ -330,11 +330,18 @@ namespace HordeServer.Perforce
 		{
 			IReadOnlyList<StreamConfig> streams = _buildConfig.CurrentValue.Streams;
 
-			Dictionary<string, List<StreamInfo>> clusters = new Dictionary<string, List<StreamInfo>>(StringComparer.OrdinalIgnoreCase);
+			Dictionary<string, List<StreamInfo>> clusters = new (StringComparer.OrdinalIgnoreCase);
 			foreach (IGrouping<string, StreamConfig> group in streams.GroupBy(x => x.ClusterName, StringComparer.OrdinalIgnoreCase))
 			{
-				List<StreamInfo> streamInfoList = await CreateStreamInfoForClusterAsync(group.Key, group, cancellationToken);
-				clusters[group.Key] = streamInfoList;
+				try
+				{
+					List<StreamInfo> streamInfoList = await CreateStreamInfoForClusterAsync(group.Key, group, cancellationToken);
+					clusters[group.Key] = streamInfoList;
+				}
+				catch (Exception e)
+				{
+					_logger.LogError(e, "Failed resolving stream infos for Perforce cluster {ClusterId}", group.Key);
+				}
 			}
 
 			return clusters;
@@ -342,20 +349,18 @@ namespace HordeServer.Perforce
 
 		async Task<List<StreamInfo>> CreateStreamInfoForClusterAsync(string clusterName, IEnumerable<StreamConfig> streams, CancellationToken cancellationToken)
 		{
-			using (IPooledPerforceConnection perforce = await ConnectAsync(clusterName, null, cancellationToken))
-			{
-				InfoRecord serverInfo = await perforce.GetInfoAsync(cancellationToken);
+			using IPooledPerforceConnection perforce = await ConnectAsync(clusterName, null, cancellationToken);
+			InfoRecord serverInfo = await perforce.GetInfoAsync(cancellationToken);
 
-				List<StreamInfo> streamInfoList = new List<StreamInfo>();
-				foreach (StreamConfig stream in streams)
-				{
-					StreamRecord record = await perforce.GetStreamAsync(stream.Name, true, cancellationToken);
-					PerforceViewMap view = PerforceViewMap.Parse(record.View);
-					PerforceChangeView changeView = PerforceChangeView.Parse(record.ChangeView, !serverInfo.IsCaseSensitive);
-					streamInfoList.Add(new StreamInfo(stream, view, changeView));
-				}
-				return streamInfoList;
+			List<StreamInfo> streamInfoList = [];
+			foreach (StreamConfig stream in streams)
+			{
+				StreamRecord record = await perforce.GetStreamAsync(stream.Name, true, cancellationToken);
+				PerforceViewMap view = PerforceViewMap.Parse(record.View);
+				PerforceChangeView changeView = PerforceChangeView.Parse(record.ChangeView, !serverInfo.IsCaseSensitive);
+				streamInfoList.Add(new StreamInfo(stream, view, changeView));
 			}
+			return streamInfoList;
 		}
 
 		ClusterTicker StartTicker(string clusterName, List<StreamInfo> streamInfos)
