Horde: Display warning to surface rendering issues on large log files
Fixed heightAdjust values that were off.

diff --git a/Engine/Source/Programs/Horde/HordeDashboard/src/components/LogView.tsx b/Engine/Source/Programs/Horde/HordeDashboard/src/components/LogView.tsx
--- a/Engine/Source/Programs/Horde/HordeDashboard/src/components/LogView.tsx
+++ b/Engine/Source/Programs/Horde/HordeDashboard/src/components/LogView.tsx
@@ -1,6 +1,6 @@
 // Copyright Epic Games, Inc. All Rights Reserved.
 
-import { Callout, DefaultButton, DetailsList, DetailsListLayoutMode, DetailsRow, DirectionalHint, Dropdown, FocusZone, FocusZoneDirection, FontIcon, IColumn, Icon, IconButton, IContextualMenuItem, IContextualMenuProps, IDetailsListProps, ITextField, List, mergeStyles, Modal, PrimaryButton, ProgressIndicator, ScrollToMode, Selection, SelectionMode, SelectionZone, Separator, Spinner, SpinnerSize, Stack, Text, TextField, TooltipHost } from '@fluentui/react';
+import { Callout, DefaultButton, DetailsList, DetailsListLayoutMode, DetailsRow, DirectionalHint, Dropdown, FocusZone, FocusZoneDirection, FontIcon, IColumn, Icon, IconButton, IContextualMenuItem, IContextualMenuProps, IDetailsListProps, ITextField, List, mergeStyles, MessageBar, MessageBarType, Modal, PrimaryButton, ProgressIndicator, ScrollToMode, Selection, SelectionMode, SelectionZone, Separator, Spinner, SpinnerSize, Stack, Text, TextField, TooltipHost } from '@fluentui/react';
 import { action, makeObservable, observable } from 'mobx';
 import { observer } from 'mobx-react-lite';
 import moment from 'moment-timezone';
@@ -1000,6 +1000,9 @@ export const LogList: React.FC<{ logId: string }> = observer(({ logId }) => {
 
    let summaryText = logSource.summary;
 
+   let largeLogWarning = logSource.logItems?.length > 1500000;
+   let largeLogWarningText= `Warning: large log file (${logSource.logItems?.length} lines) may not load entirely - Download to view full log and consider reducing spam`;
+
    let baseUrl = location.pathname;
 
    /*
@@ -1391,10 +1394,13 @@ export const LogList: React.FC<{ logId: string }> = observer(({ logId }) => {
    }
 
 
-   let heightAdjust = 292;
+   let heightAdjust = 285;
    if (logTelemetry) {
       // sparkline height
-      heightAdjust += 120;
+      heightAdjust += 238;
+   }
+   if (largeLogWarning) {
+      heightAdjust += 55
    }
 
    const nextErrorEnabled = handler?.getNextLogEvent() ? true : false;
@@ -1630,6 +1636,12 @@ export const LogList: React.FC<{ logId: string }> = observer(({ logId }) => {
                      </Stack>
                   </Stack>
                   {logTelemetry && !!logSource.agentTelemetry && <AgentTelemetrySparkline handler={logSource.agentTelemetry} />}
+                  {largeLogWarning && <Stack style={{ paddingBottom: 12, paddingTop: 12, margin: 'auto'}}>
+                     <MessageBar
+                        messageBarType={MessageBarType.warning} isMultiline={false} >
+                        <Markdown styles={{ root: { fontSize: 12, fontFamily: "Horde Open Sans Bold" } }}>{largeLogWarningText}</Markdown>
+                     </MessageBar>
+                  </Stack>}
                   <Stack horizontalAlign="center" style={{ paddingBottom: 12 }}>
                      <Separator styles={{ root: { fontSize: 0, width: "100%", padding: 0, selectors: { '::before': { background: dashboard.darktheme ? '#313638' : '#D3D2D1' } } } }} />
                   </Stack>
