Horde: Tweak log warning verbage.

diff --git a/Engine/Source/Programs/Horde/HordeDashboard/src/components/LogView.tsx b/Engine/Source/Programs/Horde/HordeDashboard/src/components/LogView.tsx
--- a/Engine/Source/Programs/Horde/HordeDashboard/src/components/LogView.tsx
+++ b/Engine/Source/Programs/Horde/HordeDashboard/src/components/LogView.tsx
@@ -1001,7 +1001,7 @@ export const LogList: React.FC<{ logId: string }> = observer(({ logId }) => {
    let summaryText = logSource.summary;
 
    let largeLogWarning = logSource.logItems?.length > 1500000;
-   let largeLogWarningText= `Warning: large log file (${logSource.logItems?.length} lines) may not load entirely - Download to view full log and consider reducing spam`;
+   let largeLogWarningText= `Warning: large log file (${logSource.logItems?.length} lines) may not display entirely.  Please download to view full log and consider reducing verbosity.`;
 
    let baseUrl = location.pathname;
 
