Horde Server: Don't require a user from user collection when creating devices, for service accounts

diff --git a/Engine/Source/Programs/Horde/Plugins/Build/HordeServer.Build/Devices/DevicesController.cs b/Engine/Source/Programs/Horde/Plugins/Build/HordeServer.Build/Devices/DevicesController.cs
--- a/Engine/Source/Programs/Horde/Plugins/Build/HordeServer.Build/Devices/DevicesController.cs
+++ b/Engine/Source/Programs/Horde/Plugins/Build/HordeServer.Build/Devices/DevicesController.cs
@@ -58,10 +58,6 @@ namespace HordeServer.Devices
 			}
 
 			IUser? internalUser = await _userCollection.GetUserAsync(User, cancellationToken);
-			if (internalUser == null)
-			{
-				return NotFound();
-			}
 
 			IDevicePlatform? platform = _deviceService.GetPlatform(new DevicePlatformId(deviceRequest.PlatformId!));
 
@@ -92,7 +88,7 @@ namespace HordeServer.Devices
 			string name = deviceRequest.Name.Trim();
 			string? address = deviceRequest.Address?.Trim();
 
-			IDevice? device = await _deviceService.TryCreateDeviceAsync(DeviceId.Sanitize(name), name, platform.Id, pool.Id, deviceRequest.Enabled, address, modelId, internalUser.Id);
+			IDevice? device = await _deviceService.TryCreateDeviceAsync(DeviceId.Sanitize(name), name, platform.Id, pool.Id, deviceRequest.Enabled, address, modelId, internalUser?.Id);
 
 			if (device == null)
 			{
