Horde Dashboard: Improve job template selection when on an exiting job

diff --git a/Engine/Source/Programs/Horde/HordeDashboard/src/components/build/NewBuildV2.tsx b/Engine/Source/Programs/Horde/HordeDashboard/src/components/build/NewBuildV2.tsx
--- a/Engine/Source/Programs/Horde/HordeDashboard/src/components/build/NewBuildV2.tsx
+++ b/Engine/Source/Programs/Horde/HordeDashboard/src/components/build/NewBuildV2.tsx
@@ -979,24 +979,40 @@ class BuildOptions {
                   return;
                }
                templates.push(t);
-            });            
+            });
          }
 
          templates = templates.filter(t => !!t.canRun);
-         
+
       } else {
          const query = new URLSearchParams(window.location.search);
          const allowtemplatechange = query.get("allowtemplatechange") === "true";
-         if (allowtemplatechange) {
-            templates = [...this.allTemplates];
-         } else {
+         if (!allowtemplatechange) {
             const template = this.allTemplates.find(t => t.name === this.jobDetails?.template?.name);
             if (template) {
                templates = [template];
             }
-            templates = [...this.allTemplates];
-            console.error(`Unable to find job template ${this.jobDetails.template?.name}`)
          }
+
+         if (!templates.length) {
+            if (!this.showAllTemplates) {
+               const tab = this.stream.tabs.filter(t => t.type === TabType.Jobs).find(t => !!(t as JobsTabData).templates?.find(template => template === this.jobDetails?.template?.id));
+
+               (tab as JobsTabData | undefined)?.templates?.forEach(tid => {
+                  const t = this.allTemplates.find(t => t.id === tid);
+                  if (!t) {
+                     console.error(`Could not find template ${tid}`);
+                     return;
+                  }
+                  templates.push(t);
+               });
+            }
+         }
+
+         if (!templates.length) {
+            templates = [...this.allTemplates].sort((a, b) => a.name.localeCompare(b.name));
+         }
+
       }
 
       this.templates = templates;
@@ -1295,23 +1311,23 @@ const LongListParameter: React.FC<{ param: ListParameterData, disabled?: boolean
 
    const key = `parameter_key_${param.label}` + `_${options.currentRenderKey}`;
 
-   const allItems: Map<ITag['key'], PickerItem> = new Map( param.items.map( item => {
+   const allItems: Map<ITag['key'], PickerItem> = new Map(param.items.map(item => {
       return [
-         item.text, 
+         item.text,
          {
             itemData: item,
             key: item.text,
             name: item.text
          }
       ]
-   }) );
+   }));
 
-   const selectFilter = allowSelectFilter ?? (( item: PickerItem ) => {
-      return !(item.name.startsWith("**") && item.name.endsWith( "**" ));
+   const selectFilter = allowSelectFilter ?? ((item: PickerItem) => {
+      return !(item.name.startsWith("**") && item.name.endsWith("**"));
    });
    const selectedItems = allItems.values().filter(item => {
-         return options.parameters[item.itemData.id]?.toLowerCase() === "true" && selectFilter( item );
-      })
+      return options.parameters[item.itemData.id]?.toLowerCase() === "true" && selectFilter(item);
+   })
       .toArray();
 
    // tag picker
@@ -1323,39 +1339,39 @@ const LongListParameter: React.FC<{ param: ListParameterData, disabled?: boolean
          disabled={options.readOnly || disabled}
          onResolveSuggestions={(filter, selected) => {
             return allItems.values().filter(i => {
-                  return !selected?.find(s => i.key === s.key) && i.name.toLowerCase().indexOf(filter.toLowerCase()) !== -1;
-               })
+               return !selected?.find(s => i.key === s.key) && i.name.toLowerCase().indexOf(filter.toLowerCase()) !== -1;
+            })
                .toArray();
          }}
 
          onEmptyResolveSuggestions={(selected) => {
             return allItems.values().filter(i => {
-                  return !selected?.find(s => i.key === s.key);
-               })
+               return !selected?.find(s => i.key === s.key);
+            })
                .toArray();
          }}
 
          selectedItems={selectedItems}
 
          onChange={(item_tags?) => {
-            
-            if ( item_tags ) {
+
+            if (item_tags) {
 
                // Get the items for all selected tags for filtering 
-               const selected_items = item_tags.map( tag => allItems.get( tag.key ) ).filter( item => item !== undefined );
+               const selected_items = item_tags.map(tag => allItems.get(tag.key)).filter(item => item !== undefined);
 
-               let filtered_items = selected_items.filter( item => selectFilter( item ) );
-               if( !allowMultiple ) {
+               let filtered_items = selected_items.filter(item => selectFilter(item));
+               if (!allowMultiple) {
                   // Keep the last item which passed the filter list as it's most likely to be the newest (need more selection metadata to perform accurately)
-                  filtered_items = [ filtered_items[ filtered_items.length - 1 ] ];
+                  filtered_items = [filtered_items[filtered_items.length - 1]];
                }
 
                // Update all tags
-               for( const [key, item] of allItems ) {
-                  const was_selected    = filtered_items.find( selected_item => key === selected_item?.key ) !== undefined;
-                  const can_be_selected = selectFilter( item );
-                  
-                  options.onListItemChanged( item.itemData.id, was_selected && can_be_selected );
+               for (const [key, item] of allItems) {
+                  const was_selected = filtered_items.find(selected_item => key === selected_item?.key) !== undefined;
+                  const can_be_selected = selectFilter(item);
+
+                  options.onListItemChanged(item.itemData.id, was_selected && can_be_selected);
                }
             }
             options.onListItemsUpdated();
@@ -1369,7 +1385,7 @@ const ListParameter: React.FC<{ param: ListParameterData, disabled?: boolean }>
    if (param.style === ListParameterStyle.TagPicker) {
       return <TagPickerParameter param={param} disabled={disabled} />
    } else if (param.items.length > LONG_LIST_PARAMETER_THRESHOLD) {
-      return <LongListParameter param={param} disabled={disabled} allowMultiple={param.style === ListParameterStyle.MultiList}/>
+      return <LongListParameter param={param} disabled={disabled} allowMultiple={param.style === ListParameterStyle.MultiList} />
    } else if (param.style === ListParameterStyle.MultiList) {
       return <MultiListParameter param={param} disabled={disabled} />
    } else if (param.style === ListParameterStyle.List) {
@@ -1570,9 +1586,9 @@ const TemplateSelector: React.FC = () => {
       templateOptions.push({ key: `show_all_templates_divider`, itemType: ContextualMenuItemType.Divider });
 
       if (!options.showAllTemplates) {
-         templateOptions.push({ key: `show_all_templates`, text: "Show All Templates", onClick: (ev) => { ev?.stopPropagation(); ev?.preventDefault(); options.setShowAllTemplates(true) } });
+         templateOptions.unshift({ key: `show_all_templates`, text: "Show All Templates", onClick: (ev) => { ev?.stopPropagation(); ev?.preventDefault(); options.setShowAllTemplates(true) } });
       } else {
-         templateOptions.push({ key: `show_all_templates`, text: "Filter Templates", onClick: (ev) => { ev?.stopPropagation(); ev?.preventDefault(); options.setShowAllTemplates(false) } });
+         templateOptions.unshift({ key: `show_all_templates`, text: "Filter Templates", onClick: (ev) => { ev?.stopPropagation(); ev?.preventDefault(); options.setShowAllTemplates(false) } });
       }
    }
 
