Horde Dashboard: Require agent disable reason and auto postpend username

diff --git a/Engine/Source/Programs/Horde/HordeDashboard/src/backend/Backend.ts b/Engine/Source/Programs/Horde/HordeDashboard/src/backend/Backend.ts
--- a/Engine/Source/Programs/Horde/HordeDashboard/src/backend/Backend.ts
+++ b/Engine/Source/Programs/Horde/HordeDashboard/src/backend/Backend.ts
@@ -660,6 +660,13 @@ export class Backend {
     }
 
     updateAgent(agentId: string, request: UpdateAgentRequest): Promise<boolean> {
+        if (request.enabled === false) {
+            let comment = request.comment ?? "";
+            const userName = dashboard.user?.name;
+            if (userName) {
+                request.comment = comment + " - "  + userName;
+            }
+        }
         return new Promise<boolean>((resolve, reject) => {
             this._fetch.put(`api/v1/agents/${agentId}`, request).then((value) => {
                 resolve(value.data as boolean);
diff --git a/Engine/Source/Programs/Horde/HordeDashboard/src/components/AgentView.tsx b/Engine/Source/Programs/Horde/HordeDashboard/src/components/AgentView.tsx
--- a/Engine/Source/Programs/Horde/HordeDashboard/src/components/AgentView.tsx
+++ b/Engine/Source/Programs/Horde/HordeDashboard/src/components/AgentView.tsx
@@ -2507,8 +2507,9 @@ export const AgentViewInner: React.FC<{ agentId?: string, poolId?: string, searc
          isOpen={localState.disableAgentDialogIsOpen}
          confirmText={"Disable"}
          cancelText={"Cancel"}
-         textBoxLabel={"Enter Disable Reason"}
+         textBoxLabel={"Enter Disable Reason - Required"}
          isTextBoxSpawned={true}
+         isTextRequired={true}
          onConfirm={(textFieldText: string) => { localState.changeBuilderEnabled("disable", textFieldText) }}
          onCancel={() => { localState.setDisableBuilderDialogOpen(false) }}
       />
diff --git a/Engine/Source/Programs/Horde/HordeDashboard/src/components/ConfirmationDialog.tsx b/Engine/Source/Programs/Horde/HordeDashboard/src/components/ConfirmationDialog.tsx
--- a/Engine/Source/Programs/Horde/HordeDashboard/src/components/ConfirmationDialog.tsx
+++ b/Engine/Source/Programs/Horde/HordeDashboard/src/components/ConfirmationDialog.tsx
@@ -3,7 +3,7 @@
 import { Dialog, PrimaryButton, DialogFooter, DefaultButton, DialogType, TextField, Checkbox, ICheckbox } from '@fluentui/react';
 import React, { useState } from 'react';
 
-export const ConfirmationDialog: React.FC<{ isOpen: boolean; title: string; cancelText: string; confirmText: string; subText?: string; isTextBoxSpawned?: boolean; textBoxLabel?: string; checkBoxText?: string; onConfirm: (...args: any[]) => any; onCancel: (...args: any[]) => any }> = (({ isOpen, title, cancelText, confirmText, subText, isTextBoxSpawned, textBoxLabel, checkBoxText, onConfirm, onCancel }) => {
+export const ConfirmationDialog: React.FC<{ isOpen: boolean; title: string; cancelText: string; confirmText: string; subText?: string; isTextBoxSpawned?: boolean; textBoxLabel?: string; isTextRequired?: boolean, checkBoxText?: string; onConfirm: (...args: any[]) => any; onCancel: (...args: any[]) => any }> = (({ isOpen, title, cancelText, confirmText, subText, isTextBoxSpawned, textBoxLabel, isTextRequired, checkBoxText, onConfirm, onCancel }) => {
    const [textFieldText, setTextFieldText] = useState<string | undefined>(undefined);
    const checkboxRef = React.useRef<ICheckbox>(null);
 
@@ -24,7 +24,7 @@ export const ConfirmationDialog: React.FC<{ isOpen: boolean; title: string; canc
          {!!checkBoxText && <Checkbox componentRef={checkboxRef} label={checkBoxText} defaultChecked={false} />}
          {isTextBoxSpawned && <TextField label={textBoxLabel} value={textFieldText} onChange={(event, value) => setTextFieldText(value)}></TextField>}
          <DialogFooter>
-            <PrimaryButton onClick={() => { onConfirm(textFieldText,checkboxRef?.current?.checked ); setTextFieldText(undefined); }} text={confirmText} />
+            <PrimaryButton disabled={isTextBoxSpawned && isTextRequired && !textFieldText?.length} onClick={() => { onConfirm(textFieldText,checkboxRef?.current?.checked ); setTextFieldText(undefined); }} text={confirmText} />
             <DefaultButton onClick={() => { onCancel(textFieldText); setTextFieldText(undefined); }} text={cancelText} />
          </DialogFooter>
       </Dialog>
diff --git a/Engine/Source/Programs/Horde/HordeDashboard/src/components/HistoryModal.tsx b/Engine/Source/Programs/Horde/HordeDashboard/src/components/HistoryModal.tsx
--- a/Engine/Source/Programs/Horde/HordeDashboard/src/components/HistoryModal.tsx
+++ b/Engine/Source/Programs/Horde/HordeDashboard/src/components/HistoryModal.tsx
@@ -623,7 +623,7 @@ const TelemetryPanel: React.FC<{}> = observer(({ }) => {
             cpuText = `${v.replace("CPU=", "")}`;
          }
 
-         if (v.startsWith("RAM=")) {            
+         if (v.startsWith("RAM=")) {
             ramText = `${v.replace("RAM=", "")} GB`;
          }
 
@@ -723,7 +723,7 @@ const TelemetryPanel: React.FC<{}> = observer(({ }) => {
       </Stack>
       <Stack verticalAlign="center" verticalFill>
          <Stack>
-            <Label>{cpuLabel} Utilization {cpuText ? `- ${cpuText}` : ""}</Label>            
+            <Label>{cpuLabel} Utilization {cpuText ? `- ${cpuText}` : ""}</Label>
          </Stack>
          <div onMouseLeave={() => state.setLeaseTooltip(undefined)} onMouseMove={(ev) => { relativeCoords(ev, "cpu") }} onClick={() => {
 
@@ -1339,10 +1339,10 @@ export const HistoryModal: React.FC<{ agentId: string | undefined, onDismiss: (.
                <Stack style={{ paddingBottom: 18, paddingLeft: 4 }}>
                   <Text>{currentAction.confirmText}</Text>
                </Stack>
-               {!!currentAction.textInput && <TextField componentRef={actionTextInputRef} label={currentAction.name === "Edit Comment" ? "New Comment" : "Disable Reason"} />}
+               {!!currentAction.textInput && <TextField componentRef={actionTextInputRef} label={currentAction.name === "Edit Comment" ? "New Comment" : "Disable Reason - Required"} onChange={(ev, value) => setActionState({ ...actionState, comment: value })} spellCheck={false} autoComplete="false"/>}
                {currentAction.name === "Request Restart" && <Checkbox componentRef={forceRestartCheckboxRef} label={"Force Restart"} />}
                <DialogFooter>
-                  <PrimaryButton disabled={actionState.confirmed} onClick={() => { setActionState({ ...actionState, confirmed: true, comment: actionTextInputRef.current?.value }) }} text={currentAction.name} />
+                  <PrimaryButton disabled={actionState.confirmed || (!!currentAction.textInput && currentAction.name !== "Edit Comment" && !actionState.comment?.length)} onClick={() => { setActionState({ ...actionState, confirmed: true, comment: actionTextInputRef.current?.value }) }} text={currentAction.name} />
                   <DefaultButton disabled={actionState.confirmed} onClick={() => { setActionState({}) }} text="Cancel" />
                </DialogFooter>
             </Dialog>
