Horde: Updating schema documentation

#rnx

diff --git a/Engine/Source/Programs/Horde/Docs/Config/Schema/AclActions.md b/Engine/Source/Programs/Horde/Docs/Config/Schema/AclActions.md
--- a/Engine/Source/Programs/Horde/Docs/Config/Schema/AclActions.md
+++ b/Engine/Source/Programs/Horde/Docs/Config/Schema/AclActions.md
@@ -608,12 +608,6 @@
 | `UpdateReplicator` |  |
 | `ViewReplicator` |  |
 
-## Secrets
-
-| Name | Description |
-| ---- | ----------- |
-| `ViewSecret` | View a credential |
-
 ## ServiceAccounts
 
 | Name | Description |
@@ -787,7 +781,7 @@
 
 | Name | Description |
 | ---- | ----------- |
-| `ViewSecret` |  |
+| `ViewSecret` | View a credential |
 
 ## ServiceAccounts
 
@@ -988,19 +982,6 @@
 | `DownloadSoftware` |  |
 | `DeleteSoftware` |  |
 
-## Storage
-
-| Name | Description |
-| ---- | ----------- |
-| `ReadBlobs` | Ability to read blobs from the storage service |
-| `WriteBlobs` | Ability to write blobs to the storage service |
-| `ReadAliases` | Ability to read aliases from the storage service |
-| `WriteAliases` | Ability to write aliases to the storage service |
-| `DeleteAliases` | Ability to write aliases to the storage service |
-| `ReadRefs` | Ability to read refs from the storage service |
-| `WriteRefs` | Ability to write refs to the storage service |
-| `DeleteRefs` | Ability to delete refs |
-
 ## Streams
 
 | Name | Description |
@@ -1012,6 +993,12 @@
 | `ViewChanges` |  |
 | `ViewTemplate` |  |
 
+## Symbols
+
+| Name | Description |
+| ---- | ----------- |
+| `ReadSymbols` | Ability to download symbols |
+
 ## Accounts
 
 | Name | Description |
@@ -1180,14 +1167,14 @@
 
 | Name | Description |
 | ---- | ----------- |
-| `ReadBlobs` |  |
-| `WriteBlobs` |  |
-| `ReadAliases` |  |
-| `WriteAliases` |  |
-| `DeleteAliases` |  |
-| `ReadRefs` |  |
-| `WriteRefs` |  |
-| `DeleteRefs` |  |
+| `ReadBlobs` | Ability to read blobs from the storage service |
+| `WriteBlobs` | Ability to write blobs to the storage service |
+| `ReadAliases` | Ability to read aliases from the storage service |
+| `WriteAliases` | Ability to write aliases to the storage service |
+| `DeleteAliases` | Ability to write aliases to the storage service |
+| `ReadRefs` | Ability to read refs from the storage service |
+| `WriteRefs` | Ability to write refs to the storage service |
+| `DeleteRefs` | Ability to delete refs |
 
 ## Streams
 
@@ -1204,7 +1191,7 @@
 
 | Name | Description |
 | ---- | ----------- |
-| `ReadSymbols` | Ability to download symbols |
+| `ReadSymbols` |  |
 
 ## Accounts
 
diff --git a/Engine/Source/Programs/Horde/Docs/Config/Schema/Globals.md b/Engine/Source/Programs/Horde/Docs/Config/Schema/Globals.md
--- a/Engine/Source/Programs/Horde/Docs/Config/Schema/Globals.md
+++ b/Engine/Source/Programs/Horde/Docs/Config/Schema/Globals.md
@@ -62,10 +62,10 @@ Name | Description
 `build` | [BuildConfig](#buildconfig)<br>Configuration for the build plugin
 `compute` | [ComputeConfig](#computeconfig)<br>Configuration for the compute plugin
 `ddc` | [EmptyPluginConfig](#emptypluginconfig)<br>Configuration for the ddc plugin
-`secrets` | [SecretsConfig](#secretsconfig)<br>Configuration for the secrets plugin
 `experimental` | [ExperimentalConfig](#experimentalconfig)<br>Configuration for the experimental plugin
-`storage` | [StorageConfig](#storageconfig)<br>Configuration for the storage plugin
+`secrets` | [SecretsConfig](#secretsconfig)<br>Configuration for the secrets plugin
 `symbols` | [SymbolsConfig](#symbolsconfig)<br>Configuration for the symbols plugin
+`storage` | [StorageConfig](#storageconfig)<br>Configuration for the storage plugin
 `tools` | [ToolsConfig](#toolsconfig)<br>Configuration for the tools plugin
 
 ## AnalyticsConfig
@@ -423,6 +423,14 @@ Name | Description
 Empty implementation of
 
 
+## ExperimentalConfig
+
+Configuration for the experimental plugin
+
+Name | Description
+---- | -----------
+`notifications` | [NotificationConfig](Notification.md)`[]`<br>Notifications to be sent on the Horde server
+
 ## SecretsConfig
 
 Configuration for the secrets system
@@ -496,13 +504,24 @@ Name | Description
 `PreSharedKey` | Use a pre-shared token provided by Vault for authentication.
 `AwsAuth` | Use AWS IAM credentials to authenticate against the Vault server.
 
-## ExperimentalConfig
+## SymbolsConfig
 
-Configuration for the experimental plugin
+Configuration for the tools system
 
 Name | Description
 ---- | -----------
-`notifications` | [NotificationConfig](Notification.md)`[]`<br>Notifications to be sent on the Horde server
+`stores` | [SymbolStoreConfig](#symbolstoreconfig)`[]`<br>List of symbol stores
+
+## SymbolStoreConfig
+
+Configuration for a symbol store
+
+Name | Description
+---- | -----------
+`id` | `string`<br>Identifier for this store
+`namespaceId` | `string`<br>Configuration for the symbol store backend
+`public` | `boolean`<br>Whether to make this store available without auth
+`acl` | [AclConfig](#aclconfig)<br>Access to the symbol store
 
 ## StorageConfig
 
@@ -576,25 +595,6 @@ Name | Description
 `enableAliases` | `boolean`<br>Support querying exports by their aliases
 `acl` | [AclConfig](#aclconfig)<br>Access list for this namespace
 
-## SymbolsConfig
-
-Configuration for the tools system
-
-Name | Description
----- | -----------
-`stores` | [SymbolStoreConfig](#symbolstoreconfig)`[]`<br>List of symbol stores
-
-## SymbolStoreConfig
-
-Configuration for a symbol store
-
-Name | Description
----- | -----------
-`id` | `string`<br>Identifier for this store
-`namespaceId` | `string`<br>Configuration for the symbol store backend
-`public` | `boolean`<br>Whether to make this store available without auth
-`acl` | [AclConfig](#aclconfig)<br>Access to the symbol store
-
 ## ToolsConfig
 
 Configuration for the tools system
