Horde: Fix issue with agent modal not refreshing when modified, tweak disable comment attribution

diff --git a/Engine/Source/Programs/Horde/HordeDashboard/src/backend/Backend.ts b/Engine/Source/Programs/Horde/HordeDashboard/src/backend/Backend.ts
--- a/Engine/Source/Programs/Horde/HordeDashboard/src/backend/Backend.ts
+++ b/Engine/Source/Programs/Horde/HordeDashboard/src/backend/Backend.ts
@@ -662,9 +662,9 @@ export class Backend {
     updateAgent(agentId: string, request: UpdateAgentRequest): Promise<boolean> {
         if (request.enabled === false) {
             let comment = request.comment ?? "";
-            const userName = dashboard.user?.name;
+            const userName = dashboard.hordeUser ?? dashboard.user?.name;
             if (userName) {
-                request.comment = comment + " - "  + userName;
+                request.comment = comment + ` (${userName})`;
             }
         }
         return new Promise<boolean>((resolve, reject) => {
diff --git a/Engine/Source/Programs/Horde/HordeDashboard/src/backend/Dashboard.ts b/Engine/Source/Programs/Horde/HordeDashboard/src/backend/Dashboard.ts
--- a/Engine/Source/Programs/Horde/HordeDashboard/src/backend/Dashboard.ts
+++ b/Engine/Source/Programs/Horde/HordeDashboard/src/backend/Dashboard.ts
@@ -149,6 +149,16 @@ export class Dashboard {
 
     }
 
+    get hordeUser(): string | undefined {
+        const claims = this.claims;
+        const user = claims.filter(c => c.type.endsWith("/horde/user"));
+        if (!user.length) {
+            return undefined;
+        }
+
+        return user[0].value?.length ? user[0].value : undefined;
+    }
+
     get preview(): boolean {
 
         try {
@@ -206,7 +216,7 @@ export class Dashboard {
         return this.config?.artifactTypes ?? [];
     }
 
-    getArtifactTypeInfo(type:string): GetDashboardArtifactTypeResponse | undefined {
+    getArtifactTypeInfo(type: string): GetDashboardArtifactTypeResponse | undefined {
         return this.config?.artifactTypeInfo?.find(a => a.type === type);
     }
 
diff --git a/Engine/Source/Programs/Horde/HordeDashboard/src/components/HistoryModal.tsx b/Engine/Source/Programs/Horde/HordeDashboard/src/components/HistoryModal.tsx
--- a/Engine/Source/Programs/Horde/HordeDashboard/src/components/HistoryModal.tsx
+++ b/Engine/Source/Programs/Horde/HordeDashboard/src/components/HistoryModal.tsx
@@ -953,10 +953,17 @@ export const HistoryModal: React.FC<{ agentId: string | undefined, onDismiss: (.
          const request: UpdateAgentRequest = {};
          currentAction.update(request, actionState.comment);
          backend.updateAgent(agentId, request).then(() => {
-            agentStore.update(agentStore.pools?.length ? true : false).then(() => {
-               state.setSelectedAgent(agentStore.agents.find(agent => agent.id === agentId));
-               setActionState({});
-            });
+            if (!agentStore.pools?.length) {
+               agentStore.update(false).then(() => {
+                  state.setSelectedAgent(agentStore.agents.find(agent => agent.id === agentId));
+                  setActionState({});
+               });
+            } else {
+               agentStore.updateAgent(agentId).then(() => {
+                  state.setSelectedAgent(agentStore.agents.find(agent => agent.id === agentId));
+                  setActionState({});
+               });
+            }
          }).catch((reason) => {
             console.error(reason);
          });
@@ -1339,7 +1346,7 @@ export const HistoryModal: React.FC<{ agentId: string | undefined, onDismiss: (.
                <Stack style={{ paddingBottom: 18, paddingLeft: 4 }}>
                   <Text>{currentAction.confirmText}</Text>
                </Stack>
-               {!!currentAction.textInput && <TextField componentRef={actionTextInputRef} label={currentAction.name === "Edit Comment" ? "New Comment" : "Disable Reason - Required"} onChange={(ev, value) => setActionState({ ...actionState, comment: value })} spellCheck={false} autoComplete="false"/>}
+               {!!currentAction.textInput && <TextField componentRef={actionTextInputRef} label={currentAction.name === "Edit Comment" ? "New Comment" : "Disable Reason - Required"} onChange={(ev, value) => setActionState({ ...actionState, comment: value })} spellCheck={false} autoComplete="false" />}
                {currentAction.name === "Request Restart" && <Checkbox componentRef={forceRestartCheckboxRef} label={"Force Restart"} />}
                <DialogFooter>
                   <PrimaryButton disabled={actionState.confirmed || (!!currentAction.textInput && currentAction.name !== "Edit Comment" && !actionState.comment?.length)} onClick={() => { setActionState({ ...actionState, confirmed: true, comment: actionTextInputRef.current?.value }) }} text={currentAction.name} />
