Horde Dashboard: Click-to-dimiss for modals

diff --git a/Engine/Source/Programs/Horde/HordeDashboard/src/components/IssueViewV2.tsx b/Engine/Source/Programs/Horde/HordeDashboard/src/components/IssueViewV2.tsx
--- a/Engine/Source/Programs/Horde/HordeDashboard/src/components/IssueViewV2.tsx
+++ b/Engine/Source/Programs/Horde/HordeDashboard/src/components/IssueViewV2.tsx
@@ -1974,7 +1974,7 @@ export const IssueModalV2: React.FC<{ popHistoryOnClose: boolean, issueId?: stri
       </Modal>
    }
 
-   return <Modal isOpen={true} isBlocking={true} topOffsetFixed={true} styles={{ main: { padding: 8, width: 1420, backgroundColor: dashboard.darktheme ? `${modeColors.background} !important` : modeColors.background, hasBeenOpened: false, top: "24px", position: "absolute", height: "95vh" } }} className={hordeClasses.modal} onDismiss={() => { if (onCloseExternal) { onCloseExternal() } else { onClose() } }}>
+   return <Modal isOpen={true} topOffsetFixed={true} styles={{ main: { padding: 8, width: 1420, backgroundColor: dashboard.darktheme ? `${modeColors.background} !important` : modeColors.background, hasBeenOpened: false, top: "24px", position: "absolute", height: "95vh" } }} className={hordeClasses.modal} onDismiss={() => { if (onCloseExternal) { onCloseExternal() } else { onClose() } }}>
       {editShown && <EditIssueModal onClose={() => { setEditShown(false) }} />}
       <Stack style={{ height: "93vh" }}>
          <Stack style={{ height: "100%" }}>
diff --git a/Engine/Source/Programs/Horde/HordeDashboard/src/components/artifacts/ArtifactsModal.tsx b/Engine/Source/Programs/Horde/HordeDashboard/src/components/artifacts/ArtifactsModal.tsx
--- a/Engine/Source/Programs/Horde/HordeDashboard/src/components/artifacts/ArtifactsModal.tsx
+++ b/Engine/Source/Programs/Horde/HordeDashboard/src/components/artifacts/ArtifactsModal.tsx
@@ -700,7 +700,7 @@ export const JobArtifactsModal: React.FC<{ jobId: string; stepId: string, artifa
    const { hordeClasses } = getHordeStyling();
 
    return <Stack>
-      <Modal isOpen={true} isBlocking={true} topOffsetFixed={true} styles={{ main: { padding: 8, width: 1180, height: 820, hasBeenOpened: false, top: "80px", position: "absolute" } }} onDismiss={() => onClose()} className={hordeClasses.modal}>
+      <Modal isOpen={true} topOffsetFixed={true} styles={{ main: { padding: 8, width: 1180, height: 820, hasBeenOpened: false, top: "80px", position: "absolute" } }} onDismiss={() => onClose()} className={hordeClasses.modal}>
          <Stack className="horde-no-darktheme" styles={{ root: { paddingTop: 10, paddingRight: 12 } }}>
             <Stack style={{ paddingLeft: 24, paddingRight: 24 }}>
                <Stack tokens={{ childrenGap: 12 }} style={{ height: 800 }}>
diff --git a/Engine/Source/Programs/Horde/HordeDashboard/src/components/devices/DeviceEditor.tsx b/Engine/Source/Programs/Horde/HordeDashboard/src/components/devices/DeviceEditor.tsx
--- a/Engine/Source/Programs/Horde/HordeDashboard/src/components/devices/DeviceEditor.tsx
+++ b/Engine/Source/Programs/Horde/HordeDashboard/src/components/devices/DeviceEditor.tsx
@@ -275,7 +275,7 @@ export const DeviceEditor: React.FC<{ handler: DeviceHandler, deviceIn?: GetDevi
 
    }
 
-   return <Modal className={hordeClasses.modal} isOpen={true} isBlocking={true} styles={{ main: { padding: 8, width: 700 } }}>
+   return <Modal className={hordeClasses.modal} isOpen={true} styles={{ main: { padding: 8, width: 700 } }} onDismiss={() => onClose()}>
       {!!error && <MessageBar
          messageBarType={MessageBarType.error}
          isMultiline={false}
diff --git a/Engine/Source/Programs/Horde/HordeDashboard/src/components/devices/DeviceView.tsx b/Engine/Source/Programs/Horde/HordeDashboard/src/components/devices/DeviceView.tsx
--- a/Engine/Source/Programs/Horde/HordeDashboard/src/components/devices/DeviceView.tsx
+++ b/Engine/Source/Programs/Horde/HordeDashboard/src/components/devices/DeviceView.tsx
@@ -985,7 +985,7 @@ export const CheckoutConfirmModal: React.FC<{ check: "in" | "out" | "error", dev
    }
 
 
-   return <Modal isOpen={true} isBlocking={true} topOffsetFixed={true} styles={{ main: { padding: 8, width: 840, hasBeenOpened: false, top: "80px", position: "absolute" } }} className={hordeClasses.modal} onDismiss={() => { close() }}>
+   return <Modal isOpen={true} topOffsetFixed={true} styles={{ main: { padding: 8, width: 840, hasBeenOpened: false, top: "80px", position: "absolute" } }} className={hordeClasses.modal} onDismiss={() => { close() }}>
       <Stack tokens={{ childrenGap: 12 }}>
          <Stack horizontal styles={{ root: { paddingLeft: 12 } }}>
             <Stack grow style={{ paddingTop: 12, paddingLeft: 8 }}>
