Horde Dashboard: Alt-text for project logos

diff --git a/Engine/Source/Programs/Horde/HordeDashboard/src/components/UserHome.tsx b/Engine/Source/Programs/Horde/HordeDashboard/src/components/UserHome.tsx
--- a/Engine/Source/Programs/Horde/HordeDashboard/src/components/UserHome.tsx
+++ b/Engine/Source/Programs/Horde/HordeDashboard/src/components/UserHome.tsx
@@ -97,7 +97,7 @@ const ProjectsPanel: React.FC = observer(() => {
                      }}
                      >
                         <Link onClick={() => { projectStore.setActive(project.id); }} to={`/project/${project.id}`}>
-                           <img src={logoSrc} alt="" width={560 * .65} height={280 * .65} />
+                           <img src={logoSrc} alt={project.name + " logo button"} width={560 * .65} height={280 * .65} />
                         </Link>
                      </Stack>
                   );
