Horde Dashboard: Fixed HistoryModal.tsx hot reload issue

diff --git a/Engine/Source/Programs/Horde/HordeDashboard/src/App.tsx b/Engine/Source/Programs/Horde/HordeDashboard/src/App.tsx
--- a/Engine/Source/Programs/Horde/HordeDashboard/src/App.tsx
+++ b/Engine/Source/Programs/Horde/HordeDashboard/src/App.tsx
@@ -10,7 +10,7 @@ import { getSiteConfig } from './backend/Config';
 import dashboard from './backend/Dashboard';
 import { ThemeTester } from './base/components/ThemeTester/ThemeTester';
 import { AdminToken } from './components/AdminToken';
-import { AgentView } from './components/AgentView';
+import { AgentView } from './components/agents/AgentView';
 import { AuditLogView } from './components/AuditLog';
 import { AutomationView } from './components/AutomationView';
 import { DashboardView } from './components/DashboardView';
diff --git a/Engine/Source/Programs/Horde/HordeDashboard/src/backend/AgentStore.ts b/Engine/Source/Programs/Horde/HordeDashboard/src/backend/AgentStore.ts
deleted file mode 100644
diff --git a/Engine/Source/Programs/Horde/HordeDashboard/src/components/AgentView.tsx b/Engine/Source/Programs/Horde/HordeDashboard/src/components/AgentView.tsx
deleted file mode 100644
diff --git a/Engine/Source/Programs/Horde/HordeDashboard/src/components/AuditLog.tsx b/Engine/Source/Programs/Horde/HordeDashboard/src/components/AuditLog.tsx
--- a/Engine/Source/Programs/Horde/HordeDashboard/src/components/AuditLog.tsx
+++ b/Engine/Source/Programs/Horde/HordeDashboard/src/components/AuditLog.tsx
@@ -14,7 +14,7 @@ import { useWindowSize } from "../base/utilities/hooks";
 import { displayTimeZone } from "../base/utilities/timeUtils";
 import { BreadcrumbItem, Breadcrumbs } from "./Breadcrumbs";
 import { DateTimeRange } from "./DateTimeRange";
-import { HistoryModal } from "./HistoryModal";
+import { HistoryModal } from "./agents/HistoryModal";
 import { IssueModalV2 } from "./IssueViewV2";
 import { getLogStyles, logMetricNormal } from "./LogStyle";
 import { TopNav } from "./TopNav";
diff --git a/Engine/Source/Programs/Horde/HordeDashboard/src/components/HistoryModal.tsx b/Engine/Source/Programs/Horde/HordeDashboard/src/components/HistoryModal.tsx
deleted file mode 100644
diff --git a/Engine/Source/Programs/Horde/HordeDashboard/src/components/LogView.tsx b/Engine/Source/Programs/Horde/HordeDashboard/src/components/LogView.tsx
--- a/Engine/Source/Programs/Horde/HordeDashboard/src/components/LogView.tsx
+++ b/Engine/Source/Programs/Horde/HordeDashboard/src/components/LogView.tsx
@@ -20,7 +20,7 @@ import { AgentTelemetrySparkline } from "./agents/AgentTelemetrySparkline";
 import { JobArtifactsModal } from './artifacts/ArtifactsModal';
 import { Breadcrumbs } from './Breadcrumbs';
 import { ChangeContextMenu, ChangeContextMenuTarget } from './ChangeButton';
-import { HistoryModal } from './HistoryModal';
+import { HistoryModal } from './agents/HistoryModal';
 import { IssueModalV2 } from './IssueViewV2';
 import { useQuery } from './JobDetailCommon';
 import { LogItem, renderLine } from './LogRender';
diff --git a/Engine/Source/Programs/Horde/HordeDashboard/src/components/PoolView.tsx b/Engine/Source/Programs/Horde/HordeDashboard/src/components/PoolView.tsx
--- a/Engine/Source/Programs/Horde/HordeDashboard/src/components/PoolView.tsx
+++ b/Engine/Source/Programs/Horde/HordeDashboard/src/components/PoolView.tsx
@@ -12,8 +12,8 @@ import { PollBase } from "../backend/PollBase";
 import { projectStore } from "../backend/ProjectStore";
 import { getElapsedString, getShortNiceTime, getStepElapsed, getStepStartTime } from "../base/utilities/timeUtils";
 import { getHordeStyling, linearInterpolate } from "../styles/Styles";
-import { AgentPanel } from "./AgentView";
-import { HistoryModal } from "./HistoryModal";
+import { AgentPanel } from "./agents/AgentView";
+import { HistoryModal } from "./agents/HistoryModal";
 import { LeaseStatusIcon, StepStatusIcon } from "./StatusIcon";
 
 
diff --git a/Engine/Source/Programs/Horde/HordeDashboard/src/components/PoolsView.tsx b/Engine/Source/Programs/Horde/HordeDashboard/src/components/PoolsView.tsx
--- a/Engine/Source/Programs/Horde/HordeDashboard/src/components/PoolsView.tsx
+++ b/Engine/Source/Programs/Horde/HordeDashboard/src/components/PoolsView.tsx
@@ -10,7 +10,7 @@ import { PollBase } from "../backend/PollBase";
 import { useWindowSize } from "../base/utilities/hooks";
 import { getHordeStyling } from "../styles/Styles";
 import { BreadcrumbItem, Breadcrumbs } from "./Breadcrumbs";
-import { HistoryModal } from "./HistoryModal";
+import { HistoryModal } from "./agents/HistoryModal";
 import { PoolView } from "./PoolView";
 import { TopNav } from "./TopNav";
 
diff --git a/Engine/Source/Programs/Horde/HordeDashboard/src/components/agents/AgentRequestsView.tsx b/Engine/Source/Programs/Horde/HordeDashboard/src/components/agents/AgentRequestsView.tsx
--- a/Engine/Source/Programs/Horde/HordeDashboard/src/components/agents/AgentRequestsView.tsx
+++ b/Engine/Source/Programs/Horde/HordeDashboard/src/components/agents/AgentRequestsView.tsx
@@ -13,7 +13,7 @@ import { TopNav } from "../TopNav";
 import ErrorHandler from "../ErrorHandler";
 import { useNavigate } from "react-router-dom";
 import moment from "moment";
-import { agentStore } from "horde/backend/AgentStore";
+import { agentStore } from "./AgentStore"
 
 class AgentRequestsHandler extends PollBase {
 
diff --git a/Engine/Source/Programs/Horde/HordeDashboard/src/components/agents/AgentStore.ts b/Engine/Source/Programs/Horde/HordeDashboard/src/components/agents/AgentStore.ts
new file mode 100644
--- /dev/null
+++ b/Engine/Source/Programs/Horde/HordeDashboard/src/components/agents/AgentStore.ts
@@ -0,0 +1,157 @@
+// Copyright Epic Games, Inc. All Rights Reserved.
+
+import { action, makeObservable, observable } from 'mobx';
+import backend from '../../backend';
+import { AgentData, CategoryAgents, GetAgentResponse, PoolData } from '../../backend/Api';
+
+export class AgentStore {
+
+    constructor() {
+        makeObservable(this);
+    }
+
+    @observable
+    agentsUpdated: number = 0
+
+    get agents(): AgentData[] {
+        // subscribe in any observers
+        if (this.agentsUpdated) { }
+        return this._agents;
+    }
+
+    private _agents: AgentData[] = [];
+
+    @observable
+    poolsUpdated: number = 0;
+
+    get pools(): PoolData[] {
+        // subscribe in any observers
+        if (this.poolsUpdated) { }
+        return this._pools;
+    }
+
+    private _pools: PoolData[] = [];
+
+    // don't need this for now
+    modifiedAfterDate?: Date;
+
+    inflight?: boolean;
+
+    @action
+    private _setPools(data: PoolData[]) {
+        this._pools = data;
+        this.poolsUpdated++;
+    }
+
+    @action
+    private _setAgents(data: AgentData[]) {
+        if (data.length !== 0) {
+            const toSet: AgentData[] = [...this._agents];
+            data.forEach(updatedAgent => {
+                const newUpdateIdx = toSet.findIndex(existingAgent => existingAgent.id === updatedAgent.id);
+                if (newUpdateIdx !== -1) {
+                    toSet[newUpdateIdx] = updatedAgent;
+                }
+                else {
+                    toSet.push(updatedAgent);
+                }
+            });            
+            this._agents = toSet.filter(a => true /*!!a.deleted*/);            
+        } else {
+            this._agents = [];
+        }
+        this.agentsUpdated++;
+    }
+
+    async updateAgent(agentId: string): Promise<GetAgentResponse> {
+        return new Promise<GetAgentResponse>(async (resolve, reject) => {
+            try {
+                const agent = await backend.getAgent(agentId);
+                this._agents = this._agents.filter(a => a.id !== agentId);
+                this._agents.push(agent);
+
+                const load = agent.pools?.filter(p => !this._pools.find(pool => pool.id === p));
+                if (load?.length) {
+                    this._pools = await backend.getPools()
+                }                
+
+                resolve(agent);
+            } catch (error) {
+                reject(error);
+            }
+        })
+    }
+
+    categoryAgents = new Map<string, CategoryAgents>();
+
+    private forceUpdate() {
+        this._setAgents([...this._agents]);
+    }
+
+    getCategory(condition: string): CategoryAgents {
+
+        let agents = this.categoryAgents.get(condition);
+
+        if (agents && ( agents.polling || (Date.now() - agents.lastPoll.getTime()) / 1000 < 60)) {
+            return agents;
+        }
+        
+        agents = { ids: agents?.ids ?? [], lastPoll: new Date(), polling: !agents };
+        this.categoryAgents.set(condition, agents);
+
+        const filter = "id";
+        backend.getAgents({ condition: condition, filter: filter }).then((agents) => {
+
+            const oldAgents = new Set(this.categoryAgents.get(condition)?.ids ?? []);
+            const ids = agents.map(a => a.id);        
+            this.categoryAgents.set(condition, { ids: ids, lastPoll: new Date() })
+
+            const newAgents = new Set(ids);
+
+            if (oldAgents.size !== newAgents.size || ![...oldAgents].every(x => newAgents.has(x))) {
+                // force update
+                this.forceUpdate();    
+            }
+        });
+                
+        return agents;
+
+    }
+
+    async update(slim:boolean = false, invalidateCache: boolean = false): Promise<void> {
+        return new Promise<void>((resolve, reject) => {
+
+            if (this.inflight !== undefined && slim) {                
+                resolve();
+                return;
+            }
+
+            if (invalidateCache) {
+                this.modifiedAfterDate = undefined;
+            }
+
+            this.inflight = true;
+
+            const filter = "id,name,sessionId,sessionExpiresAt,online,enabled,ephemeral,comment,version,forceVersion,pools,capabilities,leases,acl,updateTime,deleted,pendingConform,pendingFullConform,conformAttemptCount,lastConformTime,nextConformTime,lastShutdownReason,pendingShutdown,status";
+            const promises: any[] = [];
+            promises.push(backend.getAgents({ includeDeleted: false, modifiedAfter: this.modifiedAfterDate?.toISOString(), filter:filter,invalidateCache:invalidateCache}));
+            if (!slim) {
+                promises.push(backend.getPools());
+            }
+            Promise.all(promises).then(values => {
+                this.modifiedAfterDate = new Date();
+                this._setAgents(values[0]);
+                if (!slim) {
+                    this._setPools(values[1]);
+                }
+                resolve();
+            }).catch(reason => {
+                console.error(`Error getting agents and pools: ${reason}`);
+            }).finally(() => {
+                this.inflight = undefined;
+            });
+        });
+    }
+}
+
+export const agentStore = new AgentStore();
\ No newline at end of file
diff --git a/Engine/Source/Programs/Horde/HordeDashboard/src/components/agents/AgentView.tsx b/Engine/Source/Programs/Horde/HordeDashboard/src/components/agents/AgentView.tsx
new file mode 100644
--- /dev/null
+++ b/Engine/Source/Programs/Horde/HordeDashboard/src/components/agents/AgentView.tsx
@@ -0,0 +1,2917 @@
+// Copyright Epic Games, Inc. All Rights Reserved.
+import { Checkbox, CommandButton, ConstrainMode, ContextualMenu, DefaultButton, DetailsHeader, DetailsList, DetailsListLayoutMode, Dialog, DialogType, DirectionalHint, Dropdown, FontIcon, FontSizes, FontWeights, IBasePickerProps, IColumn, IContextualMenuItem, IContextualMenuProps, IDetailsHeaderProps, IDetailsHeaderStyles, IDetailsListProps, ITag, ITagItemStyles, ITooltipHostStyles, Icon, IconButton, Pivot, PivotItem, PrimaryButton, ProgressIndicator, Link as ReactLink, ScrollablePane, ScrollbarVisibility, SearchBox, Selection, SelectionMode, Slider, Spinner, SpinnerSize, Stack, Sticky, StickyPositionType, TagItem, TagPicker, Text, TextField, TooltipHost, mergeStyleSets } from '@fluentui/react';
+import { action, makeObservable, observable } from 'mobx';
+import { observer } from 'mobx-react-lite';
+import moment from 'moment-timezone';
+import React, { createRef, useEffect, useState } from 'react';
+import { Link, useSearchParams } from 'react-router-dom';
+import Marquee from 'react-text-marquee';
+import backend from '../../backend';
+import { agentStore } from './AgentStore';
+import { AgentData, BatchUpdatePoolRequest, GetAgentResponse, GetDashboardAgentCategoryResponse, LeaseState, PoolData } from '../../backend/Api';
+import dashboard, { StatusColor } from '../../backend/Dashboard';
+import { copyToClipboard } from '../../base/utilities/clipboard';
+import { useWindowSize } from '../../base/utilities/hooks';
+import { getShortNiceTime } from '../../base/utilities/timeUtils';
+import { getHordeStyling, hexToRGB, linearInterpolate } from '../../styles/Styles';
+import { getHordeTheme } from '../../styles/theme';
+import { Breadcrumbs } from '../Breadcrumbs';
+import { ConfirmationDialog } from '../ConfirmationDialog';
+import { HistoryModal } from './HistoryModal';
+import { TopNav } from '../TopNav';
+
+let _agentStyles: any;
+const getAgentStyles = () => {
+   const theme = getHordeTheme();
+
+   const agentStyles = _agentStyles ?? mergeStyleSets({
+
+      ticker: {
+         width: '100%'
+      },
+      checkboxCell: {
+         selectors: {
+            '> div': {
+               height: '100%'
+            }
+         }
+      },
+      dialog: {
+         selectors: {
+            ".ms-Label,.ms-Button-label": {
+               fontWeight: "unset",
+               fontFamily: "Horde Open Sans SemiBold"
+            }
+         }
+      },
+      detailsList: {
+         selectors: {
+            ".ms-DetailsHeader-cellName": {
+               fontWeight: "unset",
+               fontFamily: "Horde Open Sans SemiBold"
+            },
+            ".ms-DetailsRow #editagent": {
+               opacity: 0
+            },
+            ".ms-DetailsRow:hover #editagent": {
+               opacity: 1
+            }
+         }
+      },
+      descFont: {
+         font: '8pt Horde Open Sans SemiBold !important',
+         marginLeft: 12,
+         marginRight: 12,
+         marginBottom: 4
+      },
+      buttonFont: {
+         height: '26px',
+         font: '7pt Horde Open Sans SemiBold !important',
+         flexShrink: '0 !important',
+         paddingLeft: 6,
+         paddingRight: 6,
+         selectors: {
+            '.ms-Icon': {
+               width: 0,
+               margin: 0
+            },
+            ':active': {
+               textDecoration: 'none'
+            },
+            ':hover': {
+               textDecoration: 'none'
+            },
+            ':visited': {
+               textDecoration: 'none',
+               color: "#FFFFFF"
+            }
+
+         }
+      },
+      modalHeader: [
+         {
+            font: '24px Horde Open Sans Light',
+            flex: '1 1 auto',
+            color: theme.palette.neutralPrimary,
+            display: 'flex',
+            fontSize: FontSizes.xLarge,
+            alignItems: 'center',
+            fontWeight: FontWeights.semibold,
+            padding: '12px 12px 14px 24px'
+         }
+      ],
+      modalBody: {
+         flex: '4 4 auto',
+         padding: '0 24px 24px 24px',
+         overflowY: 'hidden',
+         selectors: {
+            p: {
+               margin: '14px 0'
+            },
+            'p:first-child': {
+               marginTop: 0
+            },
+            'p:last-child': {
+               marginBottom: 0
+            }
+         }
+      },
+      ellipsesStackItem: {
+         whiteSpace: 'nowrap',
+         overflow: 'hidden',
+         textOverflow: 'ellipsis'
+      }
+   });
+
+   _agentStyles = agentStyles;
+
+   return agentStyles;
+}
+
+
+
+// column in the main table
+type ColumnItem = {
+   key: string;
+   displayText: string;
+   colSize: number;
+   isChecked: boolean;
+   isCheckable: boolean;
+   isSorted: boolean;
+   isSortedDescending: boolean;
+   columnDef: IColumn | undefined;
+}
+
+
+function getAgentCapability(agent: AgentData, inProp: string) {
+   if (agent.capabilities?.devices) {
+      for (const deviceIdx in agent.capabilities?.devices) {
+         const device = agent.capabilities?.devices[deviceIdx];
+         if (device.properties) {
+            for (const propIdx in device.properties) {
+               const prop = device.properties[propIdx];
+               if (prop.indexOf(inProp) !== -1) {
+                  return prop.split('=')[1];
+               }
+            }
+         }
+      }
+   }
+   return null;
+}
+
+function getAgentOs(agent: AgentData) {
+   const osDist = getAgentCapability(agent, "OSDistribution");
+   const osFamily = getAgentCapability(agent, "OSFamily");
+   return osDist ?? osFamily;
+}
+
+function getTaskTime(agent: GetAgentResponse): number {
+
+   let max: moment.Duration | undefined
+   if (agent.leases) {
+      agent.leases.forEach(lease => {
+         const start = moment(lease.startTime);
+         const end = moment(Date.now());
+         const duration = moment.duration(end.diff(start));
+         if (!max) {
+            max = duration;
+         }
+         else {
+            if (duration.asSeconds() > max.asSeconds()) {
+               max = duration;
+            }
+         }
+      });
+   }
+
+   return max?.asSeconds() ?? 0;
+
+}
+
+type SearchState = {
+   agentId?: string;
+   agentSearch?: string;
+   exactSearch?: boolean;
+   filter?: string[];
+   columnMode?: string;
+   category?: string;
+}
+
+class LocalState {
+
+   // agent update function params
+   editingPools = false;  // if editing pools dialog is open
+
+   // text in the filter search
+   @observable agentFilter = "";
+   @observable filterExactMatch: boolean = false;
+   @observable agentStatusFilter = new Set<string>();
+
+   // table selection
+   selection: Selection = new Selection({ onSelectionChanged: () => { this.setAgentsSelectedCount(this.selection.getSelectedCount()); } });
+   currentSelection: AgentData[] = [];
+   // number of agents selected for the button
+   @observable agentsSelectedCount = 0;
+
+   @observable deleteAgentDialogIsOpen = false;
+   @observable restartAgentDialogIsOpen = false;
+   @observable shutdownAgentDialogIsOpen = false;
+   @observable disableAgentDialogIsOpen = false;
+   @observable cancelLeasesDialogIsOpen = false;
+   @observable editCommentDialogIsOpen = false;
+
+   // column state for the table
+   @observable columnsState: ColumnItem[];
+   columnSearchState: any = {};
+   @observable columnMenuProps: IContextualMenuItem[];
+
+   @observable headerContextMenuOpen = false;
+   @observable agentContextMenuOpen = false;
+   contextMenuRef: any = null;
+   contextMenuTargetRef: any = createRef<Element>();
+   @observable mouseX = 0;
+   @observable mouseY = 0;
+
+   @action
+   resetState() {
+      this.agentFilter = "";
+      this.filterExactMatch = false;
+      this.selection = new Selection({ onSelectionChanged: () => { this.setAgentsSelectedCount(this.selection.getSelectedCount()); } });
+      this.currentSelection = [];
+      this.deleteAgentDialogIsOpen = false;
+      this.searchState = {};
+   }
+
+   @action
+   setRightClickDiv(x: number, y: number) {
+      this.mouseX = x;
+      this.mouseY = y;
+   }
+
+   @action
+   setAgentFilter(filter: string) {
+      this.searchState.agentSearch = filter?.trim() ? filter.trim() : undefined;
+      this.updateSearch();
+      this.agentFilter = filter;
+   }
+
+   @action
+   setAgentStaus(filter: Set<string>) {
+      const filters = Array.from(filter);
+      this.searchState.filter = filters.length ? filters : undefined;
+      this.updateSearch();
+      this.agentStatusFilter = filter;
+   }
+
+   @action
+   setExactMatch(match: boolean) {
+      this.searchState.exactSearch = match ? true : undefined;
+      this.updateSearch();
+      this.filterExactMatch = match;
+   }
+
+   @action
+   setAgentId(id?: string) {
+      this.searchState.agentId = id;
+      this.updateSearch();
+   }
+
+   @action
+   setAgentCategory(category?: string) {
+      this.searchState.category = category?.trim();
+      this.updateSearch();
+   }
+
+   @action
+   private _onColumnClick(ev: React.MouseEvent<HTMLElement>, column: IColumn) {
+      if (column.key === "pools") {
+         return;
+      }
+      this.columnsState.forEach(colState => {
+         if (colState.columnDef!.key === column.key) {
+            colState.isSortedDescending = !colState.isSortedDescending;
+            colState.isSorted = true;
+         }
+         else {
+            colState.isSorted = false;
+            colState.isSortedDescending = false;
+         }
+      });
+      this._updateColumnDefs();
+   }
+
+   // enable or disable builders.
+   async changeBuilderEnabled(enabled: string, disableText?: string) {
+      let that = this;
+      const enabledVal: boolean = enabled === "enable" ? true : false;
+      const allUpdates: any[] = [];
+      const selectedAgents = this.currentSelection;
+      selectedAgents.forEach(agent => {
+         // add to update queue if the agent isnt set to the value
+         if (agent.enabled !== enabledVal) {
+            allUpdates.push(backend.updateAgent(agent.id, { enabled: enabledVal, comment: disableText ?? "" }));
+         }
+      });
+      await Promise.all(allUpdates).then(function (responses) {
+         if (!enabledVal) {
+            that.setDisableBuilderDialogOpen(false);
+         }
+      }).catch(function (errors) {
+      }).finally(function () {
+         agentStore.update();
+      });
+   }
+
+   // enable or disable builders.
+   async changeCommment(commentText?: string) {
+      let that = this;
+      const comment = commentText ?? "";
+      const allUpdates: any[] = [];
+      const selectedAgents = this.currentSelection;
+      selectedAgents.forEach(agent => {
+         // add to update queue if the agent isnt set to the value
+         if (agent.comment !== comment) {
+            allUpdates.push(backend.updateAgent(agent.id, { comment: comment }));
+         }
+      });
+      await Promise.all(allUpdates).then(function () {
+         that.setEditCommentDialogOpen(false);
+      }).catch(function (errors) {
+      }).finally(function () {
+         agentStore.update();
+      });
+   }
+
+   async requestBuilderUpdate(conform: boolean, restart: boolean, fullConform?: boolean, shutdown?: boolean, forceRestart?: boolean) {
+      let that = this;
+      const allUpdates: any[] = [];
+      const selectedAgents = this.currentSelection;
+      selectedAgents.forEach(agent => {
+         allUpdates.push(backend.updateAgent(agent.id, { requestConform: conform, requestFullConform: fullConform, requestRestart: restart, requestShutdown: shutdown, requestForceRestart: !!forceRestart }));
+      });
+      await Promise.all(allUpdates).then(function (responses) {
+         if (restart) {
+            that.setRestartBuilderDialogOpen(false);
+         }
+         if (shutdown) {
+            that.setShutdownBuilderDialogOpen(false);
+         }
+      }).catch(function (errors) {
+      }).finally(function () {
+         agentStore.update();
+      });
+   };
+
+   async cancelLeases() {
+
+      const selectedAgents = this.currentSelection;
+
+      const leases = new Set<string>();
+
+      selectedAgents.forEach(agent => {
+         agent.leases?.forEach(lease => {
+            if (!lease.finishTime) {
+               leases.add(lease.id);
+            }
+         })
+      });
+
+      let requests = Array.from(leases).map(id => {
+         return backend.updateLease(id, { aborted: true });
+      });
+
+      while (requests.length) {
+
+         const batch = requests.slice(0, 5);
+
+         await Promise.all(batch).then(() => {
+
+         }).catch((errors) => {
+            console.log(errors);
+            // eslint-disable-next-line
+         }).finally(() => {
+
+            requests = requests.slice(5);
+         });
+      }
+
+      agentStore.update();
+
+   };
+
+
+   @action setDisableBuilderDialogOpen(isOpen: boolean) {
+      this.disableAgentDialogIsOpen = isOpen;
+   }
+
+   @action setEditCommentDialogOpen(isOpen: boolean) {
+      this.editCommentDialogIsOpen = isOpen;
+   }
+
+   @action setRestartBuilderDialogOpen(isOpen: boolean) {
+      this.restartAgentDialogIsOpen = isOpen;
+   }
+
+   @action setShutdownBuilderDialogOpen(isOpen: boolean) {
+      this.shutdownAgentDialogIsOpen = isOpen;
+   }
+
+   @action setCancelLeasesDialogOpen(isOpen: boolean) {
+      this.cancelLeasesDialogIsOpen = isOpen;
+   }
+
+   @action setDeleteBuilderDialogOpen(isOpen: boolean) {
+      this.deleteAgentDialogIsOpen = isOpen;
+   }
+
+   async deleteBuilder() {
+      let that = this;
+      const allUpdates: any[] = [];
+      const selectedAgents = this.currentSelection;
+      selectedAgents.forEach(agent => {
+         allUpdates.push(backend.deleteAgent(agent.id));
+      });
+      await Promise.all(allUpdates).then(function (responses) {
+         that.setDeleteBuilderDialogOpen(false);
+      }).catch(function (errors) {
+      }).finally(function () {
+         that.selection.setAllSelected(false);
+         that.currentSelection = [];
+         agentStore.update(false, true);
+      });
+   };
+
+   // sets a "Columns..." button item checked or not by updating the col state
+   @action
+   setColumnItemChecked(key: string) {
+      this.columnsState.filter(colState => { return colState.isCheckable; }).forEach(colState => {
+         if (colState.key === key) {
+            colState.isChecked = true;
+            this.searchState.columnMode = key ? key : undefined;
+            this.updateSearch();
+
+         }
+         else {
+            colState.isChecked = false;
+         }
+      });
+
+      //this.searchState.columnMode = key;
+      //this.updateSearch();
+
+      this.columnMenuProps = this._updateColumnProps();
+   }
+
+   private _updateColumnDefs() {
+      this.columnsState.forEach(colState => {
+
+         let colSize = colState.colSize;
+         
+         // True when colState is for a column type that can be selected for the last column
+         const bIsSwapColumn: boolean = (
+            colState.key === "software"
+         || colState.key === "storage"
+         || colState.key.indexOf("systemInfo") !== -1
+         || colState.key === "taskTime"
+         || colState.key === "comment"
+         );
+
+         if (!this.agentView && colState.displayText === "Pools") {
+            colSize = 480;
+         }
+
+         colState.columnDef = {
+            key: colState.key,
+            name: colState.displayText,
+            minWidth: colSize,
+            maxWidth: colSize,
+            isResizable: false,
+            isSorted: colState.key === "pools" ? undefined : colState.isSorted,
+            isSortedDescending: colState.key === "pools" ? undefined : colState.isSortedDescending,
+            onColumnClick: this._onColumnClick.bind(this),
+            styles: bIsSwapColumn ? {
+                  cellName: {
+                     flexGrow: '1' // Grow the cell name element so there can be space between it's elements
+                  }
+            } : undefined,
+            onRenderHeader: bIsSwapColumn ? () => {
+               return <Stack horizontal>
+                  <span style={{flexGrow: 1, textAlign: 'center'}}>{colState.displayText}</span>
+                  <TooltipHost
+                     content="Change column content"
+                  >
+                     <IconButton
+                        iconProps={{ iconName: 'Edit' }}
+                        onClick={(ev) => {
+                           ev.preventDefault();
+                           ev.stopPropagation(); // Consume the input before it causes a sort re-order
+                           localState.contextMenuTargetRef = ev?.nativeEvent!;
+                           localState.setHeaderContextMenuOpen(true);
+                        }
+                     }/>
+                  </TooltipHost>
+               </Stack>
+            }
+            : undefined,
+            onRender: colState.key === "editAgent" ? () => {
+               return <Stack style={{ cursor: "pointer" }} verticalFill verticalAlign='center' horizontalAlign='start' onClick={(ev) => {
+                  localState.setRightClickDiv(ev?.clientX, ev?.clientY);
+                  localState.setAgentContextMenuOpen(true);
+               }
+               }><FontIcon id="editagent" style={{ fontSize: 14 }} iconName="Edit" />
+               </Stack>
+            } : undefined
+         };
+
+      });
+   }
+
+   private _updateColumnProps() {
+      return [
+         {
+            canCheck: true,
+            key: 'software',
+            checked: this.columnsState.find(c => c.key === 'software')!.isChecked,
+            text: this.columnsState.find(c => c.key === 'software')!.displayText,
+            onClick: () => this.setColumnItemChecked('software')
+         },
+         {
+            canCheck: true,
+            key: 'taskTime',
+            checked: this.columnsState.find(c => c.key === 'taskTime')!.isChecked,
+            text: this.columnsState.find(c => c.key === 'taskTime')!.displayText,
+            onClick: () => this.setColumnItemChecked('taskTime')
+         },
+         {
+            canCheck: true,
+            key: 'storage',
+            checked: this.columnsState.find(c => c.key === 'storage')!.isChecked,
+            text: this.columnsState.find(c => c.key === 'storage')!.displayText,
+            onClick: () => this.setColumnItemChecked('storage')
+         },
+         {
+            canCheck: true,
+            key: 'comment',
+            checked: this.columnsState.find(c => c.key === 'comment')!.isChecked,
+            text: this.columnsState.find(c => c.key === 'comment')!.displayText,
+            onClick: () => this.setColumnItemChecked('comment')
+         },
+         {
+            canCheck: true,
+            key: 'OS',
+            checked: this.columnsState.find(c => c.key === 'systemInfoOS')!.isChecked,
+            text: this.columnsState.find(c => c.key === 'systemInfoOS')!.displayText,
+            onClick: () => this.setColumnItemChecked('systemInfoOS')
+         },
+         {
+            canCheck: true,
+            key: 'CPU',
+            checked: this.columnsState.find(c => c.key === 'systemInfoCPU')!.isChecked,
+            text: this.columnsState.find(c => c.key === 'systemInfoCPU')!.displayText,
+            onClick: () => this.setColumnItemChecked('systemInfoCPU')
+         },
+         {
+            canCheck: true,
+            key: 'RAM',
+            checked: this.columnsState.find(c => c.key === 'systemInfoRAM')!.isChecked,
+            text: this.columnsState.find(c => c.key === 'systemInfoRAM')!.displayText,
+            onClick: () => this.setColumnItemChecked('systemInfoRAM')
+         }
+      ];
+   }
+
+   // sets the number of agents selected
+   @action
+   setAgentsSelectedCount(count: number) {
+      this.agentsSelectedCount = count;
+   }
+
+   @action
+   setHeaderContextMenuOpen(isOpen: boolean) {
+      this.agentContextMenuOpen = false;
+      this.headerContextMenuOpen = isOpen;
+   }
+
+   @action
+   setAgentContextMenuOpen(isOpen: boolean) {
+      if (isOpen) {
+         this.currentSelection = [...this.selection.getSelection() as AgentData[]];
+      }
+      this.agentContextMenuOpen = isOpen;
+      this.headerContextMenuOpen = false;
+   }
+
+
+   @observable
+   searchUpdated: number = 0;
+
+   @action
+   setSearchUpdated() {
+      this.searchUpdated++;
+   }
+
+
+   searchState: SearchState = {};
+   search: URLSearchParams = new URLSearchParams();
+
+   updateSearch(): boolean {
+
+      const state = { ...this.searchState } as SearchState;
+
+      state.filter = state.filter?.sort((a, b) => a.localeCompare(b));
+
+      const search = new URLSearchParams();
+      const csearch = this.search.toString();
+
+      if (state.agentSearch) {
+         search.append("agent", state.agentSearch);
+      }
+
+      if (state.exactSearch) {
+         search.append("exact", "true");
+      }
+
+      if (state.agentId) {
+         search.append("agentId", state.agentId);
+      }
+
+      if (state.category) {
+         search.append("category", encodeURIComponent(state.category));
+      }
+
+      state.filter?.forEach(f => {
+         if (f) {
+            search.append("filter", f);
+         }
+      });
+
+      if (state.columnMode) {
+         search.append("mode", state.columnMode);
+      }
+
+      if (search.toString() !== csearch) {
+         this.search = search;
+         this.setSearchUpdated();
+         return true;
+      }
+
+      return false;
+   }
+
+   stateFromSearch(search: URLSearchParams) {
+
+      const state: SearchState = {};
+
+      const filters = search.getAll("filter") ?? undefined;
+      const agentSearch = search.get("agent") ?? undefined;
+      const agentId = search.get("agentId") ?? undefined;
+      const exact = search.get("exact") ?? undefined;
+      const mode = search.get("mode") ?? undefined;
+      const category = search.get("category") ?? undefined;
+
+      state.filter = filters?.sort((a, b) => a.localeCompare(b));
+      state.columnMode = mode?.length ? mode : undefined;
+      state.agentSearch = agentSearch?.length ? agentSearch : undefined;
+      state.exactSearch = exact?.trim() === "true" ? true : undefined;
+      state.agentId = agentId?.trim() ? agentId : undefined;
+      state.category = category?.trim() ? category : undefined;
+
+      this.search = search;
+      this.searchState = state;
+
+      if (state.agentSearch) {
+         this.agentFilter = state.agentSearch;
+      }
+      if (state.exactSearch) {
+         this.filterExactMatch = true;
+      }
+      if (state.filter?.length) {
+         this.agentStatusFilter = new Set<string>(state.filter);
+      }
+
+      if (state.columnMode) {
+         this.setColumnItemChecked(state.columnMode);
+      }
+
+      return state;
+   }
+
+   agentView: boolean | undefined = true;
+
+   constructor() {
+      makeObservable(this);
+      this.columnsState = [
+         {
+            key: 'name',
+            displayText: 'Name',
+            colSize: 170,
+            isChecked: true,
+            isCheckable: false,
+            isSorted: true,
+            isSortedDescending: false,
+            columnDef: undefined
+         },
+         {
+            key: 'pools',
+            displayText: 'Pools',
+            colSize: 476,
+            isChecked: true,
+            isCheckable: false,
+            isSorted: false,
+            isSortedDescending: false,
+            columnDef: undefined
+         },
+         {
+            key: 'status',
+            displayText: 'Status',
+            colSize: 300,
+            isChecked: true,
+            isCheckable: false,
+            isSorted: false,
+            isSortedDescending: false,
+            columnDef: undefined
+         },
+         {
+            key: 'software',
+            displayText: 'Software',
+            colSize: 220,
+            isChecked: false,
+            isCheckable: true,
+            isSorted: false,
+            isSortedDescending: false,
+            columnDef: undefined
+         },
+         {
+            key: 'taskTime',
+            displayText: 'Task Time',
+            colSize: 220,
+            isChecked: false,
+            isCheckable: true,
+            isSorted: false,
+            isSortedDescending: false,
+            columnDef: undefined
+         },
+         {
+            key: 'storage',
+            displayText: 'Storage',
+            colSize: 220,
+            isChecked: true,
+            isCheckable: true,
+            isSorted: false,
+            isSortedDescending: false,
+            columnDef: undefined
+         },
+         {
+            key: 'comment',
+            displayText: 'Comment',
+            colSize: 220,
+            isChecked: false,
+            isCheckable: true,
+            isSorted: false,
+            isSortedDescending: false,
+            columnDef: undefined
+         },
+         {
+            key: 'systemInfoOS',
+            displayText: 'OS',
+            colSize: 220,
+            isChecked: false,
+            isCheckable: true,
+            isSorted: false,
+            isSortedDescending: false,
+            columnDef: undefined
+         },
+         {
+            key: 'systemInfoCPU',
+            displayText: 'CPU',
+            colSize: 220,
+            isChecked: false,
+            isCheckable: true,
+            isSorted: false,
+            isSortedDescending: false,
+            columnDef: undefined
+         },
+         {
+            key: 'systemInfoRAM',
+            displayText: 'RAM',
+            colSize: 220,
+            isChecked: false,
+            isCheckable: true,
+            isSorted: false,
+            isSortedDescending: false,
+            columnDef: undefined
+         },
+         {
+            key: 'editAgent',
+            displayText: '',
+            colSize: 24,
+            isChecked: true,
+            isCheckable: false,
+            isSorted: false,
+            isSortedDescending: false,
+            columnDef: undefined
+         }
+      ];
+      this.columnSearchState = {
+         'pools': {},
+         'status': {},
+         'storage': {},
+         'comment': {},
+         'taskTime': {},
+         'taskTimeDescription': {}
+      };
+      this.columnMenuProps = this._updateColumnProps();
+      this._updateColumnDefs();
+   }
+
+   updateColumns() {
+      this._updateColumnDefs();
+   }
+}
+
+// pool editor list item
+type PoolEditorItem = {
+   key: string;
+   pool: PoolData;
+   deleted: boolean;
+   sortKey: string;
+   numAgentsAssigned: number;
+}
+
+// state class for the edit pools modal
+class EditPoolsModalState {
+   constructor() {
+      makeObservable(this);
+   }
+   // pool editor observables
+   newIdSuffix = 0;
+   @observable isOpen = false;
+   @observable isPoolValueValid = true;
+   @observable isEditorOpen = false;
+   @observable modifiedPools: PoolEditorItem[] = [];
+   @observable selectedColor = "";
+
+   @observable isConfirmationOpen = false;
+
+   // last selected color of the color modal
+   @observable lastSelectedPool: PoolEditorItem | undefined = undefined;
+   isDirectEdit = false;
+
+   numAgentsAssigned: Record<string, number> = {}
+   private _getNumAgentsAssigned() {
+      this.numAgentsAssigned = {};
+      agentStore.agents.forEach(agent => {
+         agent.pools?.forEach(pool => {
+            if (!(pool in this.numAgentsAssigned)) {
+               this.numAgentsAssigned[pool] = 0;
+            }
+            this.numAgentsAssigned[pool]++;
+         });
+      });
+   }
+
+   private _poolToPoolItemCopy(pool: PoolData) {
+      return {
+         key: pool.id,
+         sortKey: pool.name,
+         deleted: false,
+         pool: {
+            id: pool.id,
+            name: pool.name,
+            colorValue: pool.colorValue,
+            properties: { ...pool.properties },
+            enableAutoscaling: pool.enableAutoscaling,
+            workspaces: pool.workspaces
+         },
+         selected: false,
+         numAgentsAssigned: this.numAgentsAssigned[pool.id] ?? 0
+      };
+   }
+
+   @action
+   showConfirmation(show: boolean) {
+      this.isConfirmationOpen = show;
+   }
+
+   // sets the pool editor dialog open
+   @action
+   setOpen(isEditorOpen?: boolean, selectedItem?: PoolData) {
+      this._getNumAgentsAssigned();
+      this.modifiedPools = [...agentStore.pools.map(pool => { return this._poolToPoolItemCopy(pool); })];
+      this.newIdSuffix = 0;
+      this.isOpen = true;
+      this.isDirectEdit = isEditorOpen ?? false;
+
+      if (this.isDirectEdit && selectedItem) {
+         this.setSelectedPool(this.modifiedPools[this.modifiedPools.findIndex(item => item.key === selectedItem.id)]);
+         this.isOpen = false;
+      }
+      else {
+         this.isEditorOpen = false;
+      }
+   }
+
+   @action
+   setClose() {
+      this.isOpen = false;
+      this.isEditorOpen = false;
+   }
+
+   @action
+   private _syncLastSelectedToModified() {
+      if (this.lastSelectedPool) {
+         this.modifiedPools[this.modifiedPools.findIndex(pool => pool.key === this.lastSelectedPool!.key)] = { ...this.lastSelectedPool! };
+      }
+   }
+
+   @action
+   setEditorOpen(isOpen: boolean, cancel: boolean) {
+      this.isEditorOpen = isOpen;
+      if (isOpen === false) {
+         if (this.lastSelectedPool!.pool.id.indexOf(`zzzz_newPool_${this.newIdSuffix - 1}`) !== -1 && cancel) {
+            this.modifiedPools = [...this.modifiedPools.slice(0, this.modifiedPools.length - 1)];
+         }
+         else if (!cancel) {
+            this._syncLastSelectedToModified();
+         }
+         this.lastSelectedPool = undefined;
+      }
+   }
+
+   newPoolInputId = "zzzz_newPool_";
+   // adds a new pool to the collection
+   @action
+   addNewTempPool() {
+      const newId = this.newPoolInputId + this.newIdSuffix;
+      const newPool: PoolData = { id: newId, name: "", properties: { Color: "0" }, colorValue: "#aaaaaa", enableAutoscaling: false, workspaces: [] };
+      const newItem: PoolEditorItem = { key: newId, sortKey: newId, pool: newPool, numAgentsAssigned: 0, deleted: false };
+      this.newIdSuffix++;
+      this.modifiedPools.push(newItem);
+      this.setSelectedPool(this.modifiedPools[this.modifiedPools.findIndex(item => item.key === newItem.key)]);
+   }
+
+   // sets the selected pool
+   @action
+   setSelectedPool(item: PoolEditorItem | undefined) {
+      this.modifiedPools.forEach(existing => {
+         if (existing.pool.name !== "") {
+            existing.sortKey = existing.pool.name;
+         }
+      });
+      if (item) {
+         this.lastSelectedPool = this._poolToPoolItemCopy(item.pool);
+         this.isEditorOpen = true;
+      }
+   }
+
+   // sets the selected pool's name
+   @action
+   setSelectedName(name: string) {
+      if (this.lastSelectedPool) {
+         this.lastSelectedPool.sortKey = name;
+         this.lastSelectedPool.pool.name = name;
+      }
+   }
+
+   @action
+   setSelectedPoolColor(range: string) {
+      if (this.lastSelectedPool) {
+         this.selectedColor = linearInterpolate(range);
+         this.lastSelectedPool.pool.colorValue = this.selectedColor;
+         this.lastSelectedPool.pool.properties!["Color"] = range;
+      }
+   }
+
+   @action
+   setPoolDeleted(selectedItem: PoolEditorItem) {
+      selectedItem.deleted = true;
+   }
+
+   @action
+   setPoolValueValid(value: boolean) {
+      this.isPoolValueValid = value;
+   }
+
+   getPoolModifications(): { deletedPools: string[], newPools: string[], modifiedPools: string[] } {
+
+
+      const value = { deletedPools: [] as string[], newPools: [] as string[], modifiedPools: [] as string[] };
+
+      this.modifiedPools.forEach(item => {
+         if (item.deleted) {
+            value.deletedPools.push(item.pool.name);
+         }
+         else if (item.key.indexOf(this.newPoolInputId) !== -1) {
+            value.newPools.push(item.pool.name);
+         } else {
+            const pool = agentStore.pools.find(pool => { return pool.id === item.pool.id; })!;
+            if (pool) {
+               if (pool.name !== item.pool.name || pool.properties!["Color"] !== item.pool.properties!["Color"]) {
+                  value.modifiedPools.push(item.pool.name);
+               }
+            }
+         }
+
+      })
+
+      return value;
+   }
+
+   @action
+   async saveChanges() {
+      // save off pool changes
+      // 1. add any new pools into the list
+      // 2. update any existing ids with new names
+      // 3. remove any pools on agents with deleted flags
+      const allUpdates: any[] = [];
+      const batchUpdatePools: BatchUpdatePoolRequest[] = [];
+      const that = this;
+
+      this._syncLastSelectedToModified();
+
+      this.modifiedPools.forEach(item => {
+         if (item.deleted) {
+            allUpdates.push(backend.deletePool(item.pool.id));
+         }
+         else {
+            // create new if this has a new name prefix
+            if (item.key.indexOf(this.newPoolInputId) !== -1) {
+               allUpdates.push(backend.createPool({ name: item.pool.name, properties: { Color: item.pool.properties!["Color"] } }));
+            }
+            // otherwise update the name or deleted state if changed
+            else {
+               const pool = agentStore.pools.find(pool => { return pool.id === item.pool.id; })!;
+               if (pool.name !== item.pool.name || pool.properties!["Color"] !== item.pool.properties!["Color"]) {
+                  batchUpdatePools.push({ id: item.pool.id, name: item.pool.name, properties: { Color: item.pool.properties!["Color"] } });
+               }
+            }
+         }
+      });
+      // remove any deleted pools from agents
+      agentStore.agents.forEach(agent => {
+         const numPools = agent.pools!.length;
+         agent.pools = agent.pools!.filter(agentPool => {
+            // if the pool doesn't exist in the pools list, remove it
+            if (agentStore.pools.findIndex(item => item.id === agentPool) === -1) {
+               return false;
+            }
+            const modified = that.modifiedPools.find(modified => modified.pool.id === agentPool);
+            if (modified) {
+               if (modified.deleted) {
+                  return false;
+               }
+               return true;
+            }
+            return false;
+         });
+         // don't update the agent if nothing was removed
+         if (agent.pools!.length !== numPools) {
+            allUpdates.push(backend.updateAgent(agent.id, { pools: agent.pools }));
+         }
+      });
+      // take all the batch update pools and add them to the last promise
+      allUpdates.push(backend.batchUpdatePools(batchUpdatePools));
+      await Promise.all(allUpdates).then(function (responses) {
+      }).catch(function (errors) {
+      }).finally(function () {
+         that.setClose();
+         agentStore.update();
+      });
+   }
+}
+
+type SelectPoolItem = {
+   key: string;
+   name: string;
+   pool: PoolData;
+   agentsAssigned: string[];
+}
+
+type AgentDataSlim = {
+   Id: string;
+   Name: string;
+   Pools: string[];
+}
+
+// state class for the select pools modal
+class SelectPoolsModalState {
+
+   constructor() {
+      makeObservable(this);
+   }
+
+   @observable isOpen = false;
+   @observable.shallow availablePools: SelectPoolItem[] = [];
+   @observable.shallow selectedPools: SelectPoolItem[] = [];
+   agentChangeState: AgentDataSlim[] = [];
+   agentsSelectedTooltip = "";
+
+   @observable confirmDialogIsOpen = false;
+   confirmKey: "discard" | "no-pools" = "discard";
+   confirmDialogTitle = "Discard Changes";
+   confirmDialogText = "There are unsaved changes, discard?";
+   confirmConfirmButtonText = "Discard";
+   confirmCancelButtonText = "Cancel";
+
+   // sets the pool editor dialog open
+   @action
+   setIsOpen(isOpen: boolean) {
+      // if we're closing don't reset
+      if (isOpen) {
+         // first get all nondeleted pools
+         this.availablePools = agentStore.pools.map(pool => { return { key: pool.id, name: pool.name, pool: pool, agentsAssigned: [] }; });
+         this.agentChangeState = [];
+         this.selectedPools = [];
+         this.agentsSelectedTooltip = "";
+         this.confirmDialogIsOpen = false;
+         // then loop through each agent and split them up into selected at all and available
+         if (localState.currentSelection) {
+            (localState.currentSelection).forEach(agent => {
+               this.agentChangeState.push({ Id: agent.id, Name: agent.name, Pools: [...agent.pools!] });
+               this.agentsSelectedTooltip += agent.name + "\n";
+               agent.pools!.forEach(pool => {
+                  const index = this.availablePools.findIndex(available => available.pool.id === pool);
+                  if (index !== -1) {
+                     const item = this.availablePools.splice(index, 1)[0];
+                     item.agentsAssigned.push(agent.id);
+                     this.selectedPools.push(item);
+                  }
+                  else {
+                     this.selectedPools.find(selected => selected.pool.id === pool)?.agentsAssigned.push(agent.id);
+                  }
+               });
+            });
+         }
+         this.availablePools = this.availablePools.sort((a, b) => { return a.pool.name.localeCompare(b.pool.name); });
+         this.selectedPools = this.selectedPools.sort((a, b) => { return a.pool.name.localeCompare(b.pool.name); });
+      }
+      this.isOpen = isOpen;
+   }
+
+   @action
+   transitionPool(item: SelectPoolItem) {
+      // look for the item in available pools
+      let foundIdx = this.availablePools.findIndex(c => c.key === item.key);
+      // if not found, this is an UNASSIGN
+      if (foundIdx === -1) {
+         foundIdx = this.selectedPools.findIndex(c => c.key === item.key);
+         const itemTransfer = this.selectedPools.splice(foundIdx, 1)[0];
+         itemTransfer.agentsAssigned = [];
+         this.availablePools = [...this.availablePools, itemTransfer].sort((a, b) => { return a.pool.name.localeCompare(b.pool.name); });
+         // remove pools from existing temp agents
+         this.agentChangeState.forEach(agent => {
+            const existingIndex = agent.Pools!.indexOf(itemTransfer.pool.id);
+            // possible an agent isn't assigned to this pool, so check
+            if (existingIndex !== -1) {
+               agent.Pools!.splice(existingIndex, 1);
+            }
+         });
+         // force selectedPools to itself so we trigger a re-render (not sure why setting to not shallow and just splicing it doesn't achieve this, but here we are.)
+         this.selectedPools = [...this.selectedPools];
+      }
+      // if it is, its an ASSIGN
+      else {
+         const itemTransfer = this.availablePools.splice(foundIdx, 1)[0];
+         itemTransfer.agentsAssigned = this.agentChangeState.map(agent => { return agent.Id; });
+         this.selectedPools = [...this.selectedPools, itemTransfer];
+         // change the temp agents to contain the new assignments
+         this.agentChangeState.forEach(agent => {
+            agent.Pools!.push(itemTransfer.pool.id);
+         });
+      }
+      // always return null, the observable changing should rerender the list
+      return null;
+   }
+
+   @action
+   setConfirmOrCloseModal(type: "save" | "cancel") {
+      if (this.selectedPools.length === 0 && type === "save") {
+         this.confirmKey = "no-pools";
+         this.confirmDialogTitle = "No Pools";
+         this.confirmDialogText = "The agent(s) have no pools. Save anyway?";
+         this.confirmConfirmButtonText = "Save";
+         this.confirmCancelButtonText = "Cancel";
+         this.confirmDialogIsOpen = true;
+      }
+      else {
+         let dirty = false;
+         if (localState.currentSelection) {
+            for (const idx in localState.currentSelection) {
+               let original = agentStore.agents.find(agent => agent.id === localState.currentSelection[idx].id)?.pools!;
+               let newPools = this.agentChangeState.find(a => a.Id === localState.currentSelection[idx].id)!.Pools!;
+               // if they're not the same length we know they're different.
+               if (original.length !== newPools.length) {
+                  dirty = true;
+               }
+               else {
+                  // if they are the same length, ensure all the items are the same
+                  original = original.sort();
+                  newPools = newPools.sort();
+                  for (let jdx = 0; jdx < original.length; jdx++) {
+                     if (original[jdx] !== newPools[jdx]) {
+                        dirty = true;
+                        break;
+                     }
+                  }
+               }
+               if (dirty) {
+                  break;
+               }
+
+            }
+
+         }
+         if (dirty) {
+            if (type === "cancel") {
+               this.confirmKey = "discard";
+               this.confirmDialogTitle = "Discard Changes";
+               this.confirmDialogText = "There are unsaved changes, discard?";
+               this.confirmConfirmButtonText = "Discard";
+               this.confirmCancelButtonText = "Cancel";
+               this.confirmDialogIsOpen = true;
+            }
+            else {
+               this.saveChanges();
+            }
+         }
+         else {
+            this.confirmDialogIsOpen = false;
+            this.setIsOpen(false);
+            return;
+         }
+      }
+   }
+
+   @action
+   setConfirmModalClosed() {
+      this.confirmDialogIsOpen = false;
+   }
+
+   @action
+   async saveChanges(conform?: boolean | undefined) {
+      // save off pool changes
+      // 3. remove any pools on agents with deleted flags
+      const allUpdates: any[] = [];
+      const that = this;
+
+      this.confirmDialogIsOpen = false;
+      this.agentChangeState.forEach(agent => {
+         allUpdates.push(backend.updateAgent(agent.Id, { pools: agent.Pools, requestFullConform: conform }));
+      });
+      // take all the batch update pools and add them to the last promise
+      await Promise.all(allUpdates).then(function (responses) {
+      }).catch(function (errors) {
+      }).finally(function () {
+         that.setIsOpen(false);
+         agentStore.update();
+      });
+   }
+}
+
+const localState = new LocalState();
+const editPoolsModalState = new EditPoolsModalState();
+const selectPoolsModalState = new SelectPoolsModalState();
+
+// these are static and dont need to be re-init'd on every draw
+const agentSelectedProps: IContextualMenuProps = {
+   onItemClick: (ev?: React.MouseEvent<HTMLElement, MouseEvent> | React.KeyboardEvent<HTMLElement> | undefined, item?: IContextualMenuItem | undefined) => {
+      if (item) {
+         localState.currentSelection = [...localState.selection.getSelection() as AgentData[]];
+         if (item.key === "disable") {
+            localState.setDisableBuilderDialogOpen(true);
+         }
+         else if (item.key === "enable") {
+            localState.changeBuilderEnabled("enable");
+         }
+         else if (item.key === "changePools") {
+            selectPoolsModalState.setIsOpen(true);
+         }
+         else if (item.key === "requestConform") {
+            localState.requestBuilderUpdate(true, false);
+         }
+         else if (item.key === "requestFullConform") {
+            localState.requestBuilderUpdate(false, false, true);
+         }
+         else if (item.key === "clearConform") {
+            localState.requestBuilderUpdate(false, false, false);
+         }            
+         else if (item.key === "restart") {
+            localState.setRestartBuilderDialogOpen(true);
+         }
+         else if (item.key === "delete") {
+            localState.setDeleteBuilderDialogOpen(true);
+         }
+         else if (item.key === "editcomment") {
+            localState.setEditCommentDialogOpen(true);
+         }
+
+      }
+   },
+   items: [
+      {
+         key: 'enable',
+         text: 'Enable',
+      },
+      {
+         key: 'disable',
+         text: 'Disable',
+      },
+      {
+         key: 'audit',
+         text: 'Audit',
+      },
+      {
+         key: 'editcomment',
+         text: 'Edit Comment',
+      },
+      {
+         key: 'remotedesktop',
+         text: 'Remote Desktop',
+      },
+      {
+         key: 'changePools',
+         text: 'Edit Pools'
+      },
+      {
+         key: 'cancelleases',
+         text: 'Cancel Leases',
+      },
+      {
+         key: 'clearConform',
+         text: 'Clear Pending Conform'
+      },
+      {
+         key: 'requestConform',
+         text: 'Request Conform'
+      },
+      {
+         key: 'requestFullConform',
+         text: 'Request Full Conform'
+      },
+      {
+         key: 'restart',
+         text: 'Request Restart'
+      },
+      {
+         key: 'shutdown',
+         text: 'Request Shutdown'
+      },
+      {
+         key: 'delete',
+         text: 'Delete'
+      }
+   ]
+};
+
+// these are static and dont need to be re-init'd on every draw
+const agentContextMenuProps: IContextualMenuItem[] = [
+   {
+      key: 'enable',
+      text: 'Enable',
+      onClick: () => localState.changeBuilderEnabled('enable')
+   },
+   {
+      key: 'disable',
+      text: 'Disable',
+      onClick: () => localState.setDisableBuilderDialogOpen(true)
+   },
+   {
+      key: 'remotedesktop',
+      text: 'Remote Desktop',
+      onClick: () => {
+         if (localState.currentSelection?.length && localState.currentSelection[0].id) {
+            // @todo: this assumes the id is the ip, should have a private ip property on agents instead
+            let ip = localState.currentSelection[0].id;
+            if (ip.startsWith("10-")) {
+               ip = ip.replaceAll("-", ".")
+            }
+            window.open(`ugs://rdp?host=${ip}`, "_self")
+         }
+      }
+   },
+   {
+      key: 'editcomment',
+      text: 'Edit Comment',
+      onClick: () => localState.setEditCommentDialogOpen(true)
+   },
+   {
+      key: 'audit',
+      text: 'Audit',
+      onClick: () => {
+         if (localState.currentSelection?.length && localState.currentSelection[0].id) {
+
+            const href = `/audit/agent/${encodeURIComponent(localState.currentSelection[0].id)}`;
+            window.open(href, "_blank")
+         }
+      }
+   },
+   {
+      key: 'editPools',
+      text: 'Edit Pools',
+      onClick: () => selectPoolsModalState.setIsOpen(true)
+   },
+   {
+      key: 'cancelleases',
+      text: 'Cancel Leases',
+      onClick: () => localState.setCancelLeasesDialogOpen(true)
+   },
+   {
+      key: 'clearConform',
+      text: 'Clear Pending Conform',
+      onClick: () => localState.requestBuilderUpdate(false, false, false)
+   },
+   {
+      key: 'requestConform',
+      text: 'Request Conform',
+      onClick: () => localState.requestBuilderUpdate(true, false)
+   },
+   {
+      key: 'requestFullConform',
+      text: 'Request Full Conform',
+      onClick: () => localState.requestBuilderUpdate(false, false, true)
+   },
+   {
+      key: 'requestRestart',
+      text: 'Request Restart',
+      onClick: () => localState.setRestartBuilderDialogOpen(true)
+   },
+   {
+      key: 'requestShutdown',
+      text: 'Request Shutdown',
+      onClick: () => localState.setShutdownBuilderDialogOpen(true)
+   },
+   {
+      key: 'delete',
+      text: 'Delete',
+      onClick: () => localState.setDeleteBuilderDialogOpen(true)
+   },
+];
+
+const agentStatus = ["Active", "Ready", "Disabled", "Ephemeral", "No Pools", "Pending Conform", "Pending Shutdown", "Offline", "Offline (Autoscaler)", "Offline (Manual)", "Offline (Unexpected)"];
+
+
+export const AgentMenuBar: React.FC<{ agentView?: boolean }> = observer(({ agentView }) => {
+
+   let selectedButton: any = <div></div>;
+   if (localState.agentsSelectedCount !== 0) {
+      selectedButton = <PrimaryButton
+         disabled={localState.agentsSelectedCount === 0 ? true : false}
+         styles={{ root: { fontFamily: 'Horde Open Sans SemiBold !important' } }}
+         menuProps={agentSelectedProps}
+         onClick={() => { }}>
+         {`${localState.agentsSelectedCount} Agent${localState.agentsSelectedCount > 1 ? "s" : ""} Selected`}
+      </PrimaryButton>;
+   }
+
+   const agentStatusItems = agentStatus.map(status => {
+      return {
+         text: status,
+         key: `dropdown_agent_status_${status}_key`,
+         status: status
+      }
+   });
+
+   return (
+      <Stack horizontal verticalAlign='center' tokens={{ childrenGap: 24 }}>
+         {!!agentView && <Stack horizontal tokens={{ childrenGap: 12 }} verticalAlign='center'>
+            {!!dashboard.user?.dashboardFeatures?.showPoolEditor && <CommandButton
+               onClick={() => { editPoolsModalState.setOpen(); }}
+               iconProps={{ iconName: 'Edit' }}
+               styles={{ root: { fontFamily: 'Horde Open Sans SemiBold !important' } }}>
+               {`Pools`}
+            </CommandButton>}
+            {selectedButton}
+            <PoolEditorModal></PoolEditorModal>
+            <PoolSelectionModal></PoolSelectionModal>
+         </Stack>
+         }
+         <Stack>
+            <Dropdown
+               placeholder="Filter Status"
+               style={{ width: 200 }}
+               selectedKeys={Array.from(localState.agentStatusFilter).map(status => `dropdown_agent_status_${status}_key`)}
+               multiSelect
+               options={agentStatusItems}
+               onChange={(event, option, index) => {
+
+                  if (option) {
+
+                     // new set instance to update observable action update
+                     const newFilter = new Set(localState.agentStatusFilter);
+
+                     if (option.selected) {
+                        newFilter.add((option as any).status);
+                     } else {
+                        newFilter.delete((option as any).status);
+                     }
+
+                     localState.setAgentStaus(newFilter);
+
+                  }
+               }}
+            />
+         </Stack>
+         <Stack>
+            <Checkbox label={"Exact Match"} checked={localState.filterExactMatch} onChange={(ev, checked) => localState.setExactMatch(checked!)} />
+         </Stack>
+         <Stack>
+            <SearchBox
+               showIcon={true}
+               disableAnimation={true}
+               placeholder="Search Agents"
+               value={localState.agentFilter}
+               styles={{ root: { width: 240 } }}
+               onChange={(event?: React.ChangeEvent<HTMLInputElement> | undefined, newValue?: string | undefined) => { localState.setAgentFilter(newValue ?? ""); }}
+               onClear={() => { localState.setAgentFilter(""); }}
+            />
+         </Stack>
+
+      </Stack>
+   );
+});
+
+
+export const AgentPivot: React.FC = () => {
+
+   const { hordeClasses, modeColors } = getHordeStyling();
+
+   const categories = dashboard.agentCategories;
+
+   const pivotItems = categories.map(tab => {
+      return <PivotItem headerText={tab.name} itemKey={tab.name} key={tab.name} style={{ color: modeColors.text }} />;
+   });
+
+   pivotItems.unshift(<PivotItem headerText="All" itemKey="all" key={"all"} style={{ color: modeColors.text }} />);
+
+   return <Stack grow>
+      <Pivot className={hordeClasses.pivot}
+         overflowBehavior='menu'
+         selectedKey={localState.searchState.category ?? "all"}
+         linkSize="normal"
+         linkFormat="links"
+         onLinkClick={(item) => {
+            if (item) {
+               localState.setAgentCategory(item.props.itemKey === "all" ? undefined : item.props.itemKey);
+            }
+         }}>
+         {pivotItems}
+      </Pivot>
+   </Stack>
+
+}
+
+const PoolEditorConfirmation: React.FC = observer(() => {
+
+
+   if (!editPoolsModalState.isConfirmationOpen) {
+      return null;
+   }
+
+   const mods = editPoolsModalState.getPoolModifications();
+
+   if (mods.deletedPools.length === 0 && mods.newPools.length === 0 && mods.modifiedPools.length === 0) {
+      return null;
+   }
+
+   let title = "Modify Pool" + (mods.modifiedPools.length > 1 ? "s" : "");
+   let subText = "";
+
+   if (mods.deletedPools.length && !mods.newPools.length) {
+      title = "Delete Pool";
+      if (mods.deletedPools.length > 1) {
+         title += "s";
+      }
+   } else if (!mods.deletedPools.length && mods.newPools.length) {
+      title = "Create Pool";
+      if (mods.newPools.length > 1) {
+         title += "s";
+      }
+   }
+
+   title += "?";
+
+   if (mods.deletedPools.length) {
+      subText += "Delete: " + mods.deletedPools.join(", ") + " ";
+   }
+
+   if (mods.newPools.length) {
+      subText += "Create: " + mods.newPools.join(", ") + " ";
+   }
+
+   if (mods.modifiedPools.length) {
+      subText += "Modified: " + mods.modifiedPools.join(", ");
+   }
+
+   return <ConfirmationDialog
+      title={title}
+      subText={subText}
+      isOpen={true}
+      confirmText="Save"
+      cancelText="Cancel"
+      onConfirm={() => { editPoolsModalState.showConfirmation(false); editPoolsModalState.saveChanges(); }}
+      onCancel={() => editPoolsModalState.showConfirmation(false)}
+   />
+
+});
+
+export const PoolEditorModal: React.FC = observer(() => {
+
+   const [state, setState] = useState<{ sortBy: "Agents" | "Name" }>({ sortBy: "Name" });
+   const { hordeClasses } = getHordeStyling();
+   const agentStyles = getAgentStyles();
+
+   let selectedColor = "#ffffff";
+   let selectedColorValue = "0";
+   if (editPoolsModalState.lastSelectedPool) {
+      selectedColorValue = editPoolsModalState.lastSelectedPool.pool.properties!["Color"];
+      selectedColor = linearInterpolate(selectedColorValue);
+   }
+
+   function onRenderPoolModalItem(item: PoolEditorItem, index?: number, column?: IColumn) {
+      switch (column!.key) {
+         case 'name':
+            const color = item.pool.colorValue;
+            const textColor = "white";
+            return (
+               <Stack styles={{ root: { height: '100%', } }} horizontal>
+                  <Stack.Item align={"center"}>
+                     <PrimaryButton key={"pooleditormodal_item_" + item.pool.id}
+                        className={agentStyles.buttonFont}
+                        onClick={() => {
+                           editPoolsModalState.setSelectedPool(item);
+                        }}
+                        styles={{
+                           root: { borderColor: color, backgroundColor: color, color: textColor },
+                           rootHovered: { borderColor: color, backgroundColor: color, color: textColor, },
+                           rootPressed: { borderColor: color, backgroundColor: color, color: textColor, }
+                        }}>{item.pool.name}
+                     </PrimaryButton>
+                  </Stack.Item>
+               </Stack>
+            );
+         case 'numMachines':
+            return <Stack styles={{ root: { height: '100%', } }} horizontal horizontalAlign={'center'}><Stack.Item align={"center"}>{item.numAgentsAssigned}</Stack.Item></Stack>;
+         case 'del':
+            return <Stack styles={{ root: { height: '100%', } }} horizontal horizontalAlign={"center"}>
+               <Stack.Item grow={0} align={"center"}>
+                  <IconButton iconProps={{ iconName: 'Cancel' }} onClick={(event: any) => { editPoolsModalState.setPoolDeleted(item); }} />
+               </Stack.Item>
+            </Stack>;
+         default:
+            return <div></div>;
+      }
+   }
+
+   // main header
+   const onRenderDetailsHeader: IDetailsListProps['onRenderDetailsHeader'] = (props) => {
+      const customStyles: Partial<IDetailsHeaderStyles> = {};
+      if (props) {
+         return (
+            <Sticky stickyPosition={StickyPositionType.Header} isScrollSynced={true}>
+               <DetailsHeader {...props} styles={customStyles} onColumnClick={(ev: React.MouseEvent<HTMLElement>, column: IColumn) => {
+                  if (column.name === "Agents") {
+                     if (state.sortBy !== 'Agents') {
+                        setState({ ...state, sortBy: "Agents" });
+                     }
+
+                  } else if (column.name === "Name") {
+                     if (state.sortBy !== 'Name') {
+                        setState({ ...state, sortBy: "Name" });
+                     }
+                  }
+               }} onRenderColumnHeaderTooltip={onRenderColumnHeaderTooltip} />
+            </Sticky>
+         );
+      }
+      return null;
+   };
+
+   const onRenderColumnHeaderTooltip: IDetailsHeaderProps['onRenderColumnHeaderTooltip'] = (props: any) => {
+      const customStyles: Partial<ITooltipHostStyles> = {};
+      if (props) {
+         customStyles.root = { selectors: { "> span": { paddingLeft: '8px !important', display: 'flex', justifyContent: 'center' } } };
+         /*
+          if (props.children._owner.child?.child.child.child.elementType === "span") {
+              const data = props.children._owner.child.child.child.child.child.child.stateNode.data;
+              if (data === "Name") {
+                  customStyles.root = {};
+              }
+          }
+          */
+         return <Text styles={customStyles}>{props.children}</Text>;
+      }
+      return null;
+   };
+
+   const pools = editPoolsModalState.modifiedPools.filter(item => item.deleted === false).sort((a, b) => {
+      if (state.sortBy === "Name")
+         return a.sortKey.localeCompare(b.sortKey);
+      else
+         return b.numAgentsAssigned - a.numAgentsAssigned;
+   });
+
+   // 75vh for max
+   const viewportHeight = document.documentElement.clientHeight * .75;
+   const modalHeight = Math.min(viewportHeight, 80 + (44 * pools.length));
+
+   return (
+      <Stack>
+         <PoolEditorConfirmation />
+         <Dialog
+            modalProps={{
+               isBlocking: false,
+               dragOptions: {
+                  closeMenuItemText: "Close",
+                  moveMenuItemText: "Move",
+                  menu: ContextualMenu
+               }
+            }}
+            onDismiss={() => { editPoolsModalState.setClose(); }}
+            className={agentStyles.dialog}
+            minWidth={440}
+            hidden={!editPoolsModalState.isOpen}
+            dialogContentProps={{
+               type: DialogType.close,
+               onDismiss: () => { editPoolsModalState.setClose(); },
+               title: "Pool Editor",
+            }}
+         >
+            <Stack>
+               <Stack horizontal horizontalAlign="space-evenly">
+                  <Stack.Item className={hordeClasses.relativeWrapper} styles={{ root: { height: modalHeight, width: 400 } }}>
+                     <ScrollablePane scrollbarVisibility={ScrollbarVisibility.auto}>
+                        <DetailsList
+                           className={agentStyles.detailsList}
+                           compact={true}
+                           items={pools}
+                           columns={[{ key: 'name', name: 'Name', minWidth: 200, maxWidth: 500 }, { key: 'numMachines', name: 'Agents', minWidth: 60, maxWidth: 60 }, { key: 'del', name: '', minWidth: 55, maxWidth: 55 }]}
+                           onRenderDetailsHeader={onRenderDetailsHeader}
+                           onRenderItemColumn={onRenderPoolModalItem}
+                           layoutMode={DetailsListLayoutMode.justified}
+                           constrainMode={ConstrainMode.horizontalConstrained}
+                           selectionMode={SelectionMode.none}
+                        />
+                     </ScrollablePane>
+                  </Stack.Item>
+               </Stack>
+               <Stack styles={{ root: { paddingTop: 10 } }}>
+                  <Stack horizontal>
+                     <Stack>
+                        <DefaultButton onClick={() => editPoolsModalState.addNewTempPool()}>New Pool</DefaultButton>
+                     </Stack>
+                     <Stack grow />
+                     <Stack>
+                        <PrimaryButton disabled={editPoolsModalState.getPoolModifications().deletedPools.length === 0 && editPoolsModalState.getPoolModifications().newPools.length === 0 && editPoolsModalState.getPoolModifications().modifiedPools.length === 0} onClick={() => editPoolsModalState.showConfirmation(true)} styles={{ root: { marginRight: "10px" } }}>Save</PrimaryButton>
+                     </Stack>
+                     <Stack>
+                        <DefaultButton onClick={() => editPoolsModalState.setClose()} styles={{ root: { marginRight: "10px" } }}>Cancel</DefaultButton>
+                     </Stack>
+                  </Stack>
+               </Stack>
+            </Stack>
+         </Dialog>
+
+         <Dialog
+            modalProps={{
+               isBlocking: false,
+               dragOptions: {
+                  closeMenuItemText: "Close",
+                  moveMenuItemText: "Move",
+                  menu: ContextualMenu
+               }
+            }}
+            onDismiss={() => { editPoolsModalState.setEditorOpen(false, true); }}
+            className={agentStyles.dialog}
+            minWidth={440}
+            hidden={!editPoolsModalState.isEditorOpen}
+            dialogContentProps={{
+               type: DialogType.close,
+               onDismiss: () => { editPoolsModalState.setEditorOpen(false, true); },
+               title: "Edit Pool",
+            }}
+         >
+            <Stack>
+               <Stack horizontal horizontalAlign="space-evenly">
+                  <Stack styles={{ root: { width: 400 } }}>
+                     <TextField
+                        spellCheck={false}
+                        label="Name"
+                        borderless={false}
+                        autoFocus={true}
+                        styles={{
+                           root: { paddingLeft: 10, paddingRight: 10 },
+                        }}
+                        onChange={(event: React.FormEvent<HTMLInputElement | HTMLTextAreaElement>, newValue?: string | undefined) => {
+                           editPoolsModalState.setSelectedName(newValue!);
+                        }}
+                        placeholder={`Enter pool name...`}
+                        defaultValue={editPoolsModalState.lastSelectedPool?.pool.name}
+                        onGetErrorMessage={(value: string) => {
+                           let editingExistingPool = agentStore.pools.find(pool => pool.name === value);
+                           if (!editingExistingPool) {
+                              editingExistingPool = editPoolsModalState.modifiedPools.find(pool => pool.pool.name === value)?.pool;
+                           }
+                           if (value.length < 1 || value.length > 64) {
+                              editPoolsModalState.setPoolValueValid(false);
+                              return "Length must be between 1 and 64 characters.";
+                           }
+                           if (editingExistingPool && editingExistingPool.id !== editPoolsModalState.lastSelectedPool?.pool.id) {
+                              editPoolsModalState.setPoolValueValid(false);
+                              return "Name is already taken.";
+                           }
+                           editPoolsModalState.setPoolValueValid(true);
+                           return "";
+                        }}
+                     />
+                     <Stack horizontal styles={{ root: { marginBottom: 40 } }}>
+                        <Stack grow>
+                           <Slider
+                              styles={{ root: { paddingTop: 15, paddingLeft: 10, paddingRight: 10 } }}
+                              label="Color"
+                              onChange={(value: number) => editPoolsModalState.setSelectedPoolColor(value.toString())}
+                              defaultValue={parseInt(selectedColorValue)}
+                              max={599}
+                              showValue={false}
+                           />
+                        </Stack>
+                        <Stack>
+                           <div className={hordeClasses.colorPreview} style={{ backgroundColor: selectedColor }} />
+                        </Stack>
+                     </Stack>
+                  </Stack>
+               </Stack>
+               <Stack styles={{ root: { paddingTop: 10 } }}>
+                  <Stack horizontal>
+                     <Stack style={{ paddingLeft: 12 }}>
+                        <DefaultButton href={`/pools?pool=${editPoolsModalState.lastSelectedPool?.pool?.id}`} target="_blank">Details</DefaultButton>
+                     </Stack>
+                     <Stack grow />
+                     <PrimaryButton disabled={!editPoolsModalState.isPoolValueValid} onClick={() => { editPoolsModalState.isDirectEdit ? editPoolsModalState.saveChanges() : editPoolsModalState.setEditorOpen(false, false); }} styles={{ root: { marginRight: "10px" } }}>
+                        {editPoolsModalState.isDirectEdit ? "Update" : "Save"}
+                     </PrimaryButton>
+                     <DefaultButton onClick={() => { editPoolsModalState.isDirectEdit ? editPoolsModalState.setClose() : editPoolsModalState.setEditorOpen(false, true); }} styles={{ root: { marginRight: "10px" } }}>Cancel</DefaultButton>
+                  </Stack>
+               </Stack>
+            </Stack>
+         </Dialog>
+      </Stack>
+   );
+});
+
+export const PoolSelectionModal: React.FC = observer(() => {
+
+   const agentStyles = getAgentStyles();
+
+   function onResolveSuggestions(filter: string, selectedItems?: ITag[] | undefined) {
+      return selectPoolsModalState.availablePools.filter(item => item.name.toLowerCase().indexOf(filter.toLowerCase()) !== -1);
+   }
+
+   // renderitem override
+   const onTagRenderItem: IBasePickerProps<ITag>['onRenderItem'] = (props) => {
+      const item = props.item as SelectPoolItem;
+      let color = item.pool.colorValue;
+      const hoverColor = hexToRGB(color);
+      const hoverCloseColor = hexToRGB(color);
+      hoverColor.r *= .8;
+      hoverColor.g *= .8;
+      hoverColor.b *= .8;
+      hoverCloseColor.r *= .7;
+      hoverCloseColor.g *= .7;
+      hoverCloseColor.b *= .7;
+      const textColor = "white";
+      let name = item.pool.name;
+      let title = "";
+      if (item.agentsAssigned.length !== localState.agentsSelectedCount && item.agentsAssigned.length !== 0) {
+         color += "80";
+         name += ` (${item.agentsAssigned.length}/${localState.agentsSelectedCount})`;
+         title = "Assigned to: \n\t" + item.agentsAssigned.map(agent => { return selectPoolsModalState.agentChangeState.find(found => found.Id === agent)!.Name; }).join('\n\t');
+      }
+
+      const customStyles: Partial<ITagItemStyles> = {};
+      customStyles.root = {
+         borderColor: color, backgroundColor: color, textAlign: 'center',
+         selectors: {
+            ':hover': { backgroundColor: `rgb(${hoverColor.r},${hoverColor.g},${hoverColor.b})` },
+            ':hover .ms-TagItem-close': { color: 'white' },
+            ':active': { color: 'white' }
+         }
+      };
+      customStyles.text = { font: '8pt Horde Open Sans SemiBold !important', color: textColor, position: 'relative', top: 6, paddingLeft: 16, paddingRight: 4, marginRight: 0 };
+      customStyles.close = { color: 'white', selectors: { ':hover': { color: 'white', background: `rgb(${hoverCloseColor.r},${hoverCloseColor.g},${hoverCloseColor.b})` } } };
+      return <TagItem
+         key={`tagItem_pools_${item.pool.id}`}
+         title={title}
+         index={props.index}
+         item={props.item}
+         styles={customStyles}
+         onRemoveItem={() => { selectPoolsModalState.transitionPool(item); }}>
+         {name}
+      </TagItem>;
+   };
+
+   const [conform, setConform] = useState<boolean | undefined>(false);
+
+   return (
+      <Dialog
+         hidden={!selectPoolsModalState.isOpen}
+         dialogContentProps={{
+            type: DialogType.close,
+            onDismiss: () => { selectPoolsModalState.setConfirmOrCloseModal("cancel"); },
+            title: 'Pool Assignments',
+         }}
+         className={agentStyles.dialog}
+         minWidth={800}
+         onDismiss={() => { selectPoolsModalState.setConfirmOrCloseModal("cancel"); }}
+         modalProps={{
+            isBlocking: false,
+            dragOptions: {
+               closeMenuItemText: "Close",
+               moveMenuItemText: "Move",
+               menu: ContextualMenu
+            }
+         }}
+      >
+         <TagPicker
+            onResolveSuggestions={onResolveSuggestions}
+            onEmptyResolveSuggestions={() => { return selectPoolsModalState.availablePools; }}
+            onItemSelected={(selectedItem: ITag | undefined) => { return selectPoolsModalState.transitionPool(selectedItem as SelectPoolItem); }}
+            onRenderItem={onTagRenderItem}
+            selectedItems={selectPoolsModalState.selectedPools}
+            pickerSuggestionsProps={{
+               noResultsFoundText: 'No Pools Found'
+            }}
+            inputProps={{
+               'aria-label': 'Tag Picker'
+            }}
+         />
+         <Stack horizontal styles={{ root: { marginTop: 16 } }}>
+            <Stack>
+               <Checkbox label="Conform Immediately" onChange={(ev?: any, checked?: boolean | undefined) => setConform(checked)} />
+            </Stack>
+            <Stack grow />
+            <Stack horizontal>
+               <PrimaryButton styles={{ root: { marginRight: 8 } }} onClick={() => selectPoolsModalState.setConfirmOrCloseModal("save")} >Update</PrimaryButton>
+               <DefaultButton onClick={() => selectPoolsModalState.setConfirmOrCloseModal("cancel")}>Cancel</DefaultButton>
+            </Stack>
+         </Stack>
+         <ConfirmationDialog
+            title={selectPoolsModalState.confirmDialogTitle}
+            subText={selectPoolsModalState.confirmDialogText}
+            isOpen={selectPoolsModalState.confirmDialogIsOpen}
+            confirmText={selectPoolsModalState.confirmConfirmButtonText}
+            cancelText={selectPoolsModalState.confirmCancelButtonText}
+            onConfirm={() => (selectPoolsModalState.confirmKey === "no-pools" ? selectPoolsModalState.saveChanges(conform ?? undefined) : selectPoolsModalState.setIsOpen(false))}
+            onCancel={() => selectPoolsModalState.setConfirmModalClosed()}
+         />
+      </Dialog>
+   );
+});
+
+export const SearchUpdate: React.FC = observer(() => {
+
+   const [searchParams, setSearchParams] = useSearchParams();
+   const [state, setState] = useState<{ search?: string }>({ search: searchParams.toString() });
+
+   // subscribe
+   if (localState.searchUpdated) { }
+
+   const csearch = localState.search.toString();
+
+   if (state.search !== csearch) {
+      setSearchParams(csearch, { replace: true });
+      setState({ search: csearch });
+   }
+
+   return null;
+});
+
+export const AgentViewInner: React.FC<{ agentId?: string, poolId?: string, searchParams?: URLSearchParams, agentView?: boolean }> = observer(({ agentId, poolId, searchParams, agentView }) => {
+
+   poolId = poolId?.toLowerCase();
+
+   const [initAgentUpdater, setInitAgentUpdater] = useState(false);
+
+   useEffect(() => {
+      const interval = setInterval(() => {
+         agentStore.update(true)
+      }, 15000);
+      return () => clearInterval(interval);
+   }, []);
+
+   // subscribe
+   if (localState.searchUpdated) { }
+
+   const agentStyles = getAgentStyles();
+
+   if (!initAgentUpdater) {
+      agentStore.update().then(() => {
+         localState.resetState();
+         if (searchParams) {
+            localState.stateFromSearch(searchParams);
+         }
+         setInitAgentUpdater(true);
+      });
+      return <Stack>
+         <Stack style={{ position: "relative", height: "calc(100vh - 240px)" }} horizontalAlign="center">
+            <Stack horizontal tokens={{ childrenGap: 24 }}>
+               <Text variant="mediumPlus">Loading Agents</Text>
+               <Spinner size={SpinnerSize.large} />
+            </Stack>
+         </Stack>
+      </Stack>
+   }
+
+   let activeAgent: AgentData | undefined = undefined;
+   if (agentId) {
+      activeAgent = agentStore.agents.find(agent => agent.id === agentId);
+   }
+
+   // if(activeAgent) {
+   //     if(historyModalState.selectedAgent?.id !== activeAgent.id){
+   //         historyModalState.setSelectedAgent(activeAgent.id);
+   //         historyModalState.setIsOpen(true, activeAgent);
+   //     }
+   // }
+   // else {
+   //     historyModalState.setIsOpen(false);
+   // }
+
+   // main header
+   const onRenderDetailsHeader: IDetailsListProps['onRenderDetailsHeader'] = (props) => {
+      if (props) {
+         return (
+            <Sticky stickyPosition={StickyPositionType.Header} isScrollSynced={true}>
+               <DetailsHeader {...props} onRenderColumnHeaderTooltip={onRenderColumnHeaderTooltip} />
+            </Sticky>
+         );
+      }
+      return null;
+   };
+
+   // everything about this is no.
+   const onRenderColumnHeaderTooltip: IDetailsHeaderProps['onRenderColumnHeaderTooltip'] = (props: any) => {
+      const customStyles: Partial<ITooltipHostStyles> = {};
+      if (props) {
+         customStyles.root = { selectors: { "> span": { paddingLeft: '8px !important', display: 'flex', justifyContent: 'center' } } };
+
+         /*
+          // no other way to filter children being centered other than to drill into private members??
+          if (props.children._owner.child?.child.child.child.elementType === "span") {
+              const data = props.children._owner.child.child.child.child.child.child.stateNode.data;
+              // if this cluster happens to be true, reset back to the default, because doing this the other way around
+              // takes too long when the columns update on details switch
+              // ugh :(
+              if (data === "Name" || data === "Pools") {
+                  customStyles.root = {};
+              }
+          }
+          */
+         return <Text styles={customStyles}>{props.children}</Text>;
+      }
+      return null;
+   };
+
+   function filterAgents(this: string, agent: AgentData) {
+      const filter = this.toLowerCase();
+      if (agent.deleted) {
+         return null;
+      }
+
+      let category: GetDashboardAgentCategoryResponse | undefined;
+
+      if (localState.searchState.category) {
+         category = dashboard.agentCategories.find(c => c.name === localState.searchState.category);
+      }
+
+      if (category && category.condition) {
+         const cat = agentStore.getCategory(category.condition);
+         if (!cat.ids.length && cat.polling) {
+            return null;
+         } else {
+            if (!cat.ids.find((id) => id === agent.id)) {
+               return null;
+            }
+         }
+      }
+
+      if (filter !== "") {
+         if (localState.filterExactMatch) {
+            // add if there's a match to name or version
+            if (agent.name.toLowerCase() === filter) {
+               return agent;
+            }
+
+            // check status
+            if (localState.columnSearchState['status'][agent.id]) {
+               for (const idx in localState.columnSearchState['status'][agent.id]) {
+                  if (localState.columnSearchState['status'][agent.id][idx].toLowerCase() === filter) {
+                     return agent;
+                  }
+               }
+            }
+
+            // also add if there's a pool match
+            if (localState.columnSearchState['pools'][agent.id]) {
+               for (const idx in localState.columnSearchState['pools'][agent.id]) {
+                  if (localState.columnSearchState['pools'][agent.id][idx].toLowerCase() === filter) {
+                     return agent;
+                  }
+               }
+            }
+            // revert back to checking for individual pools if there's no cached version
+            else {
+               if (agent.pools) {
+                  for (const idx in agent.pools) {
+                     const pool = agentStore.pools.find(p => { return p.id === agent.pools![idx]; });
+                     if (pool && pool.name.toLowerCase() === filter) {
+                        return agent;
+                     }
+                  }
+               }
+            }
+         }
+         else {
+            // add if there's a match to name or version
+            if ((agent.name.toLowerCase().indexOf(filter) !== -1)) {
+               return agent;
+            }
+
+            // check status
+            if (localState.columnSearchState['status'][agent.id]) {
+               if (localState.columnSearchState['status'][agent.id].join('').toLowerCase().indexOf(filter) !== -1) {
+                  return agent;
+               }
+            }
+
+            // also add if there's a pool match
+            if (localState.columnSearchState['pools'][agent.id]) {
+               if (localState.columnSearchState['pools'][agent.id].join('').toLowerCase().indexOf(filter) !== -1) {
+                  return agent;
+               }
+            }
+            // revert back to checking for individual pools if there's no cached version
+            else {
+               if (agent.pools) {
+                  for (const idx in agent.pools) {
+                     const pool = agentStore.pools.find(p => { return p.id === agent.pools![idx]; });
+                     if (pool) {
+                        if (pool.name.toLowerCase().indexOf(filter) !== -1) {
+                           return agent;
+                        }
+                     } else {
+                        console.error(`No matching pool id for ${JSON.stringify(agent.pools![idx])} while filtering pools`);
+                     }
+                  }
+               }
+            }
+         }
+
+         let toCheck = "";
+         // also check currently selected othercolumn
+         switch (localState.columnMenuProps.find(prop => prop.checked)!.key) {
+            case 'software':
+               if (agent.version) {
+                  toCheck = agent.version;
+               }
+               else {
+                  toCheck = "Unknown";
+               }
+               break;
+            case 'taskTime':
+               if (agent.version) {
+                  if (localState.columnSearchState['taskTimeDescription'][agent.id]) {
+                     toCheck = localState.columnSearchState['taskTimeDescription'][agent.id];
+                  }
+               }
+               else {
+                  toCheck = "Unknown";
+               }
+               break;
+            case 'storage':
+               if (localState.columnSearchState['storage'][agent.id]) {
+                  toCheck = localState.columnSearchState['storage'][agent.id];
+               }
+               else {
+                  let diskFree: string | number | null = getAgentCapability(agent, "DiskFreeSpace");
+                  let diskCapacity: string | number | null = getAgentCapability(agent, "DiskTotalSize");
+                  toCheck = "Unknown";
+                  if (diskFree !== null && diskCapacity !== null) {
+                     diskFree = Number(diskFree) / 1073741824;
+                     diskCapacity = Number(diskCapacity) / 1073741824;
+                     const percentage = (diskCapacity - diskFree) / diskCapacity;
+                     toCheck = (diskCapacity - (percentage * diskCapacity)).toFixed(0);
+                  }
+               }
+               break;
+            case 'comment':
+               if (localState.columnSearchState['comment'][agent.id]) {
+                  toCheck = localState.columnSearchState['comment'][agent.id];
+               }
+               else {
+                  toCheck = "Unknown";
+               }
+               break;
+            case 'OS':
+               toCheck = getAgentOs(agent) ?? "Unknown";
+               break;
+            case 'CPU':
+               const CPU = getAgentCapability(agent, "CPU");
+               toCheck = CPU ?? "Unknown";
+               break;
+            case 'RAM':
+               const RAM = getAgentCapability(agent, "RAM");
+               toCheck = RAM ?? "Unknown";
+               break;
+            default:
+               toCheck = "Unknown";
+               return null;
+         }
+         if (localState.filterExactMatch) {
+            if (toCheck.toString().toLowerCase() === filter) {
+               return agent;
+            }
+         }
+         else {
+            if (toCheck.toString().toLowerCase().indexOf(filter) !== -1) {
+               return agent;
+            }
+         }
+      }
+      else {
+         return agent;
+      }
+   }
+
+   function sortAgents(a: AgentData, b: AgentData) {
+      // sort differently depending on which column this is
+      const sortedCol = localState.columnsState.find(c => c.isSorted === true)!;
+      let left = a;
+      let right = b;
+      if (sortedCol.isSortedDescending) {
+         left = b;
+         right = a;
+      }
+      switch (sortedCol.key) {
+         case 'name':
+            return left.name.toLocaleLowerCase().localeCompare(right.name, undefined, { numeric: true, sensitivity: 'base' });
+         case 'pools':
+            if (localState.columnSearchState['pools'][left.id] && localState.columnSearchState['pools'][right.id]) {
+               return localState.columnSearchState['pools'][left.id].join('').localeCompare(localState.columnSearchState['pools'][right.id].join(''));
+            }
+            return left.name.localeCompare(right.name);
+         case 'status':
+            if (localState.columnSearchState['status'][left.id] && localState.columnSearchState['status'][right.id]) {
+               return localState.columnSearchState['status'][left.id].join('').localeCompare(localState.columnSearchState['status'][right.id].join(''));
+            }
+            return left.name.localeCompare(right.name);
+         case 'software':
+            return (left.version ?? "").toLocaleLowerCase().localeCompare(right.version ?? "");
+         case 'taskTime':
+            return getTaskTime(left) - getTaskTime(right);
+         case 'storage':
+            if (localState.columnSearchState['storage'][left.id] && localState.columnSearchState['storage'][right.id]) {
+               return localState.columnSearchState['storage'][left.id] - localState.columnSearchState['storage'][right.id];
+            }
+            return Number(getAgentCapability(left, "DiskFreeSpace")) - Number(getAgentCapability(right, "DiskFreeSpace"));
+         case 'comment':
+            if (localState.columnSearchState['comment'][left.id] && localState.columnSearchState['comment'][right.id]) {
+               return localState.columnSearchState['comment'][left.id].localeCompare(localState.columnSearchState['comment'][right.id]);
+            }
+            return (left.comment ?? "").toLocaleLowerCase().localeCompare(right.comment ?? "");
+         case 'systemInfoOS':
+            return ((getAgentOs(left) ?? "Unknown").toLocaleLowerCase().localeCompare((getAgentOs(right) ?? "Unknown")));
+         case 'systemInfoCPU':
+            return (getAgentCapability(left, "CPU") ?? "Unknown").toLocaleLowerCase().localeCompare((getAgentCapability(right, "CPU") ?? "Unknown"));
+         case 'systemInfoRAM':
+            return parseInt((getAgentCapability(left, "RAM") ?? "0").replace(/\D/g, "")) - parseInt((getAgentCapability(right, "RAM") ?? "0").replace(/\D/g, ""));
+         default:
+            return left.name.localeCompare(right.name);
+      }
+   }
+
+   function onContextMenu(item?: any, index?: number | undefined, ev?: Event | undefined) {
+      const event = ev as MouseEvent;
+      localState.setRightClickDiv(event?.clientX, event?.clientY);
+      localState.setAgentContextMenuOpen(true);
+   }
+
+   function onHistoryModalDismiss() {
+      localState.setAgentId(undefined);
+   }
+
+   const totalPoolSize = agentStore.agents.filter((a) => {
+      if (!a.pools) {
+         return false;
+      }
+      return !poolId || a.pools.indexOf(poolId) !== -1
+   }).length;
+
+   let agentItems = agentStore.agents.filter(filterAgents, localState.agentFilter).sort(sortAgents);
+
+   if (poolId) {
+      agentItems = agentItems.filter((item => {
+         return !!item.pools && item.pools.indexOf(poolId!) !== -1;
+      }))
+   }
+
+   agentItems = agentItems.filter(item => {
+      const filter = localState.agentStatusFilter;
+      if (!filter.size) {
+         return true;
+      }
+
+      let filtered = true;
+
+      if (filter.has("Active")) {
+         if (item.leases?.length) {
+            filtered = false;
+         }
+      }
+
+      if (filter.has("No Pools")) {
+         if (!item.pools?.length) {
+            filtered = false;
+         } else {
+            let anyInvalid = false;
+            item.pools.forEach(pid => {
+
+               if (anyInvalid) {
+                  return;
+               }
+
+               if (!agentStore.pools.find(pool => { return pool.id === pid; })) {
+                  anyInvalid = true;
+               }
+
+            });
+
+            if (anyInvalid) {
+               filtered = false;
+            }
+         }
+
+      }
+
+      if (filter.has("Ephemeral")) {
+         if (item.ephemeral) {
+            filtered = false;
+         }
+      }
+
+      if (filter.has("Offline")) {
+         if (!item.online) {
+            filtered = false;
+         }
+      }
+
+      if (filter.has("Offline (Autoscaler)")) {
+         if (!item.online && item.lastShutdownReason === "Autoscaler") {
+            filtered = false;
+         }
+      }
+
+      if (filter.has("Offline (Unexpected)")) {
+         if (!item.online && item.lastShutdownReason === "Unexpected") {
+            filtered = false;
+         }
+      }
+
+      if (filter.has("Offline (Manual)")) {
+         if (!item.online && item.lastShutdownReason?.startsWith("Manual")) {
+            filtered = false;
+         }
+      }
+
+      if (filter.has("Pending Shutdown")) {
+
+         if (item.online && item.pendingShutdown) {
+            filtered = false;
+         }
+      }
+
+      if (filter.has("Disabled")) {
+         if (!item.enabled) {
+            filtered = false;
+         }
+      }
+
+      if (filter.has("Pending Conform")) {
+
+         if (item.pendingConform || item.pendingFullConform) {
+            filtered = false;
+         }
+      }
+
+      if (filter.has("Ready")) {
+         if (!item.leases?.length && item.online && item.enabled && !item.pendingConform && !item.pendingFullConform) {
+            filtered = false;
+         }
+      }
+
+      return !filtered;
+   });
+
+   const height = Math.min(((totalPoolSize < 5 ? 5 : totalPoolSize) * 48) + 48, 480)
+
+   return (<Stack>
+      {!!agentView && <SearchUpdate />}
+      <Stack horizontal style={{ paddingBottom: !agentView ? 18 : 0 }}>
+         {!!agentView && <AgentPivot />}
+         <Stack grow />
+         <AgentMenuBar agentView={agentView} />
+      </Stack>
+      <Stack style={{ position: "relative", height: agentView ? "calc(100vh - 240px)" : height }}>
+         <ScrollablePane scrollbarVisibility={ScrollbarVisibility.always}>
+            <DetailsList
+               className={agentStyles.detailsList}
+               onItemContextMenu={onContextMenu}
+               checkboxCellClassName={agentStyles.checkboxCell}
+               compact={true}
+               setKey="set"
+               items={agentItems}
+               columns={localState.columnsState.filter(colState => { return colState.isChecked; }).map(colState => { return colState.columnDef!; })}
+               selection={localState.selection}
+               onRenderDetailsHeader={onRenderDetailsHeader}
+               onColumnHeaderContextMenu={onColumnHeaderContextMenu}
+               onRenderItemColumn={onRenderAgentListItem}
+               layoutMode={DetailsListLayoutMode.justified}
+               constrainMode={ConstrainMode.horizontalConstrained}
+               selectionMode={SelectionMode.multiple}
+            />
+            <ContextualMenu
+               items={localState.columnMenuProps}
+               onItemClick={() => localState.setHeaderContextMenuOpen(false)}
+               onDismiss={() => localState.setHeaderContextMenuOpen(false)}
+               isBeakVisible={true}
+               hidden={!localState.headerContextMenuOpen}
+               target={localState.contextMenuTargetRef}
+               directionalHint={DirectionalHint.bottomLeftEdge}
+               directionalHintFixed={true}
+            />
+            <ContextualMenu
+               items={agentContextMenuProps}
+               onItemClick={() => localState.setAgentContextMenuOpen(false)}
+               onDismiss={() => localState.setAgentContextMenuOpen(false)}
+               isBeakVisible={true}
+               target={{ x: localState.mouseX, y: localState.mouseY }}
+               hidden={!localState.agentContextMenuOpen}
+               directionalHint={DirectionalHint.bottomLeftEdge}
+               directionalHintFixed={false}
+            />
+         </ScrollablePane>
+         <Stack grow />
+      </Stack>
+      <ConfirmationDialog
+         title={`Delete Agent${localState.selection.getSelectedCount() > 1 ? "s" : ""}`}
+         subText={`Really delete selected agent${localState.selection.getSelectedCount() > 1 ? "s" : ""}?`}
+         isOpen={localState.deleteAgentDialogIsOpen}
+         confirmText={"Delete"}
+         cancelText={"Cancel"}
+         onConfirm={() => { localState.deleteBuilder() }}
+         onCancel={() => { localState.setDeleteBuilderDialogOpen(false) }}
+      />
+      <ConfirmationDialog
+         title={`Restart Agent${localState.selection.getSelectedCount() > 1 ? "s" : ""}`}
+         subText={`Really restart selected agent${localState.selection.getSelectedCount() > 1 ? "s" : ""}?`}
+         isOpen={localState.restartAgentDialogIsOpen}
+         confirmText={"Restart"}
+         cancelText={"Cancel"}
+         checkBoxText='Force Restart'
+         onConfirm={(_, force) => { localState.requestBuilderUpdate(false, !force ? true : false, false, false, !!force) }}
+         onCancel={() => { localState.setRestartBuilderDialogOpen(false) }}
+      />
+      <ConfirmationDialog
+         title={`Shutdown Agent${localState.selection.getSelectedCount() > 1 ? "s" : ""}`}
+         subText={`Really shutdown selected agent${localState.selection.getSelectedCount() > 1 ? "s" : ""}?`}
+         isOpen={localState.shutdownAgentDialogIsOpen}
+         confirmText={"Shutdown"}
+         cancelText={"Cancel"}
+         textBoxLabel={"Type Confirm to confirm"}
+         isTextBoxSpawned={true}
+         onConfirm={() => { localState.requestBuilderUpdate(false, false, false, true) }}
+         onCancel={() => { localState.setShutdownBuilderDialogOpen(false) }}
+      />
+
+      <ConfirmationDialog
+         title={`Cancel Leases`}
+         subText={`Really cancel selected agent leases?`}
+         isOpen={localState.cancelLeasesDialogIsOpen}
+         confirmText={"Yes"}
+         cancelText={"No"}
+         onConfirm={() => { localState.setCancelLeasesDialogOpen(false); localState.cancelLeases() }}
+         onCancel={() => { localState.setCancelLeasesDialogOpen(false) }}
+      />
+      <ConfirmationDialog
+         title={`Disable Agent${localState.selection.getSelectedCount() > 1 ? "s" : ""}`}
+         isOpen={localState.disableAgentDialogIsOpen}
+         confirmText={"Disable"}
+         cancelText={"Cancel"}
+         textBoxLabel={"Enter Disable Reason - Required"}
+         isTextBoxSpawned={true}
+         isTextRequired={true}
+         onConfirm={(textFieldText: string) => { localState.changeBuilderEnabled("disable", textFieldText) }}
+         onCancel={() => { localState.setDisableBuilderDialogOpen(false) }}
+      />
+      <ConfirmationDialog
+         title={`Edit Agent Comment${localState.selection.getSelectedCount() > 1 ? "s" : ""}`}
+         isOpen={localState.editCommentDialogIsOpen}
+         confirmText={"Edit Comment"}
+         cancelText={"Cancel"}
+         textBoxLabel={"New Comment"}
+         isTextBoxSpawned={true}
+         onConfirm={(textFieldText: string) => { localState.changeCommment(textFieldText) }}
+         onCancel={() => { localState.setEditCommentDialogOpen(false) }}
+      />
+
+      <HistoryModal agentId={activeAgent?.id} onDismiss={onHistoryModalDismiss}></HistoryModal>
+   </Stack>
+   );
+
+
+   function onColumnHeaderContextMenu(column?: IColumn | undefined, ev?: React.MouseEvent<HTMLElement, MouseEvent> | undefined) {
+      const key = column!.key;
+      if (key === "software" || key === "storage" || key.indexOf("systemInfo") !== -1 || key === "taskTime" || key === "comment") {
+         localState.contextMenuTargetRef = ev?.nativeEvent!;
+         localState.setHeaderContextMenuOpen(true);
+      }
+   }
+
+   function getAgentStatusIcon(agent: AgentData) {
+
+      const statusColors = dashboard.getStatusColors();
+
+      let color = statusColors.get(StatusColor.Success);
+      if (agent.enabled === false) {
+         color = statusColors.get(StatusColor.Warnings);
+      }
+      if (!agent.online) {
+         color = statusColors.get(StatusColor.Skipped);
+      }
+      return <Icon iconName="FullCircle" style={{ color: color, fontSize: 16, marginRight: "13px", paddingTop: "2px" }} />;
+   }
+
+   function getPropFromDevice(agent: AgentData, propKey: string, propValue: string | null = null) {
+      const unknownElement = <Stack styles={{ root: { height: '100%' } }} horizontal horizontalAlign={"center"}><Stack.Item align={"center"}>Unknown</Stack.Item></Stack>;
+      let prop = propValue ?? getAgentCapability(agent, propKey);
+      if (prop === null)
+         return unknownElement;
+
+      if (propKey.indexOf("RAM") !== -1) {
+         prop += " GB";
+      }
+      return (
+         <Stack styles={{ root: { height: '100%' } }} horizontal horizontalAlign={"center"}>
+            <Stack.Item align={"center"} className={agentStyles.ellipsesStackItem}>
+               <Text title={prop} key={`sysInfo_${agent.id}_${propKey}`}>{`${prop}`}</Text>
+            </Stack.Item>
+         </Stack>
+      );
+   }
+
+   // when an actual item is drawn
+   function onRenderAgentListItem(agent: AgentData, index?: number, column?: IColumn) {
+      switch (column!.key) {
+         case 'name':
+            return (
+               <Stack styles={{ root: { height: 31, } }} horizontal horizontalAlign={"start"}>
+                  <Stack.Item align={"center"}>
+                     {getAgentStatusIcon(agent)}
+                  </Stack.Item>
+                  <Stack.Item align={"center"}>
+                     <ReactLink styles={{ root: { paddingTop: '1px', fontSize: "12px" } }} title="Lease and Session History" onClick={() => localState.setAgentId(agent.id)}>{agent.name}</ReactLink>
+                  </Stack.Item>
+               </Stack>
+            );
+         case 'pools':
+            const poolItems: any = [];
+            const poolSearchNames: string[] = [];
+            const invalidPoolIds = new Set<string>();
+            if (agent.pools) {
+               const poolObjs: PoolData[] = [];
+               // get actual pool objects
+               for (const key in agent.pools) {
+                  const id = agent.pools[key];
+                  const pool = agentStore.pools.find(pool => { return pool.id === id; });
+                  if (pool) {
+                     poolObjs.push(pool);
+                  } else {
+                     invalidPoolIds.add(id);
+                     poolObjs.push({
+                        id: id,
+                        name: id ?? "",
+                        colorValue: "#797979",
+                        enableAutoscaling: false,
+                        workspaces: []
+                     })
+                  }
+               }
+               // sort them by name
+               poolObjs.sort((a, b) => {
+                  return a.name.localeCompare(b.name);
+               });
+               for (let idx = 0; idx < poolObjs.length; idx++) {
+                  let color = poolObjs[idx].colorValue;
+                  const textColor = invalidPoolIds.has(poolObjs[idx]?.id) ? "#DDDDDD" : "white";
+                  if (agent.pendingConform || agent.pendingFullConform) {
+                     const pendingConformColor = hexToRGB(color);
+                     color = `rgb(${pendingConformColor.r},${pendingConformColor.g},${pendingConformColor.b}, .5)`;
+                  }
+                  const menuProps: IContextualMenuProps = {
+                     items: [
+                        {
+                           key: 'agent_pool_view',
+                           text: 'View Pool',
+                           onClick: (ev) => {
+                              window.open(`/pools?pool=${poolObjs[idx]?.id}`, '_blank')
+                           }
+                        },
+                        {
+                           key: 'agent_pool_filter',
+                           text: 'Filter to Pool',
+                           onClick: (ev) => {
+                              localState.setExactMatch(true);
+                              localState.setAgentFilter(poolObjs[idx]?.name);
+                           }
+                        },
+                        {
+                           key: 'agent_pool_copy_to_clipboard',
+                           text: 'Copy to Clipboard',
+                           onClick: (ev) => {
+                              copyToClipboard(poolObjs[idx].name)
+                           }
+                        }
+                     ],
+                  };
+
+                  const disabled = invalidPoolIds.has(poolObjs[idx].id);
+
+
+                  poolItems.push(
+                     <Stack.Item align={"center"} key={"pool_" + agent.id + "_" + poolObjs[idx].id}>
+                        <DefaultButton key={poolObjs[idx].id}
+                           text={poolObjs[idx].name}
+                           primary
+                           menuProps={!disabled ? menuProps : undefined}
+                           menuIconProps={{ iconName: "" }}
+                           className={agentStyles.buttonFont}
+                           onClick={(ev) => { ev.preventDefault(); ev.stopPropagation() }}
+                           styles={{
+                              root: { border: '0px', backgroundColor: color, color: textColor },
+                              rootHovered: { border: '0px', backgroundColor: color, color: textColor, },
+                              rootPressed: { border: '0px', backgroundColor: color, color: textColor, }
+                           }} />
+
+                     </Stack.Item>
+                  );
+                  poolSearchNames.push(poolObjs[idx].name);
+               }
+            }
+            localState.columnSearchState['pools'][agent.id] = poolSearchNames;
+            return <Stack horizontal horizontalAlign={"start"} styles={{ root: { overflow: "auto", height: '100%' } }} tokens={{ childrenGap: 6 }}>{poolItems}</Stack>;
+         case 'status':
+            const leases: any = [];
+            const leaseSearchItems: string[] = [];
+            if (agent.leases) {
+               agent.leases.forEach(lease => {
+
+                  let statusText = lease.name;
+
+                  if (lease.state === LeaseState.Cancelled) {
+                     if (lease.finishTime) {
+                        statusText += " (Canceled)"
+                     } else {
+                        statusText += " (Cancelling)"
+                     }
+                  }
+
+                  let link = "";
+                  if (lease.details) {
+                     if ('jobId' in lease.details) {
+                        link = `/job/${lease.details['jobId']}`;
+                        if ('batchId' in lease.details) {
+                           link += `?batch=${lease.details['batchId']}`;
+                        }
+                     }
+                     else if ('logId' in lease.details) {
+                        link = `/log/${lease.details['logId']}?agentId=${agent.id}`;
+                     }
+                  }
+                  if (link !== "") {
+                     // marquees dont play nice unless the parent container is 100% of the width. regular centered stack items dont play nice if
+                     // their root width's are set...
+                     if (statusText!.length > 46) {
+                        leases.push(<Stack.Item align={"center"} key={"statusText_item_" + agent.id + "_" + lease.id} styles={{ root: { width: '100%' } }}>
+                           <Link key={"statusText_" + agent.id + "_" + lease.id} to={link}>
+                              <Marquee loop={true} trailing={2000} text={statusText} />
+                           </Link>
+                        </Stack.Item>);
+                     }
+                     else {
+                        leases.push(<Stack.Item align={"center"} key={"statusText_item_" + agent.id + "_" + lease.id}>
+                           <Link key={"statusText_" + agent.id + "_" + lease.id} to={link}>{statusText}</Link>
+                        </Stack.Item>);
+                     }
+                     leases.push();
+                  }
+                  else {
+                     leases.push(<Stack.Item align={"center"} key={"statusText_item_" + agent.id + "_" + lease.id} className={agentStyles.ellipsesStackItem}>
+                        <Text key={"statusText_" + agent.id + "_" + lease.id}>{statusText}</Text>
+                     </Stack.Item>);
+                  }
+                  if (lease.name) {
+                     leaseSearchItems.push(lease.name);
+                  }
+               });
+            }
+            // if there are no leases, we'll push some other state.
+            if (leases.length === 0) {
+               let title = agent.status ?? "";
+               if (title === "Ok") {
+                  title = "Ready";
+               }
+               let subtitle = "";
+               if (!agent.online) {
+                  title = `Offline - ${agent.lastShutdownReason}`;
+                  subtitle = `Last online at ${getShortNiceTime(agent.updateTime, true, true, true)}`;
+               }
+               else if (agent.pendingShutdown) {
+                  title = `Pending Shutdown`;
+                  subtitle = agent.lastShutdownReason;
+               }
+               else if (agent.pendingConform || agent.pendingFullConform) {
+                  if (agent.conformAttemptCount && agent.conformAttemptCount > 0) {
+                     title = `${agent.pendingFullConform ? "Full " : ""}Conform failed (${agent.conformAttemptCount} attempts)`;
+                     if (agent.nextConformTime) {
+                        subtitle = `Next attempt at ${getShortNiceTime(agent.nextConformTime, true, true, true)}`;
+                     }
+                  }
+                  else {
+                     title = "Pending Conform";
+                     if (agent.pendingFullConform) {
+                        title += " (Full)";
+                     }
+                  }
+               }
+               else if (!agent.enabled) {
+                  title = "Disabled";
+               }
+
+               leases.push(<Stack.Item key={"itemStatus_" + agent.id} align={"center"} className={agentStyles.ellipsesStackItem}>
+                  <Text key={"statusText_" + agent.id}>{title}</Text>
+               </Stack.Item>);
+               leaseSearchItems.push(title);
+               if (subtitle !== "") {
+                  leases.push(<Stack.Item key={"itemSubtitleStatus_" + agent.id} align={"center"} className={agentStyles.ellipsesStackItem}>
+                     <Text variant="small" key={"statusSubtitleText_" + agent.id}>{subtitle}</Text>
+                  </Stack.Item>);
+                  leaseSearchItems.push(subtitle);
+               }
+
+            }
+
+            localState.columnSearchState['status'][agent.id] = leaseSearchItems;
+            return <Stack styles={{ root: { height: '100%', } }} verticalAlign={"center"}>{leases}</Stack>;
+         case 'software':
+            return (<Stack styles={{ root: { height: '100%', } }} horizontal horizontalAlign={"center"}>
+               <Stack.Item align={"center"}>
+                  <Text key={"softwareText_" + agent.id} className={agentStyles.ellipsesStackItem}>{agent.version ?? "Unknown"}
+                  </Text>
+               </Stack.Item>
+            </Stack>);
+         case 'taskTime':
+            let longestDuration: moment.Duration | undefined;
+            let taskText = "Agent Idle";
+            if (agent.leases) {
+               agent.leases.forEach(lease => {
+                  const start = moment(lease.startTime);
+                  const end = moment(Date.now());
+                  const duration = moment.duration(end.diff(start));
+                  if (!longestDuration) {
+                     longestDuration = duration;
+                  }
+                  else {
+                     if (duration.asSeconds() > longestDuration.asSeconds()) {
+                        longestDuration = duration;
+                     }
+                  }
+               });
+            }
+
+            if (longestDuration) {
+               taskText = longestDuration!.humanize();
+               localState.columnSearchState['taskTime'][agent.id] = longestDuration!.asSeconds();
+            }
+            else {
+               localState.columnSearchState['taskTime'][agent.id] = 0;
+            }
+
+            localState.columnSearchState['taskTimeDescription'][agent.id] = taskText;
+
+
+            return (<Stack styles={{ root: { height: '100%', } }} horizontal horizontalAlign={"center"}>
+               <Stack.Item align={"center"}>
+                  <Text key={"taskTimeText_" + agent.id} className={agentStyles.ellipsesStackItem}>{taskText}
+                  </Text>
+               </Stack.Item>
+            </Stack>);
+         case 'storage':
+            let diskFree: string | number | null = getAgentCapability(agent, "DiskFreeSpace");
+            let diskCapacity: string | number | null = getAgentCapability(agent, "DiskTotalSize");
+
+            let realPercentage = 0;
+            let nudgedPercentage = 0;
+            let descVal = "Unknown";
+            let descNumber = 0;
+            if (diskFree !== null && diskCapacity !== null) {
+               diskFree = Number(diskFree) / 1073741824;
+               diskCapacity = Number(diskCapacity) / 1073741824;
+               realPercentage = (diskCapacity - diskFree) / diskCapacity;
+               // nudge percentage for the progress bar length
+               nudgedPercentage = realPercentage * .965;
+               descNumber = (diskCapacity - (realPercentage * diskCapacity));
+               descVal = descNumber.toLocaleString(undefined, { maximumFractionDigits: 0 }) + " GiB Free";
+            }
+            localState.columnSearchState['storage'][agent.id] = descNumber;
+            return <Stack horizontal horizontalAlign={"center"}>
+               <Stack horizontal horizontalAlign={"end"} styles={{ root: { minWidth: 80, maxWidth: 80 } }}>
+                  <Text styles={{ root: { marginRight: 10, marginTop: 7 } }}>{descVal}</Text>
+               </Stack>
+               <Stack horizontal styles={{ root: { minWidth: 110, maxWidth: 110 } }}>
+                  <ProgressIndicator
+                     barHeight={15}
+                     styles={{
+                        progressBar: { height: dashboard.darktheme ? 13 : 16, marginLeft: dashboard.darktheme ? 1 : 0, marginTop: dashboard.darktheme ? 2 : 0 },
+                        progressTrack: { width: '98%', border: '1px solid !important' },
+                        root: { width: 110 }
+                     }}
+                     percentComplete={nudgedPercentage}
+                  />
+               </Stack>
+            </Stack>;
+         case 'comment':
+            localState.columnSearchState['comment'][agent.id] = agent.comment ?? "";
+            return (<Stack styles={{ root: { height: '100%', } }} horizontal horizontalAlign={"center"}>
+               <Stack.Item align={"center"}>
+                  <Text key={"commentText_" + agent.id} className={agentStyles.ellipsesStackItem}>{agent.comment ?? ""}
+                  </Text>
+               </Stack.Item>
+            </Stack>);
+         case 'systemInfoOS':
+            return getPropFromDevice(agent, "OSDistribution", getAgentOs(agent));
+         case 'systemInfoCPU':
+            return getPropFromDevice(agent, "CPU");
+         case 'systemInfoRAM':
+            return getPropFromDevice(agent, "RAM");
+         default:
+            return <span>{agent[column!.fieldName as keyof AgentData] as string}</span>;
+      }
+   }
+});
+
+export const AgentPanel: React.FC<{ agentId?: string, poolId?: string, agentView?: boolean }> = observer(({ agentId, agentView, poolId }) => {
+
+   // adjust automatically to viewport changes
+   useWindowSize();
+
+   // subscribe
+   if (localState.searchUpdated) { }
+
+   localState.agentView = agentView;
+   localState.updateColumns();
+
+   return <Stack>
+      <AgentViewInner agentId={localState.searchState?.agentId} agentView={agentView} poolId={poolId} />
+   </Stack>
+
+})
+
+export const AgentView: React.FC = () => {
+
+   const [searchParams] = useSearchParams();
+   const hordeTheme = getHordeTheme();
+   const { hordeClasses } = getHordeStyling();
+
+   const agentId = searchParams.get("agentId") ? searchParams.get("agentId") : undefined;
+
+   localState.agentView = true;
+   localState.updateColumns();
+
+   // adjust automatically to viewport changes
+   useWindowSize();
+
+   return <Stack className={hordeClasses.horde}>
+      <TopNav />
+      <Breadcrumbs items={[{ text: 'Agents' }]} />
+      <Stack horizontal>
+         <Stack grow styles={{ root: { backgroundColor: hordeTheme.horde.neutralBackground } }} />
+         <Stack tokens={{ maxWidth: 1440, childrenGap: 4 }} styles={{ root: { width: 1440, height: '100vh', backgroundColor: hordeTheme.horde.neutralBackground, paddingTop: 18, paddingLeft: 12 } }}>
+            <Stack className={hordeClasses.raised} styles={{ root: { paddingRight: '40px' } }}>
+               <AgentViewInner agentView={true} agentId={agentId ? agentId : undefined} searchParams={searchParams} />
+            </Stack>
+         </Stack>
+         <Stack grow styles={{ root: { backgroundColor: hordeTheme.horde.neutralBackground } }} />
+      </Stack>
+   </Stack>
+
+}
\ No newline at end of file
diff --git a/Engine/Source/Programs/Horde/HordeDashboard/src/components/agents/HistoryModal.tsx b/Engine/Source/Programs/Horde/HordeDashboard/src/components/agents/HistoryModal.tsx
new file mode 100644
--- /dev/null
+++ b/Engine/Source/Programs/Horde/HordeDashboard/src/components/agents/HistoryModal.tsx
@@ -0,0 +1,1140 @@
+// Copyright Epic Games, Inc. All Rights Reserved.
+
+import { Checkbox, ComboBox, ConstrainMode, DefaultButton, DetailsHeader, DetailsList, DetailsListLayoutMode, DetailsRow, Dialog, DialogFooter, DialogType, FontIcon, GroupedList, GroupHeader, ICheckbox, IColumn, IComboBoxOption, IContextualMenuItem, IContextualMenuProps, IDetailsHeaderProps, IDetailsHeaderStyles, IDetailsListProps, IGroup, ITextField, ITooltipHostStyles, Label, mergeStyleSets, Modal, Pivot, PivotItem, PrimaryButton, ScrollablePane, ScrollbarVisibility, Selection, SelectionMode, Spinner, SpinnerSize, Stack, Sticky, StickyPositionType, Text, TextField } from "@fluentui/react";
+import { observer } from "mobx-react-lite";
+import React, { useEffect, useState } from "react";
+import { Link } from "react-router-dom";
+import { Sparklines, SparklinesLine, SparklinesReferenceLine } from "react-sparklines";
+import backend from "../../backend";
+import { agentStore } from "./AgentStore";
+import { GetAgentLeaseResponse, GetAgentTelemetrySampleResponse, JobStepBatchError, JobStepOutcome, LeaseData, SessionData, UpdateAgentRequest } from "../../backend/Api";
+import dashboard, { StatusColor } from "../../backend/Dashboard";
+import { getNiceTime, getShortNiceTime } from "../../base/utilities/timeUtils";
+import { getHordeStyling } from "../../styles/Styles";
+import { getHordeTheme } from "../../styles/theme";
+import { BatchStatusIcon, LeaseStatusIcon, StepStatusIcon } from "./../StatusIcon";
+import { historyModalStore } from "./HistoryModalStore";
+
+const historyStyles = mergeStyleSets({
+   dialog: {
+      selectors: {
+         ".ms-Label,.ms-Button-label": {
+            fontWeight: "unset",
+            fontFamily: "Horde Open Sans SemiBold"
+         }
+      }
+   },
+   detailsList: {
+      selectors: {
+         ".ms-DetailsHeader-cellName": {
+            fontWeight: "unset",
+            fontFamily: "Horde Open Sans SemiBold"
+         }
+      }
+   }
+});
+
+const state = historyModalStore;
+
+const TelemetryTooltip: React.FC<{ id: string }> = observer(({ id }) => {
+
+   const { modeColors } = getHordeStyling();
+
+   const tooltip = state.currentLeaseTooltip;
+
+   if (!tooltip?.lease || tooltip.id !== id) {
+      return null;
+   }
+
+   let tipX = tooltip.x;
+   let offsetX = 32;
+   let translateX = "0%";
+
+   if ((tipX ?? 0) > 600) {
+      offsetX = -32;
+      translateX = "-100%";
+   }
+
+   const translateY = "-50%";
+
+   const textSize = 11;
+   const titleWidth = 48;
+
+   const dataElement = (title: string, value: string, link?: string, agent?: string, linkTarget?: string) => {
+
+      let valueElement = <Stack>
+         <Stack style={{ fontSize: textSize }}>
+            {value}
+         </Stack>
+      </Stack>
+
+      if (link) {
+         valueElement = <Link className="horde-link" to={link} target={linkTarget}>
+            {valueElement}
+         </Link>
+      }
+
+      let component = <Stack style={{ cursor: (link || agent) ? "pointer" : undefined }} onClick={() => {
+      }}>
+         <Stack horizontal>
+            <Stack style={{ width: titleWidth }}>
+               {!!title && <Text style={{ fontFamily: "Horde Open Sans SemiBold", fontSize: textSize }} >{title}</Text>}
+            </Stack>
+            {valueElement}
+         </Stack>
+      </Stack>
+
+      return component;
+   }
+
+   const elements: JSX.Element[] = [];
+
+   const lease = tooltip.lease;
+
+   if (lease?.details?.type === "Job") {
+      elements.push(dataElement("Job:", tooltip.lease?.name ?? "Unknown Job Lease", `/job/${lease.details.jobId}`));
+      const job = state.jobData.get(lease.details.jobId);
+      if (job && job !== true && job.batches) {
+         const step = job.batches.map(b => b.agentId === state.selectedAgent?.id ? b.steps : []).flat().find(s => {
+            if (s.startTime && tooltip.time) {
+               if (tooltip.time >= new Date(s.startTime) && (!s.finishTime || tooltip.time <= new Date(s.finishTime))) {
+                  return true;
+               }
+            }
+            return false;
+         });
+         if (step) {
+            elements.push(dataElement("Step:", step.name ?? "Unknown Job Step", `/job/${lease.details.jobId}?step=${step.id}`));
+            if (step.logId)
+               elements.push(dataElement("Log:", `Step`, `/log/${step.logId}`));
+         }
+      }
+   } else {
+      elements.push(dataElement("Lease:", tooltip.lease?.name ?? "Unknown Lease"));
+   }
+
+   if (lease?.details?.logId) {
+      elements.push(dataElement("Log:", `Lease`, `/log/${lease.details.logId}`));
+   }
+
+
+
+   if (tooltip.time) {
+      elements.push(dataElement("Time:", getNiceTime(tooltip.time)));
+   }
+
+
+   return <div style={{
+      position: "absolute",
+      display: "block",
+      top: `${tooltip.y}px`,
+      left: `${(tooltip.x ?? 0) + offsetX}px`,
+      backgroundColor: modeColors.background,
+      zIndex: 1,
+      border: "solid",
+      borderWidth: "1px",
+      borderRadius: "3px",
+      width: "max-content",
+      borderColor: dashboard.darktheme ? "#413F3D" : "#2D3F5F",
+      pointerEvents: tooltip.frozen ? undefined : "none",
+      transform: `translate(${translateX}, ${translateY})`
+   }}>
+      <Stack horizontal>
+         <Stack tokens={{ childrenGap: 6, padding: 12 }}>
+            {elements}
+         </Stack>
+         {!!tooltip.frozen && <Stack style={{ paddingLeft: 32, paddingRight: 12, paddingTop: 12, cursor: "pointer" }} onClick={() => { state.setLeaseTooltip(undefined) }}>
+            <FontIcon iconName="Cancel" />
+         </Stack>}
+         {!tooltip.frozen && <Stack style={{ paddingLeft: 32, paddingRight: 12 + 16, paddingTop: 12 }} />}
+
+      </Stack>
+   </div>
+
+});
+
+type TimeSelection = {
+   text: string;
+   key: string;
+   minutes: number;
+}
+
+const timeSelections: TimeSelection[] = [
+   {
+      text: "Past 1 Hour", key: "time_1_hour", minutes: 60
+   },
+   {
+      text: "Past 2 Hours", key: "time_2_hours", minutes: 60 * 2
+   },
+   {
+      text: "Past 4 Hours", key: "time_4_hours", minutes: 60 * 4
+   },
+   {
+      text: "Past 1 Day", key: "time_1_day", minutes: 60 * 24
+   },
+   {
+      text: "Past 2 Days", key: "time_2_days", minutes: 60 * 24 * 2
+   },
+   {
+      text: "Past 1 Week", key: "time_1_week", minutes: 60 * 24 * 7
+   },
+   {
+      text: "Past 2 Weeks", key: "time_2_weeks", minutes: 60 * 24 * 7 * 2
+   },
+   {
+      text: "Past Month", key: "time_4_weeks", minutes: 60 * 24 * 7 * 4
+   }
+]
+
+const cpuSelections: IComboBoxOption[] = [
+   {
+      text: "Combined", key: "combined"
+   },
+   {
+      text: "User", key: "user"
+   },
+   {
+      text: "System", key: "system"
+   }
+]
+
+const TelemetryPanel: React.FC<{}> = observer(({ }) => {
+
+   const [timeKey, setTimeKey] = useState("time_2_days");
+   const [cpuKey, setCPUKey] = useState("combined");
+
+   useEffect(() => {
+
+      const time = timeSelections.find(t => t.key === timeKey);
+      if (time) {
+         state.UpdateTelemetry(time.minutes);
+      }
+
+      return () => {
+      };
+
+   }, []);
+
+
+   // subscribe
+   if (state.currentData) { }
+
+   const currentData = (state.currentData ?? []);
+
+   if (!currentData.length) {
+      return <Stack horizontalAlign="center">
+         <Spinner size={SpinnerSize.large} />
+      </Stack>;
+   }
+
+   const statusColors = dashboard.getStatusColors();
+
+   const sparkWidth = 1100;
+   const data = (currentData[0] as GetAgentTelemetrySampleResponse[]);
+   const leases = ((currentData.length > 1) ? currentData[1] : []) as GetAgentLeaseResponse[];
+
+   if (!data.length) {
+      return null;
+   }
+
+   const cpuData = data.map(d => {
+      if (cpuKey == "user") return d.userCpu;
+      if (cpuKey == "system") return d.systemCpu;
+      if (cpuKey == "combined") return d.userCpu + d.systemCpu;
+      return 1;
+   })
+
+   const min = data[0].time;
+   const max = data[data.length - 1].time;
+
+   const minTime = getShortNiceTime(min, true);
+   const maxTime = getShortNiceTime(max, true);
+
+   const ramData = data.map(d => {
+      return d.usedRam / d.totalRam;
+   })
+
+   const diskData = data.map(d => {
+      if (typeof (d.totalDisk) === "number" && d.totalDisk) {
+         return 1.0 - (d.freeDisk / d.totalDisk)
+      }
+
+      return 0;
+   })
+
+   let cpuText = "";
+   let ramText = "";
+   let diskText = "";
+
+   let cpuLabel = "";
+   if (cpuKey == "user") cpuLabel = "CPU - User";
+   if (cpuKey == "system") cpuLabel = "CPU - System";
+   if (cpuKey == "combined") cpuLabel = "CPU - Combined";
+
+
+   const devices = state.selectedAgent?.capabilities?.devices;
+   if (devices?.length) {
+      const device = devices[0];
+      device.properties?.forEach(v => {
+         if (v.startsWith("CPU=")) {
+            cpuText = `${v.replace("CPU=", "")}`;
+         }
+
+         if (v.startsWith("RAM=")) {
+            ramText = `${v.replace("RAM=", "")} GB`;
+         }
+
+         if (v.startsWith("DiskTotalSize=")) {
+            const elements = v.split("=");
+            if (elements.length === 2) {
+
+               function formatBytes(bytes: number, decimals = 0) {
+                  if (!+bytes) return '0 Bytes'
+
+                  const k = 1024
+                  const dm = decimals < 0 ? 0 : decimals
+                  const sizes = ['Bytes', 'KB', 'MB', 'GB', 'TB', 'PB', 'EB', 'ZB', 'YB']
+
+                  const i = Math.floor(Math.log(bytes) / Math.log(k))
+
+                  return `${parseFloat((bytes / Math.pow(k, i)).toFixed(dm))} ${sizes[i]}`
+               }
+
+               diskText = formatBytes(parseInt(elements[1]))
+            }
+         }
+
+      })
+   }
+
+   function relativeCoords(event: any, id: string) {
+
+      if (state.currentLeaseTooltip?.frozen) {
+         return;
+      }
+
+      let bounds = event.target.getBoundingClientRect();
+      let x = event.clientX - bounds.left;
+      let y = event.clientY - bounds.top;
+
+      // interpolate time
+
+      const percent = ((x - 16) / sparkWidth);
+
+      const time = new Date(min.getTime() + ((max.getTime() - min.getTime()) * percent));
+
+      let lease: GetAgentLeaseResponse | undefined;
+      leases.forEach(s => {
+         if (time >= s.startTime && (!s.finishTime || time <= s.finishTime)) {
+            lease = s;
+         }
+      });
+
+      if (lease?.details?.type === "Job") {
+         const jobId = lease.details.jobId;
+         if (jobId) {
+            if (!state.jobData.has(jobId)) {
+               state.jobData.set(jobId, true);
+               backend.getJob(jobId, { filter: "id,batches" }).then(job => {
+                  state.jobData.set(jobId, job);
+                  state.setLeaseTooltip(state.currentLeaseTooltip ? { ...state.currentLeaseTooltip } : undefined);
+               })
+            }
+         }
+      }
+
+      state.setLeaseTooltip({ lease: lease, x: x, y: y, id: id, time: time });
+   };
+
+   return <Stack tokens={{ childrenGap: 24 }} style={{ paddingBottom: 32, paddingLeft: 8 }}>
+      <Stack horizontal verticalAlign="center" verticalFill>
+         <Text variant="mediumPlus">{minTime} - {maxTime}</Text>
+         <Stack grow />
+         <Stack horizontal tokens={{ childrenGap: 18 }}>
+
+            <ComboBox
+               label={"CPU"}
+               styles={{ root: { width: 128 } }}
+               options={cpuSelections}
+               selectedKey={cpuKey}
+               onChange={(ev, option, index, value) => {
+                  if (!option) {
+                     return;
+                  }
+                  setCPUKey(option.key as string);
+               }}
+            />
+
+            <ComboBox
+               label="Time"
+               styles={{ root: { width: 180 } }}
+               options={timeSelections}
+               selectedKey={timeKey}
+               onChange={(ev, option, index, value) => {
+                  const select = option as TimeSelection;
+                  setTimeKey(select.key);
+                  state.UpdateTelemetry(select.minutes);
+               }}
+            />
+         </Stack>
+      </Stack>
+      <Stack verticalAlign="center" verticalFill>
+         <Stack>
+            <Label>{cpuLabel} Utilization {cpuText ? `- ${cpuText}` : ""}</Label>
+         </Stack>
+         <div onMouseLeave={() => state.setLeaseTooltip(undefined)} onMouseMove={(ev) => { relativeCoords(ev, "cpu") }} onClick={() => {
+
+            if (!state.currentLeaseTooltip?.frozen) {
+               if (state.currentLeaseTooltip) {
+                  state.setLeaseTooltip({ ...state.currentLeaseTooltip, frozen: true });
+               }
+            }
+         }}>
+            <div style={{ position: "relative" }}>
+               <TelemetryTooltip id={"cpu"} />
+            </div>
+
+            <Sparklines width={sparkWidth} height={128} data={cpuData} max={100} style={{ backgroundColor: dashboard.darktheme ? "#060709" : "#F3F2F1", padding: 8, border: "solid 1px #181A1B" }}>
+               <SparklinesLine color={dashboard.darktheme ? "lightblue" : "#1E90FF"} />
+               <SparklinesReferenceLine type="custom" value={0} style={{ stroke: statusColors.get(StatusColor.Failure) }} />
+               <SparklinesReferenceLine type="custom" value={32} style={{ stroke: statusColors.get(StatusColor.Warnings) }} />
+               <SparklinesReferenceLine type="custom" value={64} style={{ stroke: statusColors.get(StatusColor.Success) }} />
+               <SparklinesReferenceLine type="custom" value={96} style={{ stroke: statusColors.get(StatusColor.Success) }} />
+
+            </Sparklines>
+         </div>
+      </Stack>
+      <Stack verticalAlign="center" verticalFill>
+         <Stack>
+            <Label>RAM Utilization {ramText ? `- ${ramText}` : ""}</Label>
+         </Stack>
+         <div onMouseLeave={() => state.setLeaseTooltip(undefined)} onMouseMove={(ev) => { relativeCoords(ev, "ram") }} onClick={() => {
+
+            if (!state.currentLeaseTooltip?.frozen) {
+               if (state.currentLeaseTooltip) {
+                  state.setLeaseTooltip({ ...state.currentLeaseTooltip, frozen: true });
+               }
+            }
+         }}>
+            <div style={{ position: "relative" }}>
+               <TelemetryTooltip id="ram" />
+            </div>
+            <Sparklines width={sparkWidth} height={128} data={ramData} max={1} style={{ backgroundColor: dashboard.darktheme ? "#060709" : "#F3F2F1", padding: 8, border: "solid 1px #181A1B" }}>
+               <SparklinesLine color={dashboard.darktheme ? "lightblue" : "#1E90FF"} />
+               <SparklinesReferenceLine type="custom" value={0} style={{ stroke: statusColors.get(StatusColor.Failure) }} />
+               <SparklinesReferenceLine type="custom" value={32} style={{ stroke: statusColors.get(StatusColor.Warnings) }} />
+               <SparklinesReferenceLine type="custom" value={64} style={{ stroke: statusColors.get(StatusColor.Success) }} />
+               <SparklinesReferenceLine type="custom" value={96} style={{ stroke: statusColors.get(StatusColor.Success) }} />
+
+            </Sparklines>
+         </div>
+      </Stack>
+
+      {!!diskText && <Stack verticalAlign="center" verticalFill>
+         <Stack>
+            <Label>Storage Utilization - {diskText}</Label>
+         </Stack>
+         <div onMouseLeave={() => state.setLeaseTooltip(undefined)} onMouseMove={(ev) => { relativeCoords(ev, "disk") }} onClick={() => {
+
+            if (!state.currentLeaseTooltip?.frozen) {
+               if (state.currentLeaseTooltip) {
+                  state.setLeaseTooltip({ ...state.currentLeaseTooltip, frozen: true });
+               }
+            }
+         }}>
+            <div style={{ position: "relative" }}>
+               <TelemetryTooltip id="disk" />
+            </div>
+            <Sparklines width={sparkWidth} height={128} data={diskData} max={1} style={{ backgroundColor: dashboard.darktheme ? "#060709" : "#F3F2F1", padding: 8, border: "solid 1px #181A1B" }}>
+               <SparklinesLine color={dashboard.darktheme ? "lightblue" : "#1E90FF"} />
+               <SparklinesReferenceLine type="custom" value={0} style={{ stroke: statusColors.get(StatusColor.Failure) }} />
+               <SparklinesReferenceLine type="custom" value={32} style={{ stroke: statusColors.get(StatusColor.Warnings) }} />
+               <SparklinesReferenceLine type="custom" value={64} style={{ stroke: statusColors.get(StatusColor.Success) }} />
+               <SparklinesReferenceLine type="custom" value={96} style={{ stroke: statusColors.get(StatusColor.Success) }} />
+            </Sparklines>
+         </div>
+      </Stack>}
+
+   </Stack>
+})
+
+export const HistoryModal: React.FC<{ agentId: string | undefined, onDismiss: (...args: any[]) => any; }> = observer(({ agentId, onDismiss }) => {
+
+   const [selectedAgent, setSelectedAgent] = useState<string | undefined>(undefined);
+   const [actionState, setActionState] = useState<{ action?: string, confirmed?: boolean, comment?: string }>({});
+   const actionTextInputRef = React.useRef<ITextField>(null);
+   const forceRestartCheckboxRef = React.useRef<ICheckbox>(null);
+   const [agentError, setAgentError] = useState(false);
+
+   const { hordeClasses, modeColors } = getHordeStyling();
+   const theme = getHordeTheme();
+
+   if (agentError) {
+      return <Dialog hidden={false} onDismiss={onDismiss} dialogContentProps={{
+         type: DialogType.normal,
+         title: 'Missing Agent',
+         subText: `Unable to find agent ${agentId}`
+      }}
+         modalProps={{ styles: { main: { width: "640px !important", minWidth: "640px !important", maxWidth: "640px !important" } } }}>
+         <DialogFooter>
+            <PrimaryButton onClick={() => onDismiss()} text="Ok" />
+         </DialogFooter>
+      </Dialog>
+   }
+
+   //  subscribe to updates
+   if (state.selectedAgent) { }
+
+   if (!agentId) {
+      if (selectedAgent) {
+         setSelectedAgent(undefined);
+      }
+      return null;
+   }
+
+   if (selectedAgent !== agentId) {
+
+      agentStore.updateAgent(agentId).then(agent => {
+
+         state.setSelectedAgent(agent);
+
+      }).catch(error => {
+         console.error(`Unable to find agent id: ${agentId} ${error}`);
+         setAgentError(true);
+         return;
+      })
+
+      setSelectedAgent(agentId);
+   }
+
+   if (!state.selectedAgent) {
+      if (agentId) {
+         return <Modal isOpen={true} isBlocking={true} topOffsetFixed={true} styles={{ main: { padding: 16, width: 1200, hasBeenOpened: false, top: "80px", position: "absolute" } }} className={hordeClasses.modal}>
+            <Stack tokens={{ childrenGap: 40 }} styles={{ root: { padding: 8 } }}>
+               <Stack grow verticalAlign="center">
+                  <Text variant="mediumPlus" styles={{ root: { fontWeight: "unset", fontFamily: "Horde Open Sans SemiBold", fontSize: "20px" } }}>Agent {agentId}</Text>
+               </Stack>
+               <Stack verticalAlign="center">
+                  <Spinner size={SpinnerSize.large} />
+               </Stack>
+            </Stack>
+         </Modal>
+      }
+
+      return null;
+   }
+
+   type AgentAction = {
+      name: string;
+      confirmText: string;
+      textInput?: boolean;
+      update?: (request: UpdateAgentRequest, comment?: string) => void;
+   }
+   const actions: AgentAction[] = [
+      {
+         name: 'Enable',
+         confirmText: "Are you sure you would like to enable this agent?",
+         update: (request) => { request.enabled = true }
+      },
+      {
+         name: 'Disable',
+         confirmText: "Are you sure you would like to disable this agent?",
+         textInput: true,
+         update: (request, comment) => { request.enabled = false; request.comment = comment }
+      },
+      {
+         name: 'Cancel Leases',
+         confirmText: "Are you sure you would like to cancel this agent's leases?"
+      },
+      {
+         name: 'Request Conform',
+         confirmText: "Are you sure you would like to request an agent conform?",
+         update: (request) => { request.requestConform = true }
+      },
+      {
+         name: 'Request Full Conform',
+         confirmText: "Are you sure you would like to request a full agent conform?",
+         update: (request) => { request.requestFullConform = true }
+
+      },
+      {
+         name: 'Request Restart',
+         confirmText: "Are you sure you would like to request an agent restart?",
+         update: (request) => { request.requestRestart = !forceRestartCheckboxRef.current?.checked; request.requestForceRestart = !!forceRestartCheckboxRef.current?.checked; }
+      },
+      {
+         name: 'Edit Comment',
+         confirmText: "Please enter new comment",
+         textInput: true,
+         update: (request, comment) => { request.comment = comment }
+      }
+   ];
+
+   // cancel leases action
+   const cancelLeases = async () => {
+
+      const leases = new Set<string>();
+
+      state.selectedAgent?.leases?.forEach(lease => {
+         if (!lease.finishTime) {
+            leases.add(lease.id);
+         }
+      });
+
+      let requests = Array.from(leases).map(id => {
+         return backend.updateLease(id, { aborted: true });
+      });
+
+      while (requests.length) {
+
+         const batch = requests.slice(0, 5);
+
+         await Promise.all(batch).then(() => {
+
+         }).catch((errors) => {
+            console.log(errors);
+            // eslint-disable-next-line
+         }).finally(() => {
+
+            requests = requests.slice(5);
+         });
+      }
+      agentStore.update();
+   };
+
+
+   const currentAction = actions.find(a => a.name === actionState.action);
+
+   if (currentAction && actionState.confirmed) {
+      if (currentAction.update) {
+         const request: UpdateAgentRequest = {};
+         currentAction.update(request, actionState.comment);
+         backend.updateAgent(agentId, request).then(() => {
+            if (!agentStore.pools?.length) {
+               agentStore.update(false).then(() => {
+                  state.setSelectedAgent(agentStore.agents.find(agent => agent.id === agentId));
+                  setActionState({});
+               });
+            } else {
+               agentStore.updateAgent(agentId).then(() => {
+                  state.setSelectedAgent(agentStore.agents.find(agent => agent.id === agentId));
+                  setActionState({});
+               });
+            }
+         }).catch((reason) => {
+            console.error(reason);
+         });
+      } else if (currentAction.name === "Cancel Leases") {
+         (async () => {
+            await cancelLeases();
+            setActionState({});
+         })();
+      }
+   }
+
+   const actionMenuProps: IContextualMenuProps = {
+      shouldFocusOnMount: true,
+      items: actions.map(a => {
+         return {
+            key: a.name.toLowerCase(),
+            text: a.name,
+            onClick: () => {
+               setActionState({ action: a.name })
+            }
+         }
+      })
+   };
+
+   const poolItems: IContextualMenuItem[] = [];;
+   const agent = state.selectedAgent;
+   agent.pools?.forEach(poolId => {
+
+      const pool = agentStore.pools?.find(pool => poolId === pool.id);
+      if (pool) {
+         poolItems.push({
+            key: `pools_${poolId}`,
+            text: pool.name,
+            href: `/pools?pool=${pool.id}`,
+            target: "_blank"
+         })
+      }
+   })
+
+   function onColumnClick(ev: React.MouseEvent<HTMLElement>, column: IColumn) {
+      //historyModalState.setSorted(column.key);
+   }
+
+   function generateColumns() {
+      let columns: IColumn[] = [];
+      if (state.mode === "leases") {
+         columns = [
+            {
+               key: 'type',
+               name: 'Type',
+               minWidth: 125,
+               maxWidth: 125,
+               isResizable: false,
+               isSorted: false,
+               isSortedDescending: false,
+               onColumnClick: onColumnClick
+            },
+            {
+               key: 'id',
+               name: 'ID',
+               minWidth: 185,
+               maxWidth: 185,
+               isResizable: false,
+               isSorted: false,
+               isSortedDescending: false,
+               onColumnClick: onColumnClick
+            },
+            {
+               key: 'description',
+               name: 'Description',
+               minWidth: 350,
+               isResizable: false,
+               isSorted: false,
+               isSortedDescending: false,
+               onColumnClick: onColumnClick
+            }
+         ];
+      }
+      else if (state.mode === "sessions") {
+         columns = [
+            {
+               key: 'id',
+               name: 'ID',
+               minWidth: 692,
+               maxWidth: 692,
+               isResizable: false,
+               isSorted: false,
+               isSortedDescending: false,
+               onColumnClick: onColumnClick
+            }
+         ];
+      }
+
+      if (state.mode === "leases" || state.mode === "sessions") {
+
+
+         columns.push({
+            key: 'startTime',
+            name: 'Start Time',
+            minWidth: 200,
+            maxWidth: 200,
+            isResizable: false,
+            isSorted: false,
+            isSortedDescending: false,
+            onColumnClick: onColumnClick
+         });
+         columns.push({
+            key: 'endTime',
+            name: 'Finish Time',
+            minWidth: 200,
+            maxWidth: 200,
+            isResizable: false,
+            isSorted: false,
+            isSortedDescending: false,
+            onColumnClick: onColumnClick
+         });
+      }
+
+      if (state.mode === "leases") {
+         columns.find(col => col.key === state.sortedLeaseColumn)!.isSorted = true;
+         columns.find(col => col.key === state.sortedLeaseColumn)!.isSortedDescending = state.sortedLeaseColumnDescending;
+      }
+      else if (state.mode === "sessions") {
+         columns.find(col => col.key === state.sortedSessionColumn)!.isSorted = true;
+         columns.find(col => col.key === state.sortedSessionColumn)!.isSortedDescending = state.sortedSessionColumnDescending;
+      }
+
+      return columns;
+   }
+
+   // main header
+   const onRenderDetailsHeader: IDetailsListProps['onRenderDetailsHeader'] = (props) => {
+      const customStyles: Partial<IDetailsHeaderStyles> = {};
+      customStyles.root = { paddingTop: 0 };
+      if (props) {
+         return (
+            <Sticky stickyPosition={StickyPositionType.Header} isScrollSynced={true}>
+               <DetailsHeader {...props} styles={customStyles} onRenderColumnHeaderTooltip={onRenderColumnHeaderTooltip} />
+            </Sticky>
+         );
+      }
+      return null;
+   };
+
+   const onRenderColumnHeaderTooltip: IDetailsHeaderProps['onRenderColumnHeaderTooltip'] = (props: any) => {
+      const customStyles: Partial<ITooltipHostStyles> = {};
+      if (props) {
+         customStyles.root = { selectors: { "> span": { paddingLeft: '8px !important', display: 'flex', justifyContent: 'center' } } };
+
+         /*
+          // no other way to filter children being centered other than to drill into private members??
+          if (props.children._owner.child?.child.child.child.elementType === "span") {
+              const data = props.children._owner.child.child.child.child.child.child.stateNode.data;
+              // if this cluster happens to be true, reset back to the default, because doing this the other way around
+              // takes too long when the columns update on details switch
+              // ugh :(
+              if (data === "ID" || data === "Description") {
+                  customStyles.root = {};
+              }
+          }
+          */
+         return <Text styles={customStyles}>{props.children}</Text>;
+      }
+      return null;
+   };
+
+   function onRenderHistoryItem(item: LeaseData | SessionData, index?: number, column?: IColumn) {
+
+      let link = "";
+
+      if (state.mode === "leases") {
+         const lease = item as LeaseData;
+         switch (column!.key) {
+            case 'type':
+
+               if (lease.batch) {
+
+                  const step = lease.batch.steps.find(s => s.outcome === JobStepOutcome.Failure);
+
+                  let name = lease.type;
+
+                  if (lease.batch.error !== JobStepBatchError.None) {
+                     name = lease.batch.error;
+                  } else if (step) {
+                     name = "StepError";
+                  }
+
+                  return <Stack styles={{ root: { height: '100%', } }} horizontal horizontalAlign={'start'} verticalAlign="center" tokens={{ childrenGap: 4 }}>
+                     <Stack >
+                        {!step && <BatchStatusIcon batch={lease.batch} />}
+                        {!!step && <StepStatusIcon step={step} />}
+                     </Stack>
+                     <Stack >{name}</Stack>
+                  </Stack>
+               } else if (lease.outcome) {
+                  return <Stack styles={{ root: { height: '100%', } }} horizontal horizontalAlign={'start'} verticalAlign="center" tokens={{ childrenGap: 4 }}>
+                     <Stack >
+                        <LeaseStatusIcon lease={lease} />
+                     </Stack>
+                     <Stack >{lease.type}</Stack>
+                  </Stack>
+               }
+               return <Stack styles={{ root: { height: '100%', } }} horizontal horizontalAlign={'center'}><Stack.Item align={"center"}>{lease.type}</Stack.Item></Stack>
+            case 'id':
+               link = "";
+               if (lease.details && 'LogId' in lease.details) {
+                  link = '/log/' + lease.details['LogId'];
+               }
+               else if (lease.logId) {
+                  link = `/log/${lease.logId}`;
+               }
+               if (link) {
+                  return <Stack style={{ height: "100%" }} verticalAlign="center"><Link to={link}>{lease.id}</Link></Stack>
+               }
+               return <Stack style={{ height: "100%" }} verticalAlign="center">{lease.id}</Stack>;
+            case 'name':
+               return <Stack>{lease.name}</Stack>;
+            case 'startTime':
+               return <Stack styles={{ root: { height: '100%', } }} horizontal horizontalAlign={'center'}><Stack.Item align={"center"}>{getShortNiceTime(lease.startTime, false, true, true)}</Stack.Item></Stack>
+            case 'endTime':
+               return <Stack styles={{ root: { height: '100%', } }} horizontal horizontalAlign={'center'}><Stack.Item align={"center"}>{getShortNiceTime(lease.finishTime, false, true, true)}</Stack.Item></Stack>
+            case 'description':
+               link = "";
+               let name = lease.name;
+               if (lease.details) {
+                  if ('jobId' in lease.details) {
+                     link = `/job/${lease.details['jobId']}`;
+                     if ('batchId' in lease.details) {
+                        link += `?batch=${lease.details['batchId']}`
+                     }
+                  }
+               }
+               else if (lease.jobId) {
+                  link = `/job/${lease.jobId}`;
+               }
+
+               if (!link && (lease.parentId && lease.details && lease.details["parentLogId"])) {
+                  link = `/log/${lease.details["parentLogId"]}`
+                  return <Stack styles={{ root: { height: '100%' } }} horizontal><Stack.Item align={"center"}><Text style={{ fontSize: "12px" }}>{name} (parent: </Text><Link style={{ fontSize: "12px" }} to={link}>{lease.parentId}</Link><Text style={{ fontSize: "12px" }}>)</Text></Stack.Item></Stack>;
+               }
+
+               if (link !== "") {
+                  return <Stack styles={{ root: { height: '100%' } }} horizontal><Stack.Item align={"center"}><Link style={{ fontSize: 12 }} key={"leaseText_" + lease.id} to={link}>{name}</Link></Stack.Item></Stack>;
+               }
+               else {
+                  return <Stack styles={{ root: { height: '100%', } }} horizontal><Stack.Item align={"center"}><Text styles={{ root: { fontSize: 12 } }} key={"leaseText_" + lease.id}>{name}</Text></Stack.Item></Stack>;
+               }
+            default:
+               return <span>{lease[column!.fieldName as keyof LeaseData] as string}</span>;
+         }
+      }
+      else if (state.mode === "sessions") {
+         const session = item as SessionData;
+         switch (column!.key) {
+            case 'id':
+               return <Stack styles={{ root: { height: '100%', } }} horizontal><Stack.Item align={"center"}>{session.id}</Stack.Item></Stack>
+            case 'startTime':
+               return <Stack styles={{ root: { height: '100%', } }} horizontal horizontalAlign={'center'}><Stack.Item align={"center"}>{getShortNiceTime(session.startTime, false, true)}</Stack.Item></Stack>
+            case 'endTime':
+               return <Stack styles={{ root: { height: '100%', } }} horizontal horizontalAlign={'center'}><Stack.Item align={"center"}>{getShortNiceTime(session.finishTime, false, true)}</Stack.Item></Stack>
+            default:
+               return <span>{session[column!.fieldName as keyof SessionData] as string}</span>;
+         }
+      }
+   }
+
+   const onRenderInfoDetailsHeader: IDetailsListProps['onRenderDetailsHeader'] = (props) => {
+      const customStyles: Partial<IDetailsHeaderStyles> = {};
+      customStyles.root = { paddingTop: 0 };
+      if (props) {
+         return (
+            <Sticky stickyPosition={StickyPositionType.Header} isScrollSynced={true}>
+               <DetailsHeader {...props} styles={customStyles} />
+            </Sticky>
+         );
+      }
+      return null;
+   };
+
+   const c1 = theme.horde.darkTheme ? theme.horde.neutralBackground : "#f3f2f1";
+   const c2 = theme.horde.darkTheme ? theme.horde.contentBackground : "#ffffff";
+
+   function onRenderBuilderInfoCell(nestingDepth?: number | undefined, item?: any, index?: number | undefined) {
+      return (
+         <Stack horizontal onClick={(ev) => { state.setInfoItemSelected(item); ev.preventDefault(); }} styles={{
+            root: {
+               background: item!.selected ? c1 : c2,
+               paddingLeft: 48 + (10 * nestingDepth!),
+               paddingTop: 8,
+               paddingBottom: 8,
+               selectors: {
+                  ":hover": {
+                     background: c1,
+                     cursor: 'pointer'
+                  }
+               }
+            }
+         }}>
+            <Stack>
+               <Link to="" onClick={(ev) => { state.setInfoItemSelected(item); ev.preventDefault(); }}><Text>{item!.name}</Text></Link>
+            </Stack>
+         </Stack>
+      );
+   }
+
+   const groups: IGroup[] = [
+      {
+         count: state.agentItemCount,
+         key: "agent",
+         name: "Agent",
+         startIndex: 0,
+         level: 0,
+         isCollapsed: false,
+      },
+      {
+         count: state.devicesItemCount,
+         key: "devices",
+         name: "Devices",
+         startIndex: state.agentItemCount,
+         level: 0,
+         isCollapsed: false,
+      },
+      {
+         count: state.workspaceItemCount,
+         key: "workspaces",
+         name: "Workspaces",
+         startIndex: state.agentItemCount + state.devicesItemCount,
+         level: 0,
+         isCollapsed: false,
+      },
+
+   ];
+
+   const onDismissHistoryModal = function (ev?: React.MouseEvent<HTMLButtonElement, MouseEvent> | undefined) {
+      state.setSelectedAgent(undefined);
+      onDismiss();
+   };
+
+   return (
+      <Dialog
+         modalProps={{
+            isBlocking: false,
+            topOffsetFixed: true,
+            styles: {
+               main: {
+                  hasBeenOpened: false,
+                  top: "80px",
+                  position: "absolute"
+               },
+               root: {
+                  selectors: {
+                     ".ms-Dialog-title": {
+                        paddingTop: '24px',
+                        paddingLeft: '32px'
+                     }
+                  }
+               }
+            }
+         }}
+         onDismiss={onDismissHistoryModal}
+         className={historyStyles.dialog}
+         hidden={!state.selectedAgent}
+         minWidth={1200}
+         dialogContentProps={{
+            type: DialogType.close,
+            onDismiss: onDismissHistoryModal,
+            title: `Agent ${state.selectedAgent?.name}`,
+         }}
+      >
+         <Stack>
+            {!!currentAction && <Dialog
+               hidden={false}
+               onDismiss={() => { setActionState({}) }}
+               minWidth={512}
+               dialogContentProps={{
+                  type: DialogType.normal,
+                  title: `Confirm Action`,
+               }}
+               modalProps={{ isBlocking: true }} >
+               <Stack style={{ paddingBottom: 18, paddingLeft: 4 }}>
+                  <Text>{currentAction.confirmText}</Text>
+               </Stack>
+               {!!currentAction.textInput && <TextField componentRef={actionTextInputRef} label={currentAction.name === "Edit Comment" ? "New Comment" : "Disable Reason - Required"} onChange={(ev, value) => setActionState({ ...actionState, comment: value })} spellCheck={false} autoComplete="false" />}
+               {currentAction.name === "Request Restart" && <Checkbox componentRef={forceRestartCheckboxRef} label={"Force Restart"} />}
+               <DialogFooter>
+                  <PrimaryButton disabled={actionState.confirmed || (!!currentAction.textInput && currentAction.name !== "Edit Comment" && !actionState.comment?.length)} onClick={() => { setActionState({ ...actionState, confirmed: true, comment: actionTextInputRef.current?.value }) }} text={currentAction.name} />
+                  <DefaultButton disabled={actionState.confirmed} onClick={() => { setActionState({}) }} text="Cancel" />
+               </DialogFooter>
+            </Dialog>
+            }
+            <Stack horizontal styles={{ root: { paddingLeft: 2 } }}>
+               <Stack grow>
+                  <Stack horizontalAlign={"start"}>
+                     <Pivot className={hordeClasses.pivot}
+                        onLinkClick={(item?: PivotItem | undefined) => { state.setMode(item?.props.itemKey); }}
+                        linkSize="normal"
+                        linkFormat="links"
+                        defaultSelectedKey={state.mode ?? "info"}
+                     >
+                        <PivotItem headerText="Info" itemKey="info" />
+                        <PivotItem headerText="Sessions" itemKey="sessions" />
+                        <PivotItem headerText="Leases" itemKey="leases" />
+                        <PivotItem headerText="Telemetry" itemKey="telemetry" />
+                     </Pivot>
+                  </Stack>
+               </Stack>
+               {!!poolItems.length && <Stack style={{ paddingRight: 24 }}><DefaultButton text="Pools" menuProps={{ shouldFocusOnMount: true, items: poolItems }} /></Stack>}
+               {!!dashboard.user?.dashboardFeatures?.showRemoteDesktop && <Stack> <Stack horizontal tokens={{ childrenGap: 24 }}>
+                  <Stack>
+                     <DefaultButton text="Actions" menuProps={actionMenuProps} />
+                  </Stack>
+                  <Stack>
+                     <Link to={`/audit/agent/${encodeURIComponent(agentId)}`} ><DefaultButton text="Audit" onClick={(ev) => { }} /></Link>
+                  </Stack>
+
+                  <Stack>
+                     <DefaultButton text="Remote Desktop" onClick={() => {
+                        if (state.selectedAgent?.id) {
+                           // @todo: this assumes the id is the ip, should have a private ip property on agents instead
+                           let ip = state.selectedAgent?.id;
+                           if (ip.startsWith("10-")) {
+                              ip = ip.replaceAll("-", ".")
+                           }
+                           window.open(`ugs://rdp?host=${ip}`, "_self")
+                        }
+                     }} />
+                  </Stack>
+               </Stack>
+               </Stack>}
+            </Stack>
+            <Stack styles={{ root: { paddingTop: 30 } }}>
+               {state.mode === "info" && <Stack horizontal tokens={{ childrenGap: 20 }}>
+                  <Stack styles={{ root: { width: 300 } }}>
+                     <GroupedList
+                        items={state.infoItems}
+                        compact={true}
+                        onRenderCell={onRenderBuilderInfoCell}
+                        groups={groups}
+                        selection={new Selection()}
+                        selectionMode={SelectionMode.none}
+                        groupProps={{
+                           showEmptyGroups: true,
+                           onRenderHeader: (props) => {
+                              return <Link to="" onClick={(ev) => { ev.preventDefault(); props!.onToggleCollapse!(props!.group!); }}> <GroupHeader {...props} /></Link>;
+                           },
+                           headerProps: {
+                              styles: {
+                                 title: {
+                                    fontFamily: "Horde Open Sans Semibold",
+                                    color: modeColors.text,
+                                    paddingLeft: 0
+
+                                 },
+                                 headerCount: {
+                                    display: 'none'
+                                 }
+                              }
+                           }
+                        }}
+                     />
+                  </Stack>
+                  <Stack styles={{ root: { width: '100%' } }}>
+                     <Stack.Item className={hordeClasses.relativeModalSmall}>
+                        <ScrollablePane scrollbarVisibility={ScrollbarVisibility.auto} styles={{ contentContainer: { overflowX: 'hidden' } }}>
+                           <DetailsList
+                              className={historyStyles.detailsList}
+                              compact={true}
+                              items={state.infoSubItems}
+                              columns={[
+                                 { key: 'column1', name: 'Name', fieldName: 'name', minWidth: 200, maxWidth: 200, isResizable: false },
+                                 { key: 'column2', name: 'Value', fieldName: 'value', minWidth: 200, isResizable: false }
+                              ]}
+                              layoutMode={DetailsListLayoutMode.justified}
+                              onRenderDetailsHeader={onRenderInfoDetailsHeader}
+                              constrainMode={ConstrainMode.unconstrained}
+                              selectionMode={SelectionMode.none}
+                              onRenderRow={(props) => {
+
+                                 if (props) {
+                                    return <DetailsRow styles={{ cell: { whiteSpace: "pre-line", overflowWrap: "break-word" } }} {...props} />
+                                 }
+
+                                 return null;
+                              }}
+                           />
+                        </ScrollablePane>
+                     </Stack.Item>
+                  </Stack>
+               </Stack>
+               }
+               {(state.mode === "leases" || state.mode === "sessions") &&
+                  <Stack.Item className={hordeClasses.relativeModalSmall}>
+                     <ScrollablePane scrollbarVisibility={ScrollbarVisibility.auto} styles={{ contentContainer: { overflowX: 'hidden' } }}>
+                        <DetailsList
+                           className={historyStyles.detailsList}
+                           compact={true}
+                           items={state.currentData}
+                           columns={generateColumns()}
+                           onRenderDetailsHeader={onRenderDetailsHeader}
+                           onRenderItemColumn={onRenderHistoryItem}
+                           layoutMode={DetailsListLayoutMode.justified}
+                           constrainMode={ConstrainMode.unconstrained}
+                           selectionMode={SelectionMode.none}
+                           listProps={{ renderedWindowsAhead: 1, renderedWindowsBehind: 1 }}
+                           onRenderMissingItem={() => { state.nullItemTrigger(); return <div></div> }}
+                        />
+                     </ScrollablePane>
+                  </Stack.Item>}
+               {state.mode === "telemetry" && <TelemetryPanel />}
+            </Stack>
+         </Stack>
+         <DialogFooter>
+            <PrimaryButton onClick={() => { state.setSelectedAgent(undefined); onDismiss() }}>Close</PrimaryButton>
+         </DialogFooter>
+      </Dialog>
+   );
+});
\ No newline at end of file
diff --git a/Engine/Source/Programs/Horde/HordeDashboard/src/components/agents/HistoryModalStore.ts b/Engine/Source/Programs/Horde/HordeDashboard/src/components/agents/HistoryModalStore.ts
new file mode 100644
--- /dev/null
+++ b/Engine/Source/Programs/Horde/HordeDashboard/src/components/agents/HistoryModalStore.ts
@@ -0,0 +1,346 @@
+import { action, makeObservable, observable } from "mobx";
+import { AgentData, GetAgentLeaseResponse, GetAgentSessionResponse, GetAgentTelemetrySampleResponse, GetJobResponse, LeaseData, SessionData } from "../../backend/Api" 
+import backend from "../../backend";
+
+type LeaseTooltip = {
+   lease?: GetAgentLeaseResponse;
+   time?: Date;
+   id?: string;
+   x?: number;
+   y?: number;
+   frozen?: boolean;
+}
+
+type InfoPanelItem = {
+   key: string;
+   name: string;
+   data?: string;
+   selected: boolean;
+};
+
+type InfoPanelSubItem = {
+   name: string;
+   value: string;
+}
+
+export class HistoryModalStore {
+   constructor() {
+      makeObservable(this);
+   }
+
+   @observable
+   selectedAgentUpdated = 0;
+
+   @observable
+   currentLeaseTooltip?: LeaseTooltip;
+
+   get selectedAgent(): AgentData | undefined {
+      // subscribe in any observers
+      if (this.selectedAgentUpdated) { }
+      return this._selectedAgent;
+   }
+
+   private _selectedAgent: AgentData | undefined = undefined;
+   @observable.shallow currentData: any = [];
+   @observable.shallow infoItems: InfoPanelItem[] = [];
+   @observable.shallow infoSubItems: InfoPanelSubItem[] = [];
+   agentItemCount: number = 0;
+   devicesItemCount: number = 0;
+   workspaceItemCount: number = 0;
+   sortedLeaseColumn = "startTime";
+   sortedLeaseColumnDescending = true;
+   sortedSessionColumn = "startTime";
+   sortedSessionColumnDescending = true;
+
+   @observable mode: string | undefined = undefined;
+   modeCurrentIndex: number = 0;
+   bUpdatedQueued: boolean = false;
+
+   // sets the pool editor dialog open
+   @action
+   setSelectedAgent(selectedAgent?: AgentData | undefined) {
+      // if we're closing don't reset
+      this._selectedAgent = selectedAgent;
+      if (this.selectedAgent) {
+         // set date to now
+         this._initBuilderItems();
+         if (!this.mode) {
+            this.setMode("info");
+         }
+         else {
+            this.setMode(this.mode, true);
+         }
+         this._doUpdate();
+      }
+
+      this.selectedAgentUpdated++;
+   }
+
+   @action
+   setMode(newMode?: string, force?: boolean) {
+      if (newMode) {
+         newMode = newMode.toLowerCase();
+         if (this.mode !== newMode || force) {
+            this.mode = newMode;
+            this.currentData = [];
+            this.modeCurrentIndex = 0;
+            this._doUpdate();
+         }
+      }
+   }
+
+   @action
+   setLeaseTooltip(tooltip?: LeaseTooltip) {
+      this.currentLeaseTooltip = tooltip;
+   }
+
+   setInfoItemSelected(item: InfoPanelItem) {
+      this.infoItems.forEach(infoItem => {
+         if (infoItem.key === item.key) {
+            infoItem.selected = true;
+         }
+         else {
+            infoItem.selected = false;
+         }
+      });
+      this.setInfoItems(this.infoItems);
+   }
+
+   private _initBuilderItems() {
+      this.agentItemCount = 0;
+      this.devicesItemCount = 0;
+      this.workspaceItemCount = 0;
+      const items: InfoPanelItem[] = [
+         {
+            key: "overview",
+            name: "Overview",
+            selected: false
+         }
+      ];
+      this.agentItemCount = 1;
+      if (this.selectedAgent) {
+         if (this.selectedAgent.capabilities?.devices) {
+            for (const deviceIdx in this.selectedAgent.capabilities?.devices) {
+               const device = this.selectedAgent.capabilities?.devices[deviceIdx];
+               if (!device["name"]) {
+                  device["name"] = "Primary";
+               }
+               if (device.properties) {
+                  items.push({ key: `device${deviceIdx}`, name: device["name"], selected: device["name"] === "Primary" ? true : false });
+                  this.devicesItemCount++;
+               }
+            }
+         }
+         if (this.selectedAgent.workspaces) {
+            const streamCount = new Map<string, number>();
+            for (const workspaceIdx in this.selectedAgent.workspaces) {
+               const workspace = this.selectedAgent.workspaces[workspaceIdx];
+               let count = streamCount.get(workspace.stream) ?? 0;
+               if (!count) {
+                  count++;
+                  streamCount.set(workspace.stream, 1);
+               } else {
+                  count++;
+                  streamCount.set(workspace.stream, count);
+               }
+
+               const name: string = count > 1 ? `${workspace.stream} (${count})` : workspace.stream;
+
+               (workspace as any)._hackName = name;
+               ;
+               items.push({ key: `workspace${workspaceIdx}`, name: name, selected: false, data: name });
+               this.workspaceItemCount++;
+            }
+         }
+      }
+      this.setInfoItems(items);
+   }
+
+   @action
+   setInfoItems(items: InfoPanelItem[]) {
+      const subItems: InfoPanelSubItem[] = [];
+      if (this.selectedAgent) {
+         const selectedItem = items.find(item => item.selected);
+         if (selectedItem) {
+            if (selectedItem.key === "overview") {
+               subItems.push({ name: 'Enabled', value: this.selectedAgent.enabled.toString() });
+               subItems.push({ name: 'Comment', value: this.selectedAgent.comment ?? "None" });
+               subItems.push({ name: 'Ephemeral', value: this.selectedAgent.ephemeral.toString() });
+               subItems.push({ name: 'ForceVersion', value: this.selectedAgent.forceVersion ?? "None" });
+               subItems.push({ name: 'Id', value: this.selectedAgent.id });
+               subItems.push({ name: 'Online', value: this.selectedAgent.online.toString() });
+               subItems.push({ name: 'Last Update', value: this.selectedAgent.updateTime.toString() });
+               subItems.push({ name: 'Version', value: this.selectedAgent.version ?? "None" });
+            }
+            else if (selectedItem.key.indexOf("workspace") !== -1) {
+               if (this.selectedAgent.workspaces) {
+                  for (const workspaceIdx in this.selectedAgent.workspaces) {
+                     const workspace = this.selectedAgent.workspaces[workspaceIdx];
+
+                     if ((workspace as any)._hackName === selectedItem.data) {
+                        subItems.push({ name: 'Identifier', value: workspace.identifier });
+                        subItems.push({ name: 'Stream', value: workspace.stream });
+                        subItems.push({ name: 'Incremental', value: workspace.bIncremental.toString() });
+                        workspace.serverAndPort && subItems.push({ name: 'Server and Port', value: workspace.serverAndPort });
+                        workspace.userName && subItems.push({ name: 'Username', value: workspace.userName });
+                        workspace.password && subItems.push({ name: 'Password', value: workspace.password });
+                        workspace.view && subItems.push({ name: 'View', value: workspace.view?.join("\n") });
+                     }
+                  }
+               }
+            }
+            else if (selectedItem.key.indexOf("device") !== -1) {
+               if (this.selectedAgent.capabilities?.devices) {
+                  for (const deviceIdx in this.selectedAgent.capabilities?.devices) {
+                     const device = this.selectedAgent.capabilities?.devices[deviceIdx];
+                     if (device.name === selectedItem.name) {
+                        if (device.properties) {
+                           for (const propIdx in device.properties) {
+                              const prop = device.properties[propIdx];
+                              const subItemData = prop.split('=');
+                              if (subItemData[0].indexOf("RAM") !== -1) {
+                                 subItemData[1] += " GB";
+                              }
+                              else if (subItemData[0].indexOf("Disk") !== -1) {
+                                 subItemData[1] = (Number(subItemData[1]) / 1073741824).toLocaleString(undefined, { maximumFractionDigits: 0 }) + " GiB";
+                              }
+                              subItems.push({ name: subItemData[0], value: subItemData[1] });
+                           }
+                        }
+                     }
+                  }
+               }
+            }
+         }
+      }
+      this.infoItems = [...items];
+      this.infoSubItems = subItems;
+   }
+
+   private _doUpdate() {
+      if (this.mode === "leases") {
+         this.UpdateLeases();
+      }
+      else if (this.mode === "sessions") {
+         this.UpdateSessions();
+      }
+   }
+
+   nullItemTrigger() {
+      if (!this.bUpdatedQueued) {
+         this.bUpdatedQueued = true;
+         this._doUpdate();
+      }
+   }
+
+   UpdateLeases() {
+      const data: LeaseData[] = [];
+      backend.getLeases(this.selectedAgent!.id, this.modeCurrentIndex, 30, true).then(responseData => {
+         responseData.forEach((dataItem: GetAgentLeaseResponse) => {
+            data.push(dataItem);
+         });
+         this.appendData(data);
+      });
+   }
+
+   UpdateSessions() {
+      const data: SessionData[] = [];
+      backend.getSessions(this.selectedAgent!.id, this.modeCurrentIndex, 30).then(responseData => {
+         responseData.forEach((dataItem: GetAgentSessionResponse) => {
+            data.push(dataItem);
+         });
+         this.appendData(data);
+      });
+   }
+
+   jobData = new Map<string, GetJobResponse | boolean>();
+
+   UpdateTelemetry(minutes: number) {
+
+      const requests = [backend.getAgentTelemetry(this.selectedAgent!.id, new Date(Date.now() - 1000 * 60 * minutes), new Date()), backend.getAgentLeases(this.selectedAgent!.id, new Date(Date.now() - 1000 * 60 * minutes), new Date())]
+      Promise.all(requests).then((data) => {
+         this.setTelemetryData((data[0] as GetAgentTelemetrySampleResponse[]).reverse(), data[1] as GetAgentLeaseResponse[]);
+      })
+   }
+
+   @action
+   setTelemetryData(telemetryData: GetAgentTelemetrySampleResponse[], sessionData: GetAgentLeaseResponse[]) {
+
+      const linear: GetAgentTelemetrySampleResponse[] = [];
+
+      if (telemetryData.length) {
+
+         let ctime: number = 0;
+
+         telemetryData.forEach(t => {
+
+            if (!ctime) {
+               ctime = t.time.getTime();
+               linear.push(t);
+            } else {
+               const ntime = t.time.getTime();
+               // dependant on agent sample rate!               
+               while ((ntime - ctime) > 60000) {
+
+                  ctime += 30000;
+                  linear.push({
+                     time: new Date(ctime),
+                     userCpu: 0,
+                     idleCpu: 0,
+                     systemCpu: 0,
+                     freeRam: 0,
+                     usedRam: 0,
+                     totalRam: 1,
+                     freeDisk: 0,
+                     totalDisk: 0
+                  });
+               }
+
+               linear.push(t);
+               ctime = ntime;
+            }
+         });
+      }
+
+      this.currentData = [linear, sessionData];
+
+   }
+
+   @action
+   appendData(newData: any[]) {
+
+      if (this.mode === "telemetry") {
+         return;
+      }
+
+      // if there's any data, there might be more data next time, so add another callback.
+      if (newData.length > 0) {
+         newData.push(null);
+      }
+
+      let combinedData = [...this.currentData];
+
+      // remove previous null if it exists
+      if (combinedData[combinedData.length - 1] === null) {
+         combinedData.splice(-1, 1);
+      }
+      // add all the new data
+      Array.prototype.push.apply(combinedData, newData);
+
+      const dedupe = new Set<string>();
+
+      this.currentData = combinedData.filter(d => {
+         if (dedupe.has(d?.id)) {
+            return false;
+         }
+         dedupe.add(d?.id);
+         return true;
+      });
+
+      this.modeCurrentIndex += newData.length;
+      this.bUpdatedQueued = false;
+   }
+
+}
+
+export const historyModalStore = new HistoryModalStore();
\ No newline at end of file
diff --git a/Engine/Source/Programs/Horde/HordeDashboard/src/components/jobDetailsV2/JobDetailStepTrendsV2.tsx b/Engine/Source/Programs/Horde/HordeDashboard/src/components/jobDetailsV2/JobDetailStepTrendsV2.tsx
--- a/Engine/Source/Programs/Horde/HordeDashboard/src/components/jobDetailsV2/JobDetailStepTrendsV2.tsx
+++ b/Engine/Source/Programs/Horde/HordeDashboard/src/components/jobDetailsV2/JobDetailStepTrendsV2.tsx
@@ -13,7 +13,7 @@ import dashboard, { StatusColor } from "../../backend/Dashboard";
 import { ISideRailLink } from "../../base/components/SideRail";
 import { displayTimeZone, getElapsedString, getHumanTime, msecToElapsed } from "../../base/utilities/timeUtils";
 import { ChangeButton } from "../ChangeButton";
-import { HistoryModal } from "../HistoryModal";
+import { HistoryModal } from "../agents/HistoryModal";
 import { StepRefStatusIcon } from "../StatusIcon";
 import { JobDataView, JobDetailsV2 } from "./JobDetailsViewCommon";
 import { getHordeStyling } from "../../styles/Styles";
diff --git a/Engine/Source/Programs/Horde/HordeDashboard/src/components/jobDetailsV2/JobDetailTimeline.tsx b/Engine/Source/Programs/Horde/HordeDashboard/src/components/jobDetailsV2/JobDetailTimeline.tsx
--- a/Engine/Source/Programs/Horde/HordeDashboard/src/components/jobDetailsV2/JobDetailTimeline.tsx
+++ b/Engine/Source/Programs/Horde/HordeDashboard/src/components/jobDetailsV2/JobDetailTimeline.tsx
@@ -11,7 +11,7 @@ import { GetBatchResponse, GetStepResponse, JobStepBatchState, JobStepOutcome, S
 import dashboard, { StatusColor } from "../../backend/Dashboard";
 import { ISideRailLink } from "../../base/components/SideRail";
 import { getShortNiceTime, msecToElapsed } from "../../base/utilities/timeUtils";
-import { HistoryModal } from "../HistoryModal";
+import { HistoryModal } from "../agents/HistoryModal";
 import { JobDataView, JobDetailsV2 } from "./JobDetailsViewCommon";
 import { getHordeStyling } from "../../styles/Styles";
 
diff --git a/Engine/Source/Programs/Horde/HordeDashboard/src/components/jobDetailsV2/JobDetailViewSteps.tsx b/Engine/Source/Programs/Horde/HordeDashboard/src/components/jobDetailsV2/JobDetailViewSteps.tsx
--- a/Engine/Source/Programs/Horde/HordeDashboard/src/components/jobDetailsV2/JobDetailViewSteps.tsx
+++ b/Engine/Source/Programs/Horde/HordeDashboard/src/components/jobDetailsV2/JobDetailViewSteps.tsx
@@ -6,7 +6,7 @@ import { BatchData, GetBatchResponse, JobStepBatchError, JobStepBatchState, JobS
 import dashboard, { StatusColor } from "../../backend/Dashboard";
 import { ISideRailLink } from "../../base/components/SideRail";
 import { getBatchInitElapsed, getNiceTime, getStepElapsed, getStepETA, getStepFinishTime, getStepPercent, getStepStartTime, getStepTimingDelta, HordeTime } from "../../base/utilities/timeUtils";
-import { HistoryModal } from "../HistoryModal";
+import { HistoryModal } from "../agents/HistoryModal";
 import { getBatchText, getStepStatusMessage } from "../JobDetailCommon";
 import { StepStatusIcon } from "../StatusIcon";
 import { LabelsPanelV2 } from "./JobDetailLabels";
diff --git a/Engine/Source/Programs/Horde/HordeDashboard/src/components/jobDetailsV2/StepDetailView.tsx b/Engine/Source/Programs/Horde/HordeDashboard/src/components/jobDetailsV2/StepDetailView.tsx
--- a/Engine/Source/Programs/Horde/HordeDashboard/src/components/jobDetailsV2/StepDetailView.tsx
+++ b/Engine/Source/Programs/Horde/HordeDashboard/src/components/jobDetailsV2/StepDetailView.tsx
@@ -7,7 +7,7 @@ import { useNavigate } from "react-router-dom";
 import backend from "../../backend";
 import { GetArtifactResponse } from "../../backend/Api";
 import { getHordeStyling } from "../../styles/Styles";
-import { HistoryModal } from "../HistoryModal";
+import { HistoryModal } from "../agents/HistoryModal";
 import { useQuery } from "../JobDetailCommon";
 import { BisectionPanel } from "./JobDetailBisection";
 import { HealthPanel } from "./JobDetailHealthV2";
