Horde Dashboard: Import build error fix

diff --git a/Engine/Source/Programs/Horde/HordeDashboard/src/components/jobDetailsV2/JobDetailStepHistoryV2.tsx b/Engine/Source/Programs/Horde/HordeDashboard/src/components/jobDetailsV2/JobDetailStepHistoryV2.tsx
--- a/Engine/Source/Programs/Horde/HordeDashboard/src/components/jobDetailsV2/JobDetailStepHistoryV2.tsx
+++ b/Engine/Source/Programs/Horde/HordeDashboard/src/components/jobDetailsV2/JobDetailStepHistoryV2.tsx
@@ -13,7 +13,7 @@ import dashboard, { StatusColor } from "../../backend/Dashboard";
 import { ISideRailLink } from "../../base/components/SideRail";
 import { displayTimeZone, getElapsedString, getHumanTime, msecToElapsed } from "../../base/utilities/timeUtils";
 import { ChangeButton } from "../ChangeButton";
-import { HistoryModal } from "../HistoryModal";
+import { HistoryModal } from "../agents/HistoryModal";
 import { StepRefStatusIcon } from "../StatusIcon";
 import { JobDataView, JobDetailsV2 } from "./JobDetailsViewCommon";
 import { getHordeStyling } from "../../styles/Styles";
