Add individual test result page with phase details
- Introduce Phase session status model
- Introduce a concentric circles representation of test session failure rate in test history
- Add additional cache mechanic
- Add minimal support for artifacts fetching

#rnx

diff --git a/Engine/Source/Programs/Horde/HordeDashboard/plugins/test-automation/api.ts b/Engine/Source/Programs/Horde/HordeDashboard/plugins/test-automation/api.ts
--- a/Engine/Source/Programs/Horde/HordeDashboard/plugins/test-automation/api.ts
+++ b/Engine/Source/Programs/Horde/HordeDashboard/plugins/test-automation/api.ts
@@ -9,8 +9,8 @@ export type GetTestDataResponse = {
 	/// The key associated with this test data
 	key: string;
 
-	/// The change id related to the job
-	change: number;
+	/// The commit id related to the job
+	commitId: {name: string, order: number};
 
 	/// The job id of this test data
 	jobId: string;
@@ -30,6 +30,67 @@ export type TestData = GetTestDataResponse & {
     data: object;
 }
 
+/// An event entry
+export type EventEntry = {
+
+    time: string;
+
+    id: number;
+
+    level: EventLevel;
+
+    message: string;
+
+    format: string;
+
+    properties?: Record<string, string | number | Record<string, string | number | undefined> | undefined>;
+}
+
+/// The event level
+export enum EventLevel {
+
+    Unspecified = "Unspecified",
+
+    Information = "Information",
+
+    Warning = "Warning",
+
+    Error = "Error",
+
+    Critical = "Critical"
+}
+
+/// Known property types
+export enum LogPropertyType {
+    ImageCompare = "Image Compare"
+}
+
+/// Image Compare detailed properties
+export type ImageCompare = {
+    "$type": LogPropertyType.ImageCompare;
+    "$text": string;
+
+    //** The path to the unapproved image from the artifacts */
+    unapproved: string;
+
+    //** The path to the unapproved metadata associated with */
+    //** the image from the artifacts */
+    unapproved_metadata?: string;
+
+    //** The path to the approved image from the artifacts */
+    approved?: string;
+
+    //** The path to the approved metadata associated with */
+    //** the image from the artifacts */
+    approved_metadata?: string;
+
+    //** The path to the delta image from the artifacts */
+    difference?: string;
+
+    //** The path to the json report associated with the image from the artifacts */
+    difference_report?: string;
+}
+
 /// A test name request
 export type GetTestsRequest = {
 
@@ -172,7 +233,7 @@ export type GetTestPhaseResponse = {
     testKey: string;
 
     /// The test phases
-    testPhases: GetPhaseResponse[];
+    phases: GetPhaseResponse[];
 }
 
 /// Test outcome
@@ -236,8 +297,8 @@ export type GetTestSessionResponse = {
     /// The tag ids associated with the test session
     tagIds?: string[]
 
-    /// The duration of the test session
-    duration: string;
+    /// The duration of the test session in seconds
+    duration: number;
 
     /// The date time at which the test session started
     startDateTime: string;
@@ -288,8 +349,8 @@ export type GetTestPhaseSessionResponse = {
     /// The commit order associated with the phase session
     commitOrder: number;
 
-    /// The duration of the phase session
-    duration: string;
+    /// The duration of the phase session in seconds
+    duration: number;
 
     /// The date time at which the phase session started
     startDateTime: string;
@@ -316,6 +377,90 @@ export type GetTestPhaseSessionResponse = {
     errorFingerprint?: string;
 }
 
+export type TestSessionSummaryResponse = {
+
+    /// The name of the test associated with the session
+    testName: string;
+
+    /// The date time at which the session started
+    dateTime: string;
+    
+    /// The time duration of the session
+    timeElapseSec: number;
+    
+    /// The total count of the phases to run during the session
+    phasesTotalCount: number;
+    
+    /// The count of the phases that succeeded
+    phasesSucceededCount: number;
+    
+    /// The count of the phases that which state stayed undefined
+    phasesUndefinedCount: number;
+    
+    /// The count of the phases that failed
+    phasesFailedCount: number;
+}
+
+export type PhaseSessionDetailsResponse = {
+    
+    /// The name of the phase
+    name: string;
+    
+    /// The key of the phase
+    key: string;
+    
+    /// The date time at which the phase started
+    dateTime: string;
+    
+    /// The duration of the phase
+    timeElapseSec: number;
+    
+    /// The outcome of the phase run
+    outcome: TestPhaseOutcome;
+
+    /// The device keys associated with this phase
+    deviceKeys: string[];
+    
+    /// The error fingerprint associated with the phase failure
+    errorFingerprint?: string;
+    
+    /// Whether or not the phase has warning reported
+    hasWarning?: boolean;
+    
+    /// The event stream path that details the phase events
+    eventStreamPath?: string;
+}
+
+export type DeviceDetailsResponse = {
+
+    /// The name of the device
+    name: string;
+
+    /// The path to the log artifact
+    appInstanceLogPath: string;
+
+    /// The metadata associated with the device
+    metadata: Record<string, string>;
+}
+
+export type TestSessionDetailsResponse = {
+
+    /// The test session summary
+    summary: TestSessionSummaryResponse;
+
+    /// The phase sessions run during the test session
+    phases: PhaseSessionDetailsResponse[];
+
+    /// The device details utilized during this test session
+    devices: Record<string, DeviceDetailsResponse>;
+
+    /// The metadata associated with the test session
+    metadata: Record<string, string>;
+
+    /// The tags associated with the test session
+    tags?: string[];
+}
+
 export async function getTestDataV2(id: string): Promise<TestData> {
     return new Promise<TestData>((resolve, reject) => {
         backend.fetch.get(`/api/v3/testdata/${id}`, {
diff --git a/Engine/Source/Programs/Horde/HordeDashboard/plugins/test-automation/testAutomationCommon.tsx b/Engine/Source/Programs/Horde/HordeDashboard/plugins/test-automation/testAutomationCommon.tsx
--- a/Engine/Source/Programs/Horde/HordeDashboard/plugins/test-automation/testAutomationCommon.tsx
+++ b/Engine/Source/Programs/Horde/HordeDashboard/plugins/test-automation/testAutomationCommon.tsx
@@ -1,9 +1,10 @@
 import { mergeStyles, mergeStyleSets, Text, ComboBox, IComboBoxOption, SelectableOptionMenuItemType } from "@fluentui/react";
-import { TestSessionResult } from "./testData";
+import { TestPhaseStatus, TestSessionResult } from "./testData";
 import dashboard, { StatusColor } from "horde/backend/Dashboard";
 import { useEffect, useState } from "react";
+import { TestPhaseOutcome } from "./api";
 
-export const MultiOptionChooser: React.FC<{options: IComboBoxOption[], initialSelection: string[], updateKeys: (selectedKeys: string[]) => void, placeholder?: string, style?: any }> = ({ options, initialSelection, updateKeys, placeholder, style }) => {
+export const MultiOptionChooser: React.FC<{options: IComboBoxOption[], initialSelection: string[], updateKeys: (selectedKeys: string[]) => void, placeholder?: string, style?: any, disabled?: boolean }> = ({ options, initialSelection, updateKeys, placeholder, style, disabled }) => {
     const [selection, setSelection] = useState<string[]>(initialSelection);
     const optionsWithSelectAll = [...options];
 
@@ -20,7 +21,7 @@ export const MultiOptionChooser: React.FC<{options: IComboBoxOption[], initialSe
 
     optionsWithSelectAll.unshift({ key: 'selectAll', text: 'Select All', itemType: SelectableOptionMenuItemType.SelectAll });
 
-    return <ComboBox styles={{ root: style }} placeholder={placeholder ?? "None"} selectedKey={selection} multiSelect options={optionsWithSelectAll}
+    return <ComboBox styles={{ root: style }} placeholder={placeholder ?? "None"} selectedKey={selection} multiSelect options={optionsWithSelectAll} disabled={disabled}
         onChange={(_, option) => {
             if (option) {
                 if (option.itemType === SelectableOptionMenuItemType.SelectAll) {
@@ -53,7 +54,7 @@ const styles = mergeStyleSets({
 export type StatusBarStack = {
     value: number,
     title?: string,
-    titleValue?: number,
+    titleValue?: string,
     color?: string,
     onClick?: () => void,
     stripes?: boolean,
@@ -64,7 +65,7 @@ export const StatusBar = (stack: StatusBarStack[], width: number, height: number
     stack = stack.filter(s => s.value > 0);
 
     const mainTitle = stack.map((item) => {
-        return item.titleValue === undefined ? `${Math.ceil(item.value)}% ${item.title}` : `${item.titleValue} ${item.title}`
+        return !item.titleValue ? `${Math.ceil(item.value)}% ${item.title}` : `${item.titleValue} ${item.title}`
     }).join(' ');
 
     return (
@@ -85,10 +86,10 @@ export const StatusBar = (stack: StatusBarStack[], width: number, height: number
 
 export const SessionStatusBar = (session: TestSessionResult, width: number, height: number): JSX.Element | undefined => {
     const sessionTotalCount = session.phaseSucceededCount + session.phaseFailedCount + session.phaseUnspecifiedCount;
-    const metaFailedFactor = Math.ceil(session.phaseFailedCount / (sessionTotalCount || 1) * 20) / 20;
-    const metaUnspecifiedFactor = Math.ceil(session.phaseUnspecifiedCount / (sessionTotalCount || 1) * 20) / 20;
+    const metaFailedFactor = Math.ceil(session.phaseFailedCount / (sessionTotalCount || 1) * 50) / 50;
+    const metaUnspecifiedFactor = Math.ceil(session.phaseUnspecifiedCount / (sessionTotalCount || 1) * 50) / 50;
 
-    const bDisplayBar = (metaFailedFactor + metaUnspecifiedFactor) > 0 && (metaFailedFactor + metaUnspecifiedFactor) < 1;
+    const bDisplayBar = (metaFailedFactor + metaUnspecifiedFactor) > 0 && metaFailedFactor < 1 && metaUnspecifiedFactor < 1;
     if (!bDisplayBar) return undefined
 
     const statusColors = dashboard.getStatusColors();
@@ -97,12 +98,14 @@ export const SessionStatusBar = (session: TestSessionResult, width: number, heig
         {
             value: metaUnspecifiedFactor * 100,
             title: "Unspecified",
+            titleValue: `${ Math.ceil(session.phaseUnspecifiedCount / (sessionTotalCount || 1) * 100)}%`,
             color: statusColors.get(StatusColor.Unspecified)!,
             stripes: true
         },
         {
             value: metaFailedFactor * 100,
             title: "Failure",
+            titleValue: `${ Math.ceil(session.phaseFailedCount / (sessionTotalCount || 1) * 100)}%`,
             color: statusColors.get(StatusColor.Failure)!,
             stripes: true
         }
@@ -111,7 +114,7 @@ export const SessionStatusBar = (session: TestSessionResult, width: number, heig
     return StatusBar(stack, width, height, statusColors.get(StatusColor.Success)!, { margin: '3px !important' });
 }
 
-export const StatusPie = (stack: StatusBarStack[], radius: number, basecolor?: string, style?: any): JSX.Element => {
+export const StatusPie = (stack: StatusBarStack[], radius: number, style?: any): JSX.Element => {
 
     const width = radius * 2;
     const height = width;
@@ -144,8 +147,8 @@ export const StatusPie = (stack: StatusBarStack[], radius: number, basecolor?: s
 
 export const SessionStatusPie = (session: TestSessionResult, radius: number): JSX.Element | undefined => {
     const sessionTotalCount = session.phaseSucceededCount + session.phaseFailedCount + session.phaseUnspecifiedCount;
-    const metaFailedFactor = Math.ceil(session.phaseFailedCount / (sessionTotalCount || 1) * 20) / 20;
-    const metaUnspecifiedFactor = Math.ceil(session.phaseUnspecifiedCount / (sessionTotalCount || 1) * 20) / 20;
+    const metaFailedFactor = Math.ceil(session.phaseFailedCount / (sessionTotalCount || 1) * 100) / 100;
+    const metaUnspecifiedFactor = Math.ceil(session.phaseUnspecifiedCount / (sessionTotalCount || 1) * 100) / 100;
 
     const bDisplayBar = (metaFailedFactor + metaUnspecifiedFactor) > 0 && (metaFailedFactor + metaUnspecifiedFactor) < 1;
     if (!bDisplayBar) return;
@@ -156,14 +159,12 @@ export const SessionStatusPie = (session: TestSessionResult, radius: number): JS
         {
             value: metaUnspecifiedFactor * 100,
             title: "Unspecified",
-            color: statusColors.get(StatusColor.Unspecified)!,
-            stripes: true
+            color: statusColors.get(StatusColor.Unspecified)!
         },
         {
             value: metaFailedFactor * 100,
             title: "Failure",
-            color: statusColors.get(StatusColor.Failure)!,
-            stripes: true
+            color: statusColors.get(StatusColor.Failure)!
         },
         {
             value: (1 - (metaFailedFactor + metaUnspecifiedFactor)) * 100,
@@ -175,6 +176,52 @@ export const SessionStatusPie = (session: TestSessionResult, radius: number): JS
     return StatusPie(stack, radius);
 }
 
+export const PhasesStatusPie = (phases: TestPhaseStatus[], radius: number): JSX.Element | undefined => {
+    let totalCount = phases.length || 1;
+    let phaseFailedCount = 0;
+    let phaseUnspecifiedCount = 0;
+    phases.forEach(p => {
+        switch (p.outcome) {
+            case TestPhaseOutcome.Skipped:
+                --totalCount;
+                break;
+            case TestPhaseOutcome.Failed:
+            case TestPhaseOutcome.Interrupted:
+                ++phaseFailedCount;
+                break;
+            case TestPhaseOutcome.Unknown:
+            case TestPhaseOutcome.NotRun:
+                ++phaseUnspecifiedCount;
+                break;
+        }
+    });
+
+    const failedFactor = Math.ceil(phaseFailedCount / totalCount * 100) / 100;
+    const unspecifiedFactor = Math.ceil(phaseUnspecifiedCount / totalCount * 100) / 100;
+
+    const statusColors = dashboard.getStatusColors();
+
+    const stack: StatusBarStack[] = [
+        {
+            value: unspecifiedFactor * 100,
+            title: "Unspecified",
+            color: statusColors.get(StatusColor.Unspecified)!
+        },
+        {
+            value: failedFactor * 100,
+            title: "Failure",
+            color: statusColors.get(StatusColor.Failure)!
+        },
+        {
+            value: (1 - (failedFactor + unspecifiedFactor)) * 100,
+            title: "Passed",
+            color: statusColors.get(StatusColor.Success)!
+        }
+    ];
+
+    return StatusPie(stack, radius);
+}
+
 export const SessionValues = (session: TestSessionResult, style?: any): JSX.Element => {
     const statusColors = dashboard.getStatusColors();
     return <Text style={style}>
diff --git a/Engine/Source/Programs/Horde/HordeDashboard/plugins/test-automation/testAutomationPhaseView.tsx b/Engine/Source/Programs/Horde/HordeDashboard/plugins/test-automation/testAutomationPhaseView.tsx
new file mode 100644
--- /dev/null
+++ b/Engine/Source/Programs/Horde/HordeDashboard/plugins/test-automation/testAutomationPhaseView.tsx
@@ -0,0 +1,639 @@
+import { IconButton, Stack, Text, Spinner, SpinnerSize, Link, FontIcon, IPlainCardProps, DirectionalHint, HoverCard, HoverCardType, mergeStyleSets, mergeStyles, DefaultButton, TooltipHost, Icon, ITag, TagPicker, IPickerItemProps, Label, List, Image, Checkbox } from "@fluentui/react";
+import { useState, useEffect, useCallback, useRef, memo, useMemo } from "react";
+import { observer } from "mobx-react-lite";
+import { TestDataHandler, TestSessionDetails, TestPhaseStatus, DeviceRef, MetadataRef } from "./testData";
+import { EventEntry, EventLevel, ImageCompare, LogPropertyType, TestPhaseOutcome } from './api';
+import dashboard, { StatusColor } from "horde/backend/Dashboard";
+import { ChangeButton, JobParameters } from "horde/components/ChangeButton";
+import { getShortNiceTime, msecToElapsed } from "horde/base/utilities/timeUtils";
+import { getHordeTheme } from "horde/styles/theme";
+import { DashboardPreference } from "horde/backend/Api";
+import { isBright } from "horde/base/utilities/colors";
+import { PhasesStatusPie } from "./testAutomationCommon";
+import { projectStore } from "horde/backend/ProjectStore";
+
+const styles = {
+    labelSmall: {
+        fontSize: 12,
+        fontWeight: "bold"
+    },
+    textSmall: {
+        fontSize: 12
+    },
+    textLarge: {
+        fontSize: 24
+    },
+    defaultButton: {
+        backgroundColor: "#035ca1"
+    },
+    defaultFilter: {
+        backgroundColor: "#0078d4"
+    },
+    defaultCheckbox: {
+        root: {
+            height: 15
+        },
+        text: {
+            fontSize: 11,
+            marginLeft: 0,
+            lineHeight: 15
+        },
+        checkbox: {
+            height: 15,
+            width: 15
+        }
+    },
+    icon: {
+        cursor: "default",
+        height: 15,
+        paddingTop: 1
+    },
+}
+
+let eventStyles: any = undefined;
+const getEventStyling = () => {
+    if (eventStyles) return eventStyles;
+
+    const colors = dashboard.getStatusColors();
+    const errorColorPref = dashboard.getPreference(DashboardPreference.ColorError);
+    const errorHighlight = errorColorPref ? isBright(errorColorPref) ? "brightness(0.9)" : "brightness(1.2)" : "";
+
+    const eventClass = mergeStyles({
+        fontSize: "12px",
+        fontFamily: "Horde Cousine Regular",
+        paddingBottom: 3,
+        paddingTop: 3,
+        paddingRight: 5,
+        paddingLeft: 3
+    });
+
+    eventStyles = mergeStyleSets({
+        event: eventClass,
+        eventWarning: [{
+            background: dashboard.darktheme ? "#1E1817" : "#FEF8E7",
+            borderLeftStyle: 'solid',
+            borderLeftColor: colors.get(StatusColor.Warnings)!
+        }, eventClass],
+        eventError: [{
+            background: dashboard.darktheme ? "#330606" : "#FEF6F6",
+            borderLeftStyle: 'solid',
+            borderLeftColor: colors.get(StatusColor.Failure)!
+        }, eventClass],
+        failureButton: {
+            backgroundColor: errorColorPref ? errorColorPref : dashboard.darktheme ? "#9D1410" : "#EC4C47",
+            borderStyle: "hidden",
+            color: "#FFFFFF",
+            selectors: {
+                ':active,:hover': {
+                    color: "#FFFFFF",
+                    backgroundColor: errorColorPref ? errorColorPref : "#DC3C37",
+                    filter: errorHighlight
+                }
+            }
+        },
+        failureButtonDisabled: {         
+            borderStyle: "hidden",
+            color: dashboard.darktheme ? "#909398" : "#616E85 !important",
+            backgroundColor: dashboard.darktheme ? "#1F2223" : "#f3f2f1"
+        },
+    });
+
+    return eventStyles;
+}
+
+const missingImage = "/images/missing-image.png";
+const MissingImageLabel = (): JSX.Element => { return <span style={{ fontWeight: 'bold' }}> [missing image]</span> }
+type ImageData = { link?: string, ref?: string }
+type ImageLinks = {
+    unapproved: ImageData;
+    approved?: ImageData;
+    difference?: ImageData;
+}
+
+const EventPane: React.FC<{ entry: EventEntry, phase: TestPhaseStatus }> = memo(({entry, phase}) => {
+    const [imageInfo, setImageInfo] = useState<ImageLinks | undefined>(undefined);
+    const eventStyles = getEventStyling();
+
+    const is_error = entry.level === EventLevel.Error || entry.level === EventLevel.Critical;
+    const is_warning = entry.level === EventLevel.Warning;
+    const style = is_error ? eventStyles.eventError : (is_warning ? eventStyles.eventWarning : eventStyles.event);
+
+    const need_image_comparison = entry.properties && entry.properties["$type"] === LogPropertyType.ImageCompare;
+
+    useEffect(() => {
+        if(need_image_comparison && entry.properties) {
+            const imageCompare = entry.properties as ImageCompare;
+            const getLinks = async () => {
+                if (!phase.artifacts) return;
+                let url = await phase.artifacts.getLink(imageCompare.unapproved);
+                const imageInfo: ImageLinks = {unapproved: { link: url, ref: imageCompare.unapproved }};
+                if (imageCompare.approved) {
+                    url = await phase.artifacts.getLink(imageCompare.approved);
+                    imageInfo.approved = { link: url, ref: imageCompare.approved };
+                }
+                if (imageCompare.difference) {
+                    url = await phase.artifacts.getLink(imageCompare.difference);
+                    imageInfo.difference = { link: url, ref: imageCompare.difference };
+                }
+                setImageInfo(imageInfo);
+            }
+            getLinks();
+        }
+    }, [])
+
+    return <Stack className={style} grow>
+                <pre style={{ margin: 0, whiteSpace: "pre-wrap" }}>[{entry.time}] {entry.level}: {entry.message}</pre>
+                {need_image_comparison && imageInfo &&
+                    <Stack styles={{ root: { paddingLeft: 16, width: '100%' } }}>
+                    <Stack><Text variant="medium"><span style={{ fontWeight: "bold" }}>Image comparison: </span>{entry['$text']}</Text></Stack>
+                    <Stack horizontal>
+                        {imageInfo.approved &&
+                            <Stack styles={{ root: { padding: 5 } }}>
+                                <a href={imageInfo.approved.link}>
+                                    <Image width="100%" style={{ minWidth: 200, maxWidth: 400, minHeight: 120 }} src={imageInfo.approved.link || missingImage} alt={imageInfo.approved.ref} />
+                                </a>
+                                <Stack.Item align="center">Reference{!imageInfo.approved && MissingImageLabel()}</Stack.Item>
+                            </Stack>
+                        }
+                        {imageInfo.difference &&
+                            <Stack styles={{ root: { padding: 5 } }}>
+                                <a href={imageInfo.difference.link}>
+                                    <Image width="100%" style={{ minWidth: 200, maxWidth: 400, minHeight: 120 }} src={imageInfo.difference.link || missingImage} alt={imageInfo.difference.ref} />
+                                </a>
+                                <Stack.Item align="center">Difference{!imageInfo.difference && MissingImageLabel()}</Stack.Item>
+                            </Stack>
+                        }
+                        <Stack styles={{ root: { padding: 5 } }}>
+                            <a href={imageInfo.unapproved.link}>
+                                <Image width="100%" style={{ minWidth: 200, maxWidth: 400, minHeight: 120 }} src={imageInfo.unapproved.link || missingImage} alt={imageInfo.unapproved.ref} />
+                            </a>
+                            <Stack.Item align="center">Produced{!imageInfo.unapproved && MissingImageLabel()}</Stack.Item>
+                        </Stack>
+                    </Stack>
+                    </Stack>
+                }
+            </Stack>
+});
+
+const OnDeviceCard = (device: DeviceRef): JSX.Element => {
+    const meta = device.metadata;
+    return <Stack style={{padding: 6}}>
+            { meta && <Stack horizontal tokens={{childrenGap: 4}}>
+                        <Stack tokens={{childrenGap: 4}}>
+                            {Object.keys(meta).map((key, i) => <Stack style={{fontWeight: "bold"}} key={`device-${device.name}-key-${i}`}>{key}: </Stack>)}
+                        </Stack>
+                        <Stack tokens={{childrenGap: 4}}>
+                            {Object.values(meta).map((value, i) => <Stack key={`device-${device.name}-value-${i}`}>{value}</Stack>)}
+                        </Stack>
+                </Stack>
+            }
+            { !meta && <Text>no metadata</Text> }
+        </Stack>;
+}
+
+const DeviceItem = (device: DeviceRef): JSX.Element => {
+    const deviceCardProps: IPlainCardProps = {
+        renderData: device,
+        onRenderPlainCard: OnDeviceCard,
+        directionalHint: DirectionalHint.topRightEdge,
+        gapSpace: 8,
+        calloutProps: { isBeakVisible: true }
+    };
+
+    return <HoverCard plainCardProps={deviceCardProps} type={HoverCardType.plain} cardOpenDelay={20} key={device.name}>
+                <Text style={{fontSize: 12, cursor: 'pointer' }}>{device.name}</Text>
+            </HoverCard>
+}
+
+const TestPhaseExpandPane: React.FC<{ phase: TestPhaseStatus, onEventsLoaded?: () => void, handler: TestDataHandler }> = memo(({ phase, onEventsLoaded, handler }) => {
+    const [events, setEvents] = useState<EventEntry[] | undefined>(undefined);
+    const [filterError, setFilterError] = useState<boolean | undefined>(true);
+    const [filterWarning, setFilterWarning] = useState<boolean | undefined>(true);
+    const [filterInfo, setFilterInfo] = useState<boolean | undefined>(phase.outcome !== TestPhaseOutcome.Failed);
+
+    useEffect(() => {
+        handler.queryPhase(phase, handler.stream!).then(() => {
+            phase.fetchEventStream().then((value) => {
+                setEvents(value ?? []);
+            }).catch((reason) => setEvents([{level: EventLevel.Error, message: `Could not retrieve test events: ${reason}`, format: "", time: "", id: 0}]));
+        });
+    }, []);
+
+    useEffect(() => {
+        if (events) onEventsLoaded?.();
+    }, [events]);
+
+    const isEntryNeedDisplay = useCallback((event: EventEntry): boolean => {
+        const eventType = event.level;
+        if (filterError && eventType === EventLevel.Error) {
+            return true;
+        }
+        if (filterWarning && eventType === EventLevel.Warning) {
+            return true;
+        }
+        if (filterInfo && eventType === EventLevel.Information) {
+            return true;
+        }
+
+        return false;
+    }, [filterError, filterWarning, filterInfo]);
+
+    const borderColor = dashboard.darktheme ? "#4D4C4B" : "#6D6C6B";
+
+    return <Stack>
+                {!events &&
+                    <Stack horizontalAlign='center' style={{ padding: 12}} tokens={{ childrenGap: 10 }}>
+                        <Text style={styles.textSmall}>Loading Data</Text>
+                        <Spinner size={SpinnerSize.medium} />
+                    </Stack>
+                }
+
+                {!!events &&
+                    <Stack style={{margin: 3}} tokens={{childrenGap: 3}}>
+                        <Stack>[ History graph place holder ]</Stack>
+                        <Stack style={{borderWidth: 1, borderColor: borderColor, borderStyle: 'solid', padding: 6}} tokens={{childrenGap: 3}}>
+                            <Stack horizontal tokens={{childrenGap: 8}}>
+                                <Text><span style={styles.labelSmall}>Run on: </span><span>{getShortNiceTime(phase.start, true, true, true)}</span></Text>
+                                <Text><span style={styles.labelSmall}> For: </span><span>{msecToElapsed(phase.duration * 1000, true, true)}</span></Text>
+                                {!!phase.devices.length &&
+                                    <Stack horizontal tokens={{childrenGap: 4}} verticalAlign="end">
+                                        <Text style={styles.labelSmall}> Devices: </Text>
+                                        {phase.devices.map(d => DeviceItem(d))}
+                                    </Stack>
+                                }
+                                <Stack verticalAlign="center" horizontal tokens={{ childrenGap: 5 }} style={{borderWidth: '0px 1px', paddingLeft: 9, paddingRight: 9, borderStyle: 'solid', borderColor: borderColor}}>
+                                    <FontIcon iconName="PageListFilter" title="Events filtering" style={styles.icon} />
+                                    <Checkbox styles={styles.defaultCheckbox} label="Error" defaultChecked={filterError} onChange={(_: any, checked?: boolean) => setFilterError(checked)} disabled={!events.length}/>
+                                    <Checkbox styles={styles.defaultCheckbox} label="Warning" defaultChecked={filterWarning} onChange={(_: any, checked?: boolean) => setFilterWarning(checked)} disabled={!events.length}/>
+                                    <Checkbox styles={styles.defaultCheckbox} label="Info" defaultChecked={filterInfo} onChange={(_: any, checked?: boolean) => setFilterInfo(checked)} disabled={!events.length}/>
+                                </Stack>
+                                <Stack horizontalAlign="end">
+                                    <DefaultButton style={{ fontSize: 11, padding: 0, height: 18, backgroundColor: styles.defaultButton.backgroundColor }}
+                                        onClick={(ev) => {
+                                            ev.stopPropagation();
+                                            /** do something */
+                                        }}>Artifacts</DefaultButton>
+                                </Stack>
+                            </Stack>
+                            <Stack style={{ paddingTop: 6, paddingBottom: 6}}>
+                                { events.map((e, i) => isEntryNeedDisplay(e) && <EventPane key={`phase-${phase.key}-event-${i}`} entry={e} phase={phase}/>)}
+                                { !events.length && <Text style={styles.textSmall}>No events</Text>}
+                            </Stack>
+                        </Stack>
+                    </Stack>
+                }
+            </Stack>
+});
+
+const TestPhaseSection: React.FC<{ phase: TestPhaseStatus, focused: boolean, handler: TestDataHandler }> = memo(({ phase, focused, handler }) => {
+    const [expanded, setExpanded] = useState(false);
+
+    useEffect(() => {
+        if (focused && !expanded) {
+            setExpanded(true);
+        }
+        scrollToSelected();
+    }, [focused]);
+
+    const registerRef = useCallback((ref: any) => {
+        phase.componentRef = ref;
+    }, []);
+
+    const scrollToSelected = useCallback(() => {
+        if (expanded && focused) {
+            const component = phase.componentRef;
+            if (!component) return;
+            component.scrollIntoView();
+       }
+    }, [focused, expanded]);
+
+    const toggleView = useCallback(() => {
+        if (!expanded) {
+            handler.setSearchParam("phase", phase.key);
+        } else if(focused) {
+            handler.removeSearchParam('phase');
+        }
+        setExpanded(!expanded);
+    }, [focused, expanded]);
+    
+    const borderColor = dashboard.darktheme ? "#4D4C4B" : "#6D6C6B";
+    const statusColors = dashboard.getStatusColors();
+    let color = statusColors.get(StatusColor.Unspecified);
+    switch(phase.outcome) {
+        case TestPhaseOutcome.Failed:
+            color = statusColors.get(StatusColor.Failure);
+            break;
+        case TestPhaseOutcome.Success:
+            color = statusColors.get(phase.hasWarning? StatusColor.Warnings : StatusColor.Success);
+            break;
+    }
+
+    return <Stack style={{ borderTopWidth: 1, borderTopColor: borderColor, borderTopStyle: 'solid', paddingTop: 3, paddingBottom: 6 }}>
+                <div ref={registerRef}>
+                    <Stack tokens={{childrenGap: 2}}>
+                        <Stack horizontal verticalAlign="center">
+                            <Stack style={{ paddingTop: 1, paddingRight: 4, cursor: "pointer" }} onClick={toggleView}>
+                                <FontIcon style={{ fontSize: 13 }} iconName={expanded ? "ChevronDown" : "ChevronRight"} />
+                            </Stack>
+                            <Text style={{fontSize: 13, fontWeight: 600}}>{phase.name}</Text>
+                        </Stack>
+                        <Stack horizontal verticalAlign="center">
+                            <Stack style={{ paddingLeft: 5, paddingRight: 4 }}>
+                                <FontIcon style={{ fontSize: 12, color: color }} iconName="Square" />
+                            </Stack>
+                            <Text style={{fontSize: 11}}>{phase.outcome}</Text>
+                        </Stack>
+                    </Stack>
+                </div>
+                {expanded && <TestPhaseExpandPane phase={phase} onEventsLoaded={scrollToSelected} handler={handler}/>}
+            </Stack>
+});
+
+
+const onRenderSuggestionsItem = (item: ITag) => {
+    return <Stack style={{height: 24, padding: 4}}>
+                <Text title={item.name} style={{textOverflow: 'ellipsis', overflow: 'hidden', whiteSpace: 'nowrap', maxWidth: 300, fontSize: 12}}>{item.name}</Text>
+            </Stack>
+}
+
+const onRenderItem = (props: IPickerItemProps<ITag>) => {
+    const item = props.item;
+    return <Stack key={`picker_item_${item.name}`} horizontal verticalAlign="center" verticalFill={true} style={{ marginRight: 2 }}>
+                <TooltipHost content={item.name} delay={0} closeDelay={500}>
+                    <FontIcon
+                        style={{ fontSize: 13, padding: 4, paddingTop: 5, cursor: 'pointer', backgroundColor: styles.defaultFilter.backgroundColor, borderRadius: 3}}
+                        iconName="Filter"
+                        title="Remove"
+                        onClick={props.onRemoveItem} />
+                </TooltipHost>
+            </Stack>
+}
+
+const getTextFromTag = (item: ITag) => item.name;
+
+const PhaseFilter: React.FC<{phases: TestPhaseStatus[], onChange: (items: string[]) => void, onFocus?: () => void, onBlur?: () => void}> = ({ phases, onChange, onFocus, onBlur }) => {
+    const [filters, setFilters] = useState<ITag[]>([]);
+    const phaseTags = useRef<ITag[]>();
+
+    useEffect(() => {
+        const tokens: ITag[] = (new Set(phases.map(p => (p.name.split('.'))).flat())).keys().map(token => ({key: token.toLowerCase(), name: token})).toArray().sort((a, b) => a.key.length - b.key.length);
+        const tags: ITag[] = phases.map(p => ({key: `@${p.name.toLowerCase()}`, name: `@${p.name}`}));
+        phaseTags.current = [...tokens, ...tags];
+    }, [phases]);
+    
+    const filterSelectedTests = useCallback((filterText: string): ITag[] => {
+        if (filterText.length < 2) return [];
+        const lowerText = filterText.toLowerCase();
+        const closestTags = phaseTags.current?.values().filter((tag) => (tag.key as string).indexOf(lowerText) >= 0).take(10).toArray() ?? [];
+        const isTagExist = closestTags.find((tag) => tag.key === lowerText);
+        if (!isTagExist) closestTags.splice(0, 0, {key: lowerText, name: filterText});
+        return closestTags;
+    }, [phases]);
+
+    return <TagPicker
+                styles={{
+                    text: { minWidth: 210, maxWidth: 450, maxHeight: 28, display: 'inline-flex', flexWrap: 'nowrap' },
+                    input: { height: 28, minWidth: 40, maxWidth: 200 },
+                    itemsWrapper: { marginLeft: 3, marginTop: 2, marginRight: 3, marginBottom: 2, height: 28, maxWidth: 400, overflow: 'hidden' }
+                }}
+                onRenderItem={onRenderItem}
+                onRenderSuggestionsItem={onRenderSuggestionsItem}
+                removeButtonAriaLabel="Remove"
+                selectionAriaLabel="Filter phases"
+                selectedItems={filters}
+                onResolveSuggestions={filterSelectedTests}
+                getTextFromItem={getTextFromTag}
+                onChange={(tags) => {
+                    let items = tags ?? [];
+                    const uniqueTags = new Set(items.map(i => i.key));
+                    items = uniqueTags.keys().map(k => items.find(i => i.key === k)!).toArray();
+                    setFilters(items);
+                    onChange(items.map(t => t.name));
+                }}
+                inputProps={{ placeholder: "Filter", onFocus: onFocus, onBlur: onBlur }}
+            />
+}
+
+class CircularList<T> {
+    constructor(items: T[]) {
+        this._items = items;
+        this._cursor = -1;
+    }
+
+    setCursor(item: T | undefined): number {
+        if (item === undefined) {
+            this._cursor = -1;
+        } else {
+            this._cursor = this._items.findIndex(i => i === item);
+        }
+
+        return this._cursor;
+    }
+
+    get selected(): T | undefined {
+        if (this._cursor === -1) return;
+        return this._items[this._cursor];
+    }
+
+    next(): T {
+        if (++this._cursor === this._items.length) this._cursor = 0;
+        return this._items[this._cursor];
+    }
+
+    back(): T {
+        if (--this._cursor < 0) this._cursor = this._items.length - 1;
+        return this._items[this._cursor];
+    }
+
+    get cursor() {
+        return this._cursor;
+    }
+
+    get length() {
+        return this._items.length;
+    }
+
+    private _cursor: number;
+    private _items: T[];
+}
+
+export const TestPhasesView: React.FC<{ testDataId: string, handler: TestDataHandler, onDismiss: () => void }> = observer(({ testDataId, handler, onDismiss }) => {
+    const [sessionDetails, setSessionDetails] = useState<TestSessionDetails | undefined>(undefined);
+    const [filteredPhases, setFilteredPhases] = useState<TestPhaseStatus[]>([]);
+    const failureCursor = useRef<CircularList<string>>();
+    const listRef = useRef<List>();
+    const job = useRef<{params?: JobParameters, rangeCL?: number}>();
+    const selectedPhase = handler.getSearchParam('phase') as string | undefined;
+
+    handler.subscribeToTestDataQueryLoading();
+    handler.subscribeToSearch();
+
+    const updateFilter = useCallback((phases: TestPhaseStatus[]) => {
+        const failures = phases.filter((p) => p.outcome !== TestPhaseOutcome.Success && p.outcome !== TestPhaseOutcome.NotRun && p.outcome !== TestPhaseOutcome.Skipped).map((p) => p.key) ?? [];
+        failureCursor.current = new CircularList(failures);
+        failureCursor.current.setCursor(selectedPhase);
+        setFilteredPhases(phases);
+    }, [handler.searchUpdated]);
+
+    const onRenderCell = useCallback(
+        (phase: TestPhaseStatus) => <TestPhaseSection key={phase.key} phase={phase} focused={selectedPhase === phase.key} handler={handler}/>
+        , [handler.searchUpdated]);
+
+    useEffect(() => {
+        handler.queryTestData(testDataId).then((fetchedSessionDetails) => {
+            if (fetchedSessionDetails) {
+                setSessionDetails(fetchedSessionDetails);
+                updateFilter(fetchedSessionDetails.phases);
+                job.current = {
+                    params: {
+                        id: fetchedSessionDetails.jobId,
+                        streamId: fetchedSessionDetails.streamId,
+                        change: fetchedSessionDetails.commitOrder,
+                    },
+                    rangeCL: fetchedSessionDetails.session ?
+                        handler.getStatusStream(fetchedSessionDetails.streamId)?.getPreviousSession(fetchedSessionDetails.test!, fetchedSessionDetails.meta!, fetchedSessionDetails.session)?.commitOrder
+                        : undefined
+                }
+            }
+        });
+        setSessionDetails(undefined);
+    }, [testDataId]);
+
+    useEffect(() => {
+        if (selectedPhase !== failureCursor.current?.selected) {
+            failureCursor.current?.setCursor(selectedPhase);
+        }
+        if (!selectedPhase) return;
+        const selectedIndex = filteredPhases.findIndex((phase) => phase.key === selectedPhase);
+        if (selectedIndex === -1) return;
+
+        if (!filteredPhases[selectedIndex].componentRef) {
+            // force the list to load the page that contains the target item
+            listRef.current?.scrollToIndex(selectedIndex);
+        }
+        // force the list to update in anycase so that the loaded items get the propagation of the search params update
+        listRef.current?.forceUpdate();
+    }, [handler.searchUpdated, filteredPhases]);
+
+    const StatusPie = useMemo(
+        () => <Stack horizontal verticalAlign='center' style={{ height: '100%' }} tokens={{childrenGap: 3}}>
+                <Label>{filteredPhases.length} {filteredPhases.length > 1 ? "Phases" : "Phase"}</Label>
+                {PhasesStatusPie(filteredPhases, 14)}
+            </Stack>, [filteredPhases]);
+
+    const metaString = useMemo(() => sessionDetails?.meta?.getValues().join(' / '), [sessionDetails]);
+
+    const hordeTheme = getHordeTheme();
+    const eventStyles = getEventStyling();
+    const buttonNoneTextColor = dashboard.darktheme ? "#949898" : "#616e85";
+    const errorColorPref = dashboard.getPreference(DashboardPreference.ColorError);
+    const failureButtonTextColor = errorColorPref && isBright(errorColorPref) ? "#2d3f5f" : "#F9F9FB";
+
+    const failureEnabled = !!failureCursor.current?.length;
+    const failureText = failureEnabled ? failureCursor.current!.length.toString() : "No";
+
+    return <Stack style={{width: '100%', height: '100%'}}>
+                {!sessionDetails && handler.testDataQueryLoading &&
+                        <Stack horizontalAlign='center' style={{ padding: 12 }} tokens={{ childrenGap: 10 }}>
+                            <Text style={styles.textLarge}>Loading Data</Text>
+                            <Spinner size={SpinnerSize.large} />
+                        </Stack>
+                }
+                {!sessionDetails && !handler.testDataQueryLoading &&
+                    <Stack grow tokens={{childrenGap: 3}} style={{height: '100%'}}>
+                        <Stack horizontal verticalAlign="start" tokens={{childrenGap: 10}}>
+                            <Stack grow />
+                            <Stack style={{ paddingBottom: 4 }}>
+                                <IconButton
+                                    iconProps={{ iconName: 'Cancel' }}
+                                    ariaLabel="Close phase view"
+                                    onClick={onDismiss}
+                                />
+                            </Stack>
+                        </Stack>
+                        <Stack horizontalAlign='center' style={{ padding: 12 }}>
+                            <Text style={styles.textLarge}>No Data</Text>
+                        </Stack>
+                    </Stack>
+                }
+                {!!sessionDetails &&
+                    <Stack grow tokens={{childrenGap: 3}} style={{height: '100%'}}>
+                        <Stack horizontal verticalAlign="start" tokens={{childrenGap: 10}}>
+                            <Stack style={{ paddingLeft: 8}} tokens={{childrenGap: 5}}>
+                                <Text style={{ fontSize: 16, fontWeight: 600 }}>{`${sessionDetails.name} - ${metaString}`}</Text>
+                                <Stack style={{ paddingLeft: 8 }} horizontal tokens={{childrenGap: 8}}>
+                                    <Text>{projectStore.streamById(sessionDetails.streamId)?.fullname ?? sessionDetails.streamId}</Text>
+                                    <ChangeButton job={job.current?.params} rangeCL={job.current?.rangeCL} />
+                                    <Link href={`/job/${sessionDetails.jobId}?step=${sessionDetails.stepId}`} target="_blank"><Text style={{ fontSize: 12, color: hordeTheme.semanticColors.link }}>View Job</Text></Link>
+                                    <Text><span style={styles.labelSmall}>Run on: </span><span>{getShortNiceTime(sessionDetails.start, true, true, true)}</span></Text>
+                                    <Text><span style={styles.labelSmall}> For: </span><span>{msecToElapsed(sessionDetails.duration * 1000, true, true)}</span></Text>
+                                </Stack>
+                            </Stack>
+
+                            <Stack grow />
+
+                            {StatusPie}
+
+                            <Stack horizontal verticalAlign='center' style={{ height: '100%' }}>
+                                <Stack horizontal verticalAlign='center' style={{position: 'relative'}}>
+                                    <TooltipHost content={failureEnabled ? 'Navigate to next failure': ''}>
+                                        <DefaultButton
+                                            disabled={!failureEnabled} className={failureEnabled ? eventStyles.failureButton : eventStyles.failureButtonDisabled}
+                                            style={{ color: failureEnabled ? failureButtonTextColor : buttonNoneTextColor, height: 28, fontSize: 12, paddingLeft: 8, paddingRight: 8, borderRadius: failureEnabled ? '4px 0 0 4px' : '4px' }}
+                                            text={`${failureText} ${(failureCursor.current?.length ?? 0) > 1 ? "Failures" : "Failure"}`}
+                                            onClick={() => {
+                                                if (failureCursor.current) handler.setSearchParam('phase', failureCursor.current.next());
+                                            }} >
+                                            {failureEnabled &&
+                                                <Icon style={{ fontSize: 19, paddingLeft: 4, paddingRight: 0 }} iconName='ChevronDown' />
+                                            }
+                                        </DefaultButton>
+                                    </TooltipHost>
+                                    {failureEnabled &&
+                                        <div style={{backgroundColor: failureButtonTextColor, height: 18, width: 1, position: 'absolute', right: -0.5, zIndex: 1 }} />
+                                    }
+                                </Stack>
+                                {failureEnabled &&
+                                    <TooltipHost content={'Navigate to previous failure'}>
+                                        <IconButton
+                                            disabled={!failureEnabled} className={failureEnabled ? eventStyles.failureButton : eventStyles.failureButtonDisabled}
+                                            style={{ color: failureButtonTextColor, height: 28, fontSize: 12, padding: 2, borderRadius: '0 4px 4px 0' }}
+                                            iconProps={{ iconName: 'ChevronUp' }}
+                                            onClick={() => {
+                                                if (failureCursor.current) handler.setSearchParam('phase', failureCursor.current.back());
+                                            }} />
+                                    </TooltipHost>
+                                }
+                            </Stack>
+
+                            <Stack verticalAlign='center' style={{ height: '100%' }}>
+                                <TooltipHost content='Filter phases'>
+                                    <PhaseFilter phases={sessionDetails.phases}
+                                        onChange={(items => {
+                                            const lowerCaseItems = items.map(i => i.toLowerCase());
+                                            const phases = sessionDetails.phases.filter(phase => phase.isMatch(lowerCaseItems))
+                                            updateFilter(phases);
+                                        })}/>
+                                </TooltipHost>
+                            </Stack>
+
+                            <Stack style={{ paddingBottom: 4 }}>
+                                <IconButton
+                                    iconProps={{ iconName: 'Cancel' }}
+                                    ariaLabel="Close phase view"
+                                    onClick={onDismiss}
+                                />
+                            </Stack>
+                        </Stack>
+                        <Stack style={{ paddingTop: 3, paddingLeft: 16, paddingBottom: 3, paddingRight: 3, height: '100%', overflow: 'auto'}} grow>
+                            <List
+                                items={filteredPhases}
+                                onRenderCell={onRenderCell}
+                                getKey={(phase: TestPhaseStatus) => `phase-${phase.key}`}
+                                ref={(list: List) => listRef.current = list}
+                                />
+                        </Stack>
+                    </Stack>
+                }     
+            </Stack>
+});
\ No newline at end of file
diff --git a/Engine/Source/Programs/Horde/HordeDashboard/plugins/test-automation/testAutomationSummary.tsx b/Engine/Source/Programs/Horde/HordeDashboard/plugins/test-automation/testAutomationSummary.tsx
--- a/Engine/Source/Programs/Horde/HordeDashboard/plugins/test-automation/testAutomationSummary.tsx
+++ b/Engine/Source/Programs/Horde/HordeDashboard/plugins/test-automation/testAutomationSummary.tsx
@@ -1,5 +1,5 @@
 import { DefaultButton, FontIcon, IconButton, mergeStyleSets, Modal, Stack, Text, HoverCard, HoverCardType, IPlainCardProps, DirectionalHint, Spinner, SpinnerSize, IContextualMenuItem, IContextualMenuProps, IComboBoxOption } from "@fluentui/react";
-import { useState, useEffect } from "react";
+import { useState, useEffect, useCallback, memo } from "react";
 import { observer } from "mobx-react-lite";
 import { TestDataHandler, TestNameRef, MetadataRef, TestSessionStatus, TestSessionResult, TestStatus } from "./testData";
 import { TestOutcome } from './api';
@@ -18,11 +18,8 @@ const TestHistoryGraphWidget: React.FC<{ test: TestNameRef, stream: string, hand
     const [state, setState] = useState<{ graph?: TestHistoryGraph }>({});
 
     useEffect(() => {
-        const setup = async () => {
-            const graph = new TestHistoryGraph(test, stream, handler, commonMetaKeys, onClick);
-            setState({ graph: graph });
-        }
-        setup();
+        const graph = new TestHistoryGraph(test, stream, handler, commonMetaKeys, onClick);
+        setState({ graph: graph });
     }, []);
 
     useEffect(() => {
@@ -52,7 +49,7 @@ const TestHistoryGraphWidget: React.FC<{ test: TestNameRef, stream: string, hand
             </Stack>
 }
 
-const StreamSelector: React.FC<{ streams: string[], onClick: (stream: string) => void, selected?: string }> = ({ streams, onClick, selected }) => {
+const StreamSelector: React.FC<{ streams: string[], onClick: (stream: string) => void, selected?: string, disabled?: boolean }> = ({ streams, onClick, selected, disabled }) => {
     const options: IContextualMenuItem[] = [];
     streams.forEach(a => {
         options.push({
@@ -69,7 +66,7 @@ const StreamSelector: React.FC<{ streams: string[], onClick: (stream: string) =>
         items: options,
     };
 
-    return <DefaultButton style={{ height: 22, fontSize: 11, whiteSpace: 'nowrap', borderWidth: 0 }} menuProps={menuProps} text={selected? projectStore.streamById(selected)?.fullname ?? selected : "Compare stream"}/>
+    return <DefaultButton style={{ height: 22, fontSize: 11, whiteSpace: 'nowrap', borderWidth: 0 }} menuProps={menuProps} text={selected? projectStore.streamById(selected)?.fullname ?? selected : "Compare stream"} disabled={disabled}/>
 }
 
 export const TestHistoryView: React.FC<{ test: TestNameRef, handler: TestDataHandler, onClick: SessionCallback  }> = observer(({ test, handler, onClick }) => {
@@ -161,11 +158,11 @@ export const TestHistoryView: React.FC<{ test: TestNameRef, handler: TestDataHan
                 <Stack horizontal tokens={{childrenGap: 15}} style={{ paddingBottom: 12, paddingTop: 5 }}>
                     <Stack style={{ width: 170 }}>
                         <MultiOptionChooser style={{height: 22, fontSize: 11}} placeholder="Filter by metadata"
-                            options={metaOptions} initialSelection={selectedMeta.keys().toArray()}
+                            options={metaOptions} initialSelection={selectedMeta.keys().toArray()} disabled={metaOptions.length <= 1}
                             updateKeys={(keys) => setSelectedMeta(new Set(keys))} />
                     </Stack>
                     <Stack>
-                        <StreamSelector streams={availableStreams} selected={streams.toCompare} onClick={(targetStream: string) => {
+                        <StreamSelector streams={availableStreams} selected={streams.toCompare} disabled={availableStreams.length === 0} onClick={(targetStream: string) => {
                             setStreams({...streams, toCompare: undefined});
                             handler.query(targetStream, true).catch((reason) => console.error(reason)).finally(() => setStreams({...streams, toCompare: targetStream}))
                         }} />
@@ -205,7 +202,7 @@ const TestHistoryModal: React.FC<{ test: TestNameRef, handler: TestDataHandler,
         </Modal>
 }
 
-const TestSummaryCard: React.FC<{ test: TestNameRef, sessions: [MetadataRef, TestSessionResult][], metaKeyMask: string[] | undefined, handler: TestDataHandler, onDetailsClick: SessionCallback, onHistoryClick: HistorySessionCallback }> = ({ test, sessions, metaKeyMask, handler, onDetailsClick, onHistoryClick }) => {
+const TestSummaryCard: React.FC<{ test: TestNameRef, sessions: [MetadataRef, TestSessionResult][], metaKeyMask: string[] | undefined, handler: TestDataHandler, onDetailsClick: SessionCallback, onHistoryClick: HistorySessionCallback }> = memo(({ test, sessions, metaKeyMask, handler, onDetailsClick, onHistoryClick }) => {
     const [expanded, setExpanded] = useState(false);
  
     const { hordeClasses, modeColors } = getHordeStyling();
@@ -310,8 +307,8 @@ const TestSummaryCard: React.FC<{ test: TestNameRef, sessions: [MetadataRef, Tes
         );
     });
  
-    const testFailedFactor = Math.ceil(testFailed / (testTotal || 1) * 20) / 20;
-    const testUnspecifiedFactor = Math.ceil(testUnspecified / (testTotal || 1) * 20) / 20;
+    const testFailedFactor = Math.ceil(testFailed / (testTotal || 1) * 50) / 50;
+    const testUnspecifiedFactor = Math.ceil(testUnspecified / (testTotal || 1) * 50) / 50;
     const testStack: StatusBarStack[] = [
         {
             value: testUnspecifiedFactor * 100,
@@ -348,8 +345,7 @@ const TestSummaryCard: React.FC<{ test: TestNameRef, sessions: [MetadataRef, Tes
             <Stack>
                 <Stack style={{ paddingTop: 4 }}>
                     {StatusBar(testStack, 308, 10, statusColors.get(StatusColor.Success)!, { margin: '3px !important' })}
-                </Stack>
-        
+                </Stack>        
             </Stack>
             {!!expanded &&
                 <Stack style={{ paddingLeft: 8 }} tokens={{ childrenGap: 4 }}>
@@ -357,7 +353,7 @@ const TestSummaryCard: React.FC<{ test: TestNameRef, sessions: [MetadataRef, Tes
                 </Stack>
             }
             </Stack>
-}
+});
 
 const filterSessionsWithMetaMask = (testStatus: TestStatus, filteredTests: TestNameRef[], handler: TestDataHandler): [[TestNameRef, [MetadataRef, TestSessionStatus][]][], Map<string, string[]>] => {
     const filteredMetadata = handler.filteredMetadata;
@@ -417,11 +413,19 @@ const filterSessionsWithMetaMask = (testStatus: TestStatus, filteredTests: TestN
 
 export const TestSummary: React.FC<{ handler: TestDataHandler }> = ({ handler }) => {
     const [historyState, setHistoryState] = useState<{test?: TestNameRef, onClick?: SessionCallback}>({});
-
     const searchTestId = handler.getSearchParam('history');
     if (!searchTestId) {
         historyState.test = historyState.onClick = undefined;
     }
+    
+    const onDetailsClick = useCallback((session: TestSessionResult) => {
+        if (session.testDataId) {
+            handler.setSearchParam('session', session.testDataId);
+        }
+        else if (session.jobId && session.stepId) {
+            window.open(`/job/${session.jobId}?step=${session.stepId}`, '_blank');
+        }
+    }, []);
 
     const testStatus = handler.selectedStatusStream;
     if (!testStatus?.tests.size) {
@@ -437,13 +441,7 @@ export const TestSummary: React.FC<{ handler: TestDataHandler }> = ({ handler })
 
     const tests = testSessions.map(([test, sessions]) => {
 
-        const onDetailsClick = (session: TestSessionResult) => {
-            if (session.jobId && session.stepId) {
-                window.open(`/job/${session.jobId}?step=${session.stepId}`, '_blank');
-            }
-        }
-
-        if (searchTestId && test.id === searchTestId) {
+        if (searchTestId && !historyState.test && test.id === searchTestId) {
             historyState.test = test;
             historyState.onClick = onDetailsClick;
         }
diff --git a/Engine/Source/Programs/Horde/HordeDashboard/plugins/test-automation/testAutomationView.tsx b/Engine/Source/Programs/Horde/HordeDashboard/plugins/test-automation/testAutomationView.tsx
--- a/Engine/Source/Programs/Horde/HordeDashboard/plugins/test-automation/testAutomationView.tsx
+++ b/Engine/Source/Programs/Horde/HordeDashboard/plugins/test-automation/testAutomationView.tsx
@@ -1,14 +1,15 @@
-import { Stack, FontIcon, IComboBox, IComboBoxOption, ComboBox, IContextualMenuItem, IContextualMenuProps, Text, Spinner, SpinnerSize, Label, DefaultButton, TagItemSuggestion, ITag, TagPicker } from "@fluentui/react";
+import { Stack, FontIcon, IComboBoxOption, IContextualMenuItem, IContextualMenuProps, Text, Spinner, SpinnerSize, Label, DefaultButton, TagItemSuggestion, ITag, TagPicker, IPickerItemProps, PrimaryButton } from "@fluentui/react";
 import { Breadcrumbs } from "horde/components/Breadcrumbs";
 import { TopNav } from "horde/components/TopNav";
 import { getHordeStyling } from "horde/styles/Styles";
 import { observer } from "mobx-react-lite";
 import { useSearchParams } from "react-router-dom";
-import React, { useEffect, useState, useRef } from "react";
+import React, { useEffect, useState, useCallback } from "react";
 import { projectStore } from 'horde/backend/ProjectStore';
 import { TestDataHandler, TestDataVersionRegistrar } from "./testData";
 import { TestSummary } from "./testAutomationSummary";
 import { MultiOptionChooser } from "./testAutomationCommon";
+import { TestPhasesView } from "./testAutomationPhaseView";
 
 import { TestDataV2 } from "./testDataV2Fetcher";
 
@@ -40,12 +41,27 @@ const StreamChooser: React.FC = observer(() => {
 
     return <Stack style={{ paddingTop: 12, paddingBottom: 4 }} horizontal>
             <Label style={{ paddingTop: 5, paddingRight: 5 }}>Stream</Label>
-            <DefaultButton style={{ width: 270, textAlign: "left", height: 30 }} menuProps={menuProps} text={stream ? projectStore.streamById(stream)?.fullname : "Select"} />
+            <DefaultButton style={{ width: 270, textAlign: "left", height: 30 }} menuProps={menuProps} text={stream ? projectStore.streamById(stream)?.fullname ?? stream : "Select"} />
         </Stack>   
-    });
+});
+
+const onRenderSuggestionsItem = (item: ITag) => {
+    return <Stack style={{height: 24, padding: 4}}>
+                <Text title={item.name} style={{textOverflow: 'ellipsis', overflow: 'hidden', whiteSpace: 'nowrap', maxWidth: 300, fontSize: 12}}>{item.name}</Text>
+            </Stack>
+}
 
-const onRenderSuggestionItem = (props: ITag) => {
-    return <TagItemSuggestion style={{height: 'auto', whiteSpace: 'normal', maxWidth: 500}}>{props.name}</TagItemSuggestion>
+const onRenderItem = (props: IPickerItemProps<ITag>) => {
+    const item = props.item;
+    return <Stack style={{ marginTop: 2, marginLeft: 3 }} key={`picker_item_${item.name}`}>
+        <PrimaryButton
+            iconProps={{ iconName: "Cancel", styles: { root: { fontSize: 12, margin: 0, padding: 0 } } }}
+            styles={{label: { textOverflow: 'ellipsis', overflow: 'hidden', whiteSpace: 'nowrap', minWidth: 0, maxWidth: 150, margin: 0, paddingLeft: 3 }}}
+            style={{ padding: 2, paddingLeft: 5, paddingRight: 5, fontSize: 12, height: "unset"}}
+            text={item.name}
+            title={item.name}
+            onClick={props.onRemoveItem} />
+    </Stack>;
 }
 
 const getTextFromTag = (item: ITag) => item.name;
@@ -66,45 +82,45 @@ const TestChooser: React.FC = observer(() => {
         }
     }, [handler.searchUpdated]);
 
-    if (!streamTests.length) {
-            return null;
-    }
-
-    const testTags: ITag[] = streamTests.map(t => ({key: `#${t.toLowerCase()}`, name: `#${t}`}))
-
-    const filterSelectedTests = (filterText: string, tagList: ITag[]): ITag[] => {
-        if (!filterText) return [];
+    const testTags: ITag[] = streamTests.map(t => ({key: `@${t.toLowerCase()}`, name: `@${t}`}));
+    const filterSelectedTests = useCallback((filterText: string, _: ITag[]): ITag[] => {
+        if (filterText.length < 2) return [];
         const lowerText = filterText.toLowerCase();
-        const closestTags = testTags.filter((tag) => tag.key.toString().indexOf(lowerText) >= 0);
+        const closestTags = testTags.values().filter((tag) => (tag.key as string).indexOf(lowerText) >= 0).take(10).toArray();
         const isTagExist = closestTags.find((tag) => tag.key === lowerText);
         if (!isTagExist) closestTags.splice(0, 0, {key: lowerText, name: filterText});
         return closestTags;
+    }, [handler.queryLoading]);
+
+    if (!streamTests.length) {
+        return null;
     }
+
     const selectedItems: ITag[] = testsSelection.keys().toArray().map((key) => ({key: key.toLowerCase(), name: key} as ITag));
 
     return <Stack style={{ paddingTop: 12, paddingBottom: 4 }} horizontal>
             <Label style={{ paddingTop: 5, paddingRight: 5 }}>Tests</Label>
             <Stack tokens={{ childrenGap: 6 }} style={{ paddingLeft: 12 }} grow>
                 <TagPicker
-                        styles={{ input: { height: 28 }, itemsWrapper: { height: 28 } }}
-                        onRenderSuggestionsItem={onRenderSuggestionItem}
-                        removeButtonAriaLabel="Remove"
-                        selectionAriaLabel="Selected tests"
-                        selectedItems={selectedItems}
-                        onResolveSuggestions={filterSelectedTests}
-                        getTextFromItem={getTextFromTag}
-                        onChange={(tags) => {
-                            if (!!tags) {
-                                const keys = tags!.map(getTextFromTag);
-                                keys.forEach((key) => handler.addTest(key));
-                                testsSelection.difference(new Set(keys)).values().forEach(key => handler.removeTest(key));
-                            } else {
-                                testsSelection.values().forEach(key => handler.removeTest(key));
-                            }
-                            setTestsSelection(new Set(handler.filterState.tests ?? []));
-                        }}
-                />
-                
+                    styles={{ input: { height: 28, width: 40 }, itemsWrapper: { marginRight: 3, marginBottom: 2 } }}
+                    onRenderItem={onRenderItem}
+                    onRenderSuggestionsItem={onRenderSuggestionsItem}
+                    removeButtonAriaLabel="Remove"
+                    selectionAriaLabel="Selected tests"
+                    selectedItems={selectedItems}
+                    onResolveSuggestions={filterSelectedTests}
+                    getTextFromItem={getTextFromTag}
+                    onChange={(tags) => {
+                        if (!!tags) {
+                            const keys = tags!.map(getTextFromTag);
+                            keys.forEach((key) => handler.addTest(key));
+                            testsSelection.difference(new Set(keys)).values().forEach(key => handler.removeTest(key));
+                        } else {
+                            testsSelection.values().forEach(key => handler.removeTest(key));
+                        }
+                        setTestsSelection(new Set(handler.filterState.tests ?? []));
+                    }}
+                />                
             </Stack>
         </Stack>
 });
@@ -194,7 +210,7 @@ const TestAutomationSidebarLeft: React.FC = observer(() => {
     const [metaKeys, setMetaKeys] = useState<string[]>(filterMetaKeys);
     const { hordeClasses } = getHordeStyling();
 
-    const updateMetaKey = (key: string) => {
+    const updateMetaKey = useCallback((key: string) => {
         const idx = metaKeys.indexOf(key);
         if (idx < 0) {
             metaKeys.push(key);
@@ -202,7 +218,7 @@ const TestAutomationSidebarLeft: React.FC = observer(() => {
             metaKeys.splice(idx, 1);
         }
         setMetaKeys([...metaKeys]);
-    }
+    }, [metaKeys]);
 
     useEffect(() => {
         // keep search filter in sync
@@ -270,6 +286,8 @@ export const TestAutomationView: React.FC = () => {
     const [searchParams, setSearchParams] = useSearchParams();
     const [initialized, setInitialized] = useState(false);
 
+    const testDataId = searchParams.get('session');
+
     useEffect(() => {
         handler.initialize(searchParams.toString(), (search: string) => {setSearchParams(search)}).then(() => {setInitialized(true)});
         return () => {
@@ -288,11 +306,14 @@ export const TestAutomationView: React.FC = () => {
                 <TopNav />
                 <Breadcrumbs items={[{ text: 'Test Automation Hub' }]} />
                 <Stack grow style={{ padding: 12, backgroundColor: modeColors.background, height: 'calc(100vh - 148px)', overflow: 'auto' } } horizontalAlign='center'>
-                    {initialized && !!handler.availableStreams.length &&
+                    {initialized && !!testDataId &&
+                        <TestPhasesView testDataId={testDataId!} handler={handler} onDismiss={() => {handler.removeSearchParam('session'); handler.removeSearchParam('phase')}} />
+                    }
+                    {initialized && !testDataId && !!handler.availableStreams.length &&
                         <Stack horizontal grow>
                             <TestAutomationSidebarLeft/>
                             <Stack style={{ minWidth: 1024 }}>
-                                <TestAutomationPanel />
+                                <TestAutomationPanel/>
                             </Stack>
                         </Stack>
                     }
diff --git a/Engine/Source/Programs/Horde/HordeDashboard/plugins/test-automation/testData.ts b/Engine/Source/Programs/Horde/HordeDashboard/plugins/test-automation/testData.ts
--- a/Engine/Source/Programs/Horde/HordeDashboard/plugins/test-automation/testData.ts
+++ b/Engine/Source/Programs/Horde/HordeDashboard/plugins/test-automation/testData.ts
@@ -1,7 +1,8 @@
 import { isNumber } from '@datadog/browser-core';
-import { action, makeObservable, observable, reaction } from 'mobx';
-import { TestOutcome } from './api';
+import { action, makeObservable, observable, reaction, when } from 'mobx';
+import { EventEntry, TestOutcome, TestPhaseOutcome } from './api';
 import { projectStore } from 'horde/backend/ProjectStore';
+import backend from 'horde/backend';
 
 const defaultQueryWeeks = 2;
 const filterRefreshMS = 200;
@@ -11,11 +12,13 @@ type TestId = string;
 type MetaId = string;
 type TagId = string;
 type TestSessionId = string;
+type TestDataId = string;
+type PhaseKey = string;
 
 // TestData Version registration
-interface TestVersionFetcher {
+export interface TestVersionFetcher {
 
-    get version(): string;
+    get version(): number;
 
     /// fetch stream information and populate stream map
     fetchStreamsInfo(streamMap: Map<string, TestStreamRef>): Promise<void>;
@@ -23,6 +26,12 @@ interface TestVersionFetcher {
     /// fetch test sessions and populate test status session
     fetchTestSessions(streamRef: TestStreamRef, status: TestStatus, weeks: number, onAdd?: (session: TestSessionResult) => void): Promise<void>;
 
+    /// fetch test data and run the test session details
+    fetchTestData(testDataId: TestDataId): Promise<TestSessionDetails | undefined>;
+
+    /// fetch test phase sessions and populate test status phase session
+    fetchPhaseSessions(streamRef: TestStreamRef, status: TestMetaStatus, phaseKey: string, weeks: number): Promise<boolean>;
+
 }
 
 export class TestDataVersionRegistrar {
@@ -31,7 +40,7 @@ export class TestDataVersionRegistrar {
 
     static register(fetcher: TestVersionFetcher) {
         if (TestDataVersionRegistrar._manifest.find(f => f.version === fetcher.version)) {
-            console.warn(`${fetcher.version} already registered`);
+            console.warn(`TestData fetcher version ${fetcher.version} already registered`);
         } else {
             TestDataVersionRegistrar._manifest.push(fetcher);
         }
@@ -43,8 +52,38 @@ export class TestDataVersionRegistrar {
 
 }
 
+export class PhaseSessionResult {
+    constructor(id: string, start: Date, duration: number, commitId: string, commitOrder: number, sessionId: string, phaseRef: string, metadataId: string, streamId: string, outcome: TestPhaseOutcome) {
+        this.id = id;
+        this.sessionId = sessionId;
+        this.commitId = commitId;
+        this.commitOrder = commitOrder;
+        this.start = start;
+        this.duration = duration;
+        this.streamId = streamId;
+        this.metadataId = metadataId;
+        this.phaseRef = phaseRef;
+        this.outcome = outcome;
+    }
+
+    id: string;
+    streamId: string;
+    metadataId: string;
+    sessionId: string;
+    commitId: string;
+    commitOrder: number;
+    duration: number;
+    start: Date;
+    phaseRef: string;
+    outcome: TestPhaseOutcome;
+    jobId?: string;
+    stepId?: string;
+    hasWarning?: boolean;
+    errorFingerprint?: string;
+}
+
 export class TestSessionResult {
-    constructor(id: string, start: Date, duration: string, commitId: string, commitOrder: number, nameRef: string, metadataId: string, streamId: string, outcome: TestOutcome, testDataId?: string) {
+    constructor(id: string, start: Date, duration: number, commitId: string, commitOrder: number, nameRef: string, metadataId: string, streamId: string, outcome: TestOutcome, testDataId?: string) {
         this.id = id;
         this.start = start;
         this.duration = duration;
@@ -60,9 +99,16 @@ export class TestSessionResult {
         this.phaseUnspecifiedCount = 0;
     }
 
+    get artifacts(): ArtifactFactory | undefined {
+        if (!this._artifactFactory && this.jobId && this.stepId) {
+            this._artifactFactory = new ArtifactFactory(this.jobId, this.stepId);
+        }
+        return this._artifactFactory;
+    }
+
     id: TestSessionId;
     start: Date;
-    duration: string;
+    duration: number;
     commitId: string;
     commitOrder: number;
     nameRef: TestId;
@@ -75,12 +121,15 @@ export class TestSessionResult {
     jobId?: string;
     stepId?: string;
     testDataId?: string;
+
+    private _artifactFactory?: ArtifactFactory;
 }
 
 export class TestSessionStatus {
     constructor(test: TestNameRef) {
         this.test = test;
         this.history = [];
+        this.phaseHistories = new Map();
     }
 
     addSession(session: TestSessionResult) {
@@ -94,12 +143,40 @@ export class TestSessionStatus {
         return this.history[0];
     }
 
+    getPreviousSession(session: TestSessionResult): TestSessionResult | undefined {
+        const index = this.history.indexOf(session);
+        if (index === -1) return;
+        if (index + 1 === this.history.length) return;
+
+        return this.history[index + 1];
+    }
+
     sortSessions() {
         /// sort session history by commit order descending
         this.history.sort((a, b) => b.commitOrder - a.commitOrder);
     }
 
+    addPhaseSession(phaseKey: PhaseKey, session: PhaseSessionResult) {
+        if (!this.phaseHistories.has(phaseKey)) {
+            this.phaseHistories.set(phaseKey, []);
+        }
+        this.phaseHistories.get(phaseKey)!.push(session);
+    }
+
+    getPhaseSessions(phaseKey: PhaseKey): PhaseSessionResult[] | undefined {
+        if (!this.phaseHistories.get(phaseKey)?.length) {
+            return;
+        }
+        return this.phaseHistories.get(phaseKey)!;
+    }
+
+    sortPhaseSessions(phaseId: string) {
+        /// sort phase session history by commit order descending
+        this.phaseHistories.get(phaseId)?.sort((a, b) => b.commitOrder - a.commitOrder);
+    }
+
     history: TestSessionResult[];
+    phaseHistories: Map<PhaseKey, PhaseSessionResult[]>;
     test: TestNameRef;
 }
 
@@ -134,10 +211,161 @@ export class TestTagRef {
     private lowerName: string;
 }
 
+export class TestPhaseRef {
+    constructor(id: string, key: string, name: string, version: number) {
+        this.id = id;
+        this.key = key;
+        this.name = name;
+        this.version = version;
+        this.lowerName = this.name.toLowerCase();
+    }
+
+    isMatch(toMatch: string): boolean {
+        return toMatch.charAt(0) === '@'? this.lowerName === toMatch.substring(1) : this.lowerName.includes(toMatch);
+    }
+
+    id: string;
+    name: string;
+    key: string;
+    version: number;
+
+    private lowerName: string;
+}
+
+export class DeviceRef {
+    constructor(name: string, logPath?: string, metadata?: Metadata, version?: number) {
+        this.name = name;
+        this.version = version ?? 0;
+        this.logPath = logPath;
+        this.metadata = metadata ?? {};
+    }
+
+    name: MetaId;
+    version: number;
+    logPath?: string;
+    metadata: Metadata;
+}
+
+export abstract class TestPhaseStatus {
+    constructor (key: string, name: string, start: Date, duration: number, devices: DeviceRef[], outcome: TestPhaseOutcome, sessionDetails: TestSessionDetails) {
+        this.key = key;
+        this._name = name;
+        this._start = start;
+        this._duration = duration;
+        this.devices = devices;
+        this.outcome = outcome;
+        this._sessionDetails = sessionDetails;
+    }
+
+    associateSession(phaseSessions?: PhaseSessionResult[]) {
+        if (!this.testSession) return;
+        this.session = phaseSessions?.find(s => s.sessionId === this.testSession!.id);
+    }
+
+    isMatch(items: string[]) {
+        if (items.length === 0) return true;
+
+        const lowerName = this.name.toLowerCase();
+        return items.some(i => i.charAt(0) === '@'? lowerName === i.substring(1) : lowerName.includes(i));
+    }
+
+    get name() {
+        return this.phaseRef?.name ?? this._name;
+    }
+
+    get testSession() {
+        return this._sessionDetails.session;
+    }
+
+    get artifacts() {
+        return this._sessionDetails.artifacts;
+    }
+
+    get testMetadata() {
+        return this._sessionDetails.meta;
+    }
+
+    get test() {
+        return this._sessionDetails.test;
+    }
+    get start() {
+        return this.session?.start ?? this._start;
+    }
+
+    get duration() {
+        return this.session?.duration ?? this._duration;
+    }
+
+    abstract fetchEventStream():  Promise<EventEntry[] | undefined>;
+
+    devices: DeviceRef[];
+    key: PhaseKey;
+    outcome: TestPhaseOutcome;
+
+    phaseRef?: TestPhaseRef;
+    session?: PhaseSessionResult;
+    componentRef?: any;
+    hasWarning?: boolean;
+
+    private _name: string;
+    private _start: Date;
+    private _duration: number;
+    private _sessionDetails: TestSessionDetails;
+}
+
+export class TestSessionDetails {
+    constructor(key: string, name: string, start: Date, duration: number, commitId: string, commitOrder: number, streamId: string, jobId: string, stepId: string) {
+        this.key = key;
+        this.name = name;
+        this._start = start;
+        this._duration = duration;
+        this.commitId = commitId;
+        this.commitOrder = commitOrder
+        this.streamId = streamId;
+        this.jobId = jobId;
+        this.stepId = stepId;
+        this.devices = new Map();
+        this.phases = [];
+    }
+
+    get start() {
+        return this.session?.start ?? this._start;
+    }
+
+    get duration() {
+        return this.session?.duration ?? this._duration;
+    }
+
+    get artifacts(): ArtifactFactory | undefined {
+        if (!this._artifactFactory && this.jobId && this.stepId) {
+            this._artifactFactory = new ArtifactFactory(this.jobId, this.stepId);
+        }
+        return this._artifactFactory;
+    }
+
+    key: string;
+    name: string;
+    commitId: string;
+    commitOrder: number;
+    jobId: string;
+    stepId: string;
+    streamId: string;
+    devices: Map<string, DeviceRef>;
+    phases: TestPhaseStatus[];
+    
+    test?: TestNameRef;
+    meta?: MetadataRef;
+    session?: TestSessionResult;
+
+    private _start: Date;
+    private _duration: number;
+    private _artifactFactory?: ArtifactFactory;
+}
+
 export class TestNameRef {
     constructor(id: string, key: string, name: string, version: number) {
         this.id = id;
-        this.key = name;
+        this.key = key;
         this.name = name;
         this.version = version;
         this.lowerKey = this.key.toLowerCase();
@@ -145,7 +373,7 @@ export class TestNameRef {
     }
 
     isMatch(toMatch: string): boolean {
-        return toMatch.charAt(0) === '#'? this.lowerName === toMatch.substring(1) || this.lowerKey === toMatch.substring(1) : this.lowerName.includes(toMatch) || this.lowerKey.includes(toMatch);
+        return toMatch.charAt(0) === '@'? this.lowerName === toMatch.substring(1) || this.lowerKey === toMatch.substring(1) : this.lowerName.includes(toMatch) || this.lowerKey.includes(toMatch);
     }
 
     updateTags(tags: TestTagRef[]) {
@@ -162,6 +390,7 @@ export class TestNameRef {
     name: string;
     version: number;
     tagRefs?: Set<TestTagRef>;
+    phases?: Map<PhaseKey, TestPhaseRef>;
 
     private lowerKey: string;
     private lowerName: string;
@@ -230,8 +459,6 @@ export class MetadataRef {
     private static keysOrderPriority = ["Project", "BuildTarget", "Platform"]; // Should those be set from a config?
     private static valuesToIgnore = new Set(["default", "Default"]);
 
-    private static globalMap: Map<string, MetadataRef> = new Map();
-
     static identifyCommonKeys(items: MetadataRef[]): string[] {
         if (!items.length) return [];
         const first = items[0].entries;
@@ -267,6 +494,29 @@ export class TestMetaStatus {
         return this.sessions.get(metaRef)?.getLastSession();
     }
 
+    getPreviousSession(metaRef: MetadataRef, session: TestSessionResult): TestSessionResult | undefined {
+        return this.sessions.get(metaRef)?.getPreviousSession(session);
+    }
+
+    addPhaseSession(metaRef: MetadataRef, phaseKey: PhaseKey, session: PhaseSessionResult) {
+        if (!this.sessions.has(metaRef)) {
+            this.sessions.set(metaRef, new TestSessionStatus(this.test));
+        }
+
+        this.sessions.get(metaRef)!.addPhaseSession(phaseKey, session);
+    }
+
+    getPhaseSessions(phaseKey: PhaseKey, metaRef?: MetadataRef): PhaseSessionResult[] | undefined {
+        if (!metaRef) {
+            return this.sessions.values().map(m => m.getPhaseSessions(phaseKey)).filter(m => m?.length).toArray().flat() as PhaseSessionResult[];
+        }
+        return this.sessions.get(metaRef)?.getPhaseSessions(phaseKey);
+    }
+
+    hasPhaseSession(phaseKey: PhaseKey) {
+        return this.sessions.values().some((s) => !!s.phaseHistories.get(phaseKey));
+    }
+
     getMetadata(): MetadataRef[] {
         return this.sessions.keys().toArray();
     }
@@ -289,6 +539,12 @@ export class TestMetaStatus {
         }
     }
 
+    sortPhaseSessions(phaseKey: PhaseKey) {
+        for (const status of this.sessions.values()) {
+            status.sortPhaseSessions(phaseKey);
+        }
+    }
+
     sessions: Map<MetadataRef, TestSessionStatus>;
     test: TestNameRef;
 }
@@ -310,6 +566,22 @@ export class TestStatus {
         return this.tests.get(testNameRef)?.getLastSession(metaRef);
     }
 
+    getPreviousSession(testNameRef: TestNameRef, metaRef: MetadataRef, session: TestSessionResult): TestSessionResult | undefined {
+        return this.tests.get(testNameRef)?.getPreviousSession(metaRef, session);
+    }
+
+    addPhaseSession(testNameRef: TestNameRef, metaRef: MetadataRef, phaseKey: PhaseKey, session: PhaseSessionResult) {
+        if (!this.tests.has(testNameRef)) {
+            this.tests.set(testNameRef, new TestMetaStatus(testNameRef));
+        }
+
+        this.tests.get(testNameRef)!.addPhaseSession(metaRef, phaseKey, session);
+    }
+
+    getPhaseSessions(testNameRef: TestNameRef, phaseKey: PhaseKey, metaRef?: MetadataRef): PhaseSessionResult[] | undefined {
+        return this.tests.get(testNameRef)?.getPhaseSessions(phaseKey, metaRef);
+    }
+
     getMatchingTests(testNameRefs: Set<TestNameRef>, metaRefs: Set<MetadataRef>): TestNameRef[] {
         return this.tests.entries().filter(([test, status]) => testNameRefs.has(test) && status.includesAnyMetadata(metaRefs)).map(([test, status]) => test).toArray();
     }
@@ -534,12 +806,17 @@ export class TestDataHandler {
     constructor() {
         TestDataHandler.instance = this;
         makeObservable(this);
-        // setup a reaction on filterUpdated to debounce the filter refresh
+        // setup reactions to debounce the search and filter refresh
         reaction(
             () => this.filterUpdated,
             () => this.updateFilter(),
             { delay: filterRefreshMS }
         );
+        reaction(
+            () => this.checkSearchUpdated,
+            () => this.onSearchUpdated(),
+            { delay: 100 }
+        );
     }
 
     async initialize(search: string, searchCallback?: (search: string) => void) {
@@ -567,8 +844,8 @@ export class TestDataHandler {
     }
 
     private checkSearchCallback(previous: string) {
-        if (this.searchCallback && previous != this.search.toString()) {
-            this.searchCallback(this.search.toString());
+        if (previous != this.search.toString()) {
+            this.setCheckSearchUpdated();
         }
     }
 
@@ -800,10 +1077,14 @@ export class TestDataHandler {
 
     private updateFilter() {
         if (this.checkUpdateFilter()) {
-            this.setSearchUpdated();
-            if (this.searchCallback) {
-                this.searchCallback(this.search.toString());
-            }
+            this.onSearchUpdated();
+        }
+    }
+
+    private onSearchUpdated() {
+        this.setSearchUpdated();
+        if (this.searchCallback) {
+            this.searchCallback(this.search.toString());
         }
     }
 
@@ -870,6 +1151,10 @@ export class TestDataHandler {
 
     async query(streamId?: StreamId, isSubQuery: boolean = false) {
 
+        if (this.queryLoading || this.subQueryLoading) {
+            await when(() => !this.queryLoading && !this.subQueryLoading);
+        }
+
         if (!streamId || !this.testStreams.has(streamId) || this.testStatus.has(streamId)) {
             isSubQuery? this.setSubQueryLoading(false): this.setQueryLoading(false);
             return;
@@ -896,7 +1181,7 @@ export class TestDataHandler {
         const queryWeeks = state.weeks && state.weeks > defaultQueryWeeks ? state.weeks : defaultQueryWeeks;
 
         for (const fetcherVersion of TestDataVersionRegistrar.manifest) {
-            console.log(`Fetching from ${fetcherVersion.version} test sessions`)
+            console.debug(`Fetching from version ${fetcherVersion.version} test sessions`)
             await fetcherVersion.fetchTestSessions(testStreamRef, testStatus, queryWeeks, aggregateCommitDates);
         }
 
@@ -909,6 +1194,93 @@ export class TestDataHandler {
         }
     }
 
+    async queryTestData(testDataId: string): Promise<TestSessionDetails | undefined> {
+        if (this.testTestData.has(testDataId)) {
+            this.setTestDataQueryLoading(false);
+            return this.testTestData.get(testDataId);
+        }
+
+        this.setTestDataQueryLoading(true);
+
+        for (const fetcherVersion of TestDataVersionRegistrar.manifest) {
+            console.debug(`Fetching from version ${fetcherVersion.version} test data`)
+            try {
+                const sessionDetails = await fetcherVersion.fetchTestData(testDataId);
+                if (sessionDetails) {
+                    await this.query(sessionDetails.streamId);
+                    const session = this.getStatusStream(sessionDetails.streamId)?.tests.get(sessionDetails.test!)?.sessions.get(sessionDetails.meta!)?.history.find(s => s.testDataId === testDataId);
+                    sessionDetails.session = session;
+                    this.testTestData.set(testDataId, sessionDetails)
+                    /// stop fetching if a valid object was pulled
+                    break;
+                }
+            }  catch(reason) {
+                /** ignore failed attempt to try with next version */
+                console.debug(`Failed to fetch from version ${fetcherVersion.version}: ${reason}`);
+            }
+        }
+
+        if (this.loaded) {
+            this.setTestDataQueryLoading(false);
+        }
+
+        return this.testTestData.get(testDataId);
+    }
+
+    async queryPhase(phase: TestPhaseStatus, streamId: StreamId, isSubQuery: boolean = false) {
+
+        const phaseKey = phase.key;
+        const test = phase.test!;
+        const testStatus = this.testStatus.get(streamId)?.tests.get(test);
+
+        const associatePhaseSession = () => {
+            if (!phase.session && testStatus) {
+                /// associate the corresponding phase session with the phase status
+                phase.associateSession(testStatus.getPhaseSessions(phaseKey));
+                /// associate the corresponding phase reference
+                phase.phaseRef = test.phases?.get(phaseKey);
+            }
+        }
+
+        if (!testStatus || testStatus.hasPhaseSession(phaseKey)) {
+            associatePhaseSession();
+            isSubQuery? this.setSubPhaseQueryLoading(false): this.setPhaseQueryLoading(false);
+            return;
+        }
+
+        isSubQuery? this.setSubPhaseQueryLoading(true): this.setPhaseQueryLoading(true);
+
+        await this.query(streamId);
+
+        const state = this.filterState;
+        const queryWeeks = state.weeks && state.weeks > defaultQueryWeeks ? state.weeks : defaultQueryWeeks;
+
+        const testStreamRef = this.testStreams.get(streamId)!;
+
+        for (const fetcherVersion of TestDataVersionRegistrar.manifest) {
+            /// fetch only from the matching version
+            if (fetcherVersion.version === test.version) {
+                console.debug(`Fetching from version ${fetcherVersion.version} test phase sessions`);
+                try {
+                    await fetcherVersion.fetchPhaseSessions(testStreamRef, testStatus, phaseKey, queryWeeks);
+                } catch (reason) {
+                    /** ignore failure as getting the related phase sessions is optional to display content */
+                    console.debug(`Failed to fetch from version ${fetcherVersion.version}: ${reason}`);
+                }
+                break;
+            }
+        }
+
+        /// sort phase session history by commit order descending
+        testStatus.sortPhaseSessions(phaseKey);
+
+        associatePhaseSession();
+
+        if (this.loaded) {
+            isSubQuery? this.setSubPhaseQueryLoading(false): this.setPhaseQueryLoading(false);
+        }
+    }
+
     // load immutable data for this view session
     private async load() {
 
@@ -916,7 +1288,7 @@ export class TestDataHandler {
         const testStreams = this.testStreams = new Map(allStreams.map((s) => [s, new TestStreamRef(s)]));
 
         for (const fetcherVersion of TestDataVersionRegistrar.manifest) {
-            console.log(`Fetching from ${fetcherVersion.version} stream info`)
+            console.debug(`Fetching from version ${fetcherVersion.version} stream info`)
             await fetcherVersion.fetchStreamsInfo(testStreams);
         }
 
@@ -936,6 +1308,11 @@ export class TestDataHandler {
         this.searchUpdated++;
     }
 
+    @action
+    setCheckSearchUpdated() {
+        this.checkSearchUpdated++;
+    }
+
     @action
     setQueryLoading(value: boolean) {
         this.queryLoading = value;
@@ -946,6 +1323,21 @@ export class TestDataHandler {
         this.subQueryLoading = value;
     }
 
+    @action
+    setTestDataQueryLoading(value: boolean) {
+        this.testDataQueryLoading = value;
+    }
+
+    @action
+    setPhaseQueryLoading(value: boolean) {
+        this.phaseQueryLoading = value;
+    }
+
+    @action
+    setSubPhaseQueryLoading(value: boolean) {
+        this.subPhaseQueryLoading = value;
+    }
+
     get stream(): StreamId | undefined {
         return this.filterState.stream;
     }
@@ -975,12 +1367,24 @@ export class TestDataHandler {
     @observable
     searchUpdated: number = 0;
 
+    @observable
+    checkSearchUpdated: number = 0;
+
     @observable
     queryLoading: boolean = false;
 
     @observable
     subQueryLoading: boolean = false;
 
+    @observable
+    testDataQueryLoading: boolean = false;
+
+    @observable
+    phaseQueryLoading: boolean = false;
+
+    @observable
+    subPhaseQueryLoading: boolean = false;
+
     subscribeToFilter() {
         if (this.filterUpdated) { }
     }
@@ -997,12 +1401,26 @@ export class TestDataHandler {
         if (this.subQueryLoading) { }
     }
 
+    subscribeToTestDataQueryLoading() {
+        if (this.testDataQueryLoading) { }
+    }
+
+    subscribeToPhaseQueryLoading() {
+        if (this.phaseQueryLoading) { }
+    }
+
+    subscribeToSubPhaseQueryLoading() {
+        if (this.subPhaseQueryLoading) { }
+    }
+
     private searchCallback?: (search: string) => void;
     private search: TestDataFilterSearchParams = new TestDataFilterSearchParams();
     filterState: TestDataState = {};
 
     testStreams: Map<StreamId, TestStreamRef> = new Map();
 
+    testTestData: Map<TestDataId, TestSessionDetails> = new Map();
+
     // quantized common commits to land at same time in the views
     commitIdDates = new Map<string, Date>();
     
@@ -1012,3 +1430,61 @@ export class TestDataHandler {
 
     static instance: TestDataHandler;
 }
+
+export class ArtifactFactory {
+    constructor(jobId: string, stepId: string) {
+        this.jobId = jobId;
+        this.stepId = stepId;
+    }
+
+    private async getArtifactId() {
+
+        if (this.artifactId) {
+            return;
+        }
+
+        const v = await backend.getJobArtifactsV2(undefined, [`job:${this.jobId}/step:${this.stepId}`]);
+        const av2 = v?.artifacts.find(a => a.type === "step-saved")
+        if (av2?.id) {
+            this.artifactId = av2.id;
+        } else {
+            console.error("Unable to get step-saved artifacts for test report");
+        }
+
+    }
+
+    async getLink(referencePath: string): Promise<string> {
+
+        await this.getArtifactId();
+
+        return `${backend.serverUrl}/api/v2/artifacts/${this.artifactId}/file?path=${encodeURIComponent(referencePath)}&inline=true`;
+
+    }
+
+    async fetch(referencePath: string): Promise<any> {
+
+        const cacheKey = `job:${this.jobId}/step:${this.stepId}/${referencePath}`;
+        const cache = ArtifactFactory.cache.get(cacheKey);
+        if (!!cache) {
+            return cache;
+        }
+
+        const link = await this.getLink(referencePath);
+
+        return new Promise<any>((resolve, reject) => {
+            backend.fetch.get(link, {}).then((value) => {
+                ArtifactFactory.cache.set(cacheKey, value.data);
+                resolve(value.data);
+            }).catch((reason) => {
+                reject(reason);
+            });
+        });
+
+    }
+
+    private jobId: string;
+    private stepId: string;
+    private artifactId?: string;
+
+    private static cache: Map<string, any> = new Map();
+}
diff --git a/Engine/Source/Programs/Horde/HordeDashboard/plugins/test-automation/testDataV2Fetcher.ts b/Engine/Source/Programs/Horde/HordeDashboard/plugins/test-automation/testDataV2Fetcher.ts
--- a/Engine/Source/Programs/Horde/HordeDashboard/plugins/test-automation/testDataV2Fetcher.ts
+++ b/Engine/Source/Programs/Horde/HordeDashboard/plugins/test-automation/testDataV2Fetcher.ts
@@ -1,10 +1,27 @@
-import { getStreamTestsV2, GetTestMetadataResponse, GetTestNameResponse, GetTestSessionResponse, getTestSessionsV2, GetTestTagResponse } from "./api";
-import { TestStreamRef, TestStatus, TestSessionResult, TestNameRef, MetadataRef, TestTagRef } from "./testData"
+import { EventEntry, GetPhaseResponse, getStreamTestsV2, getTestDataV2, GetTestMetadataResponse, GetTestNameResponse, GetTestPhaseSessionResponse, getTestPhaseSessionsV2, getTestPhasesV2, GetTestSessionResponse, getTestSessionsV2, GetTestTagResponse, PhaseSessionDetailsResponse, TestData, TestSessionDetailsResponse } from "./api";
+import { TestStreamRef, TestStatus, TestSessionResult, TestNameRef, MetadataRef, TestTagRef, PhaseSessionResult, TestPhaseRef, TestMetaStatus, TestSessionDetails, TestPhaseStatus, DeviceRef, TestVersionFetcher } from "./testData"
 
 const minuteInWeek = 10080;
 
-export class TestDataV2 {
-    get version() { return "TestData Version 2" }
+class TestPhaseStatusV2 extends TestPhaseStatus {
+
+    constructor(phase: PhaseSessionDetailsResponse, devices: DeviceRef[], parent: TestSessionDetails) {
+        super(phase.key, phase.name, new Date(phase.dateTime), phase.timeElapseSec, devices, phase.outcome, parent)
+        this._eventStreamPath = phase.eventStreamPath;
+        this.hasWarning = phase.hasWarning;
+    }
+
+    async fetchEventStream(): Promise<EventEntry[] | undefined> {
+        if (!this._eventStreamPath) return;
+        return await this.artifacts?.fetch(this._eventStreamPath);
+    }
+
+    private _eventStreamPath?: string;
+
+}
+
+export class TestDataV2 implements TestVersionFetcher {
+    get version() { return 2 }
 
     async fetchStreamsInfo(testStreams: Map<string, TestStreamRef>) {
 
@@ -34,8 +51,8 @@ export class TestDataV2 {
     async fetchTestSessions(streamRef: TestStreamRef, status: TestStatus, weeks: number, onAdd?: (session: TestSessionResult) => void) {
 
         const streamId = streamRef.id;
-        const tests: Map<string, TestNameRef> = new Map(streamRef.tests.filter((t) => t.version >= 2).map((t) => [t.id, t]));
-        const meta: Map<string, MetadataRef> = new Map(streamRef.meta.filter((m) => m.version >= 2).map((m) => [m.id, m]));
+        const tests: Map<string, TestNameRef> = new Map(streamRef.tests.filter((t) => t.version >= this.version).map((t) => [t.id, t]));
+        const meta: Map<string, MetadataRef> = new Map(streamRef.meta.filter((m) => m.version >= this.version).map((m) => [m.id, m]));
 
         const minQueryDate = new Date(new Date().valueOf() - (minuteInWeek * weeks * 60000));
         const maxQueryDate = new Date();
@@ -58,6 +75,52 @@ export class TestDataV2 {
 
     }
 
+    async fetchTestData(testDataId: string): Promise<TestSessionDetails | undefined> {
+
+        const response = await getTestDataV2(testDataId);
+        if (!!response) {
+            const sessionDetails = this.createSessionDetails(response);
+            return sessionDetails;
+        }
+
+        return undefined;
+
+    }
+
+    async fetchPhaseSessions(streamRef: TestStreamRef, status: TestMetaStatus, phaseKey: string, weeks: number): Promise<boolean> {
+
+        const streamId = streamRef.id;
+        const meta: Map<string, MetadataRef> = new Map(streamRef.meta.filter((m) => m.version >= this.version).map((m) => [m.id, m]));
+
+        if (!status.test.phases) {
+            const responses = await getTestPhasesV2([status.test.id]);
+            status.test.phases = new Map(responses.map(i => i.phases.map(p => this.createPhaseRef(p))).flat().map(p => [p.key, p]));
+        }
+        
+        const phaseId = status.test.phases.get(phaseKey)?.id;
+        if (phaseId) {
+            const minQueryDate = new Date(new Date().valueOf() - (minuteInWeek * weeks * 60000));
+            const maxQueryDate = new Date();
+    
+            const responses = await getTestPhaseSessionsV2([streamId], [phaseId], minQueryDate.toISOString(), maxQueryDate.toISOString());
+            responses.forEach((session) => {
+                if (session.streamId !== streamId) {
+                    return;
+                }
+                const phaseSession = this.createPhaseSession(session);
+                const metadataRef = meta.get(phaseSession.metadataId);
+                if (!!metadataRef) {
+                    status.addPhaseSession(metadataRef, phaseKey, phaseSession);
+                }
+            });
+
+            return true;
+        }
+
+        return false;
+
+    }
+
     private testNameRefMap: Map<string, TestNameRef> = new Map();
     private metadataRefMap: Map<string, MetadataRef> = new Map();
     private testTageRefMap: Map<string, TestTagRef> = new Map();
@@ -66,7 +129,7 @@ export class TestDataV2 {
         const globalKey = test.id;
         let item = this.testNameRefMap.get(globalKey);
         if (!item) {
-            item = new TestNameRef(test.id, test.key, test.name, 2);
+            item = new TestNameRef(test.id, test.key, test.name, this.version);
             this.testNameRefMap.set(globalKey, item);
         }
 
@@ -77,7 +140,7 @@ export class TestDataV2 {
         const globalKey = item.id;
         let meta = this.metadataRefMap.get(globalKey);
         if (!meta) {
-            meta = new MetadataRef(item.id, item.entries, 2);
+            meta = new MetadataRef(item.id, item.entries, this.version);
             this.metadataRefMap.set(globalKey, meta);
         }
 
@@ -88,16 +151,20 @@ export class TestDataV2 {
         const globalKey = item.id;
         let tag = this.testTageRefMap.get(globalKey);
         if (!tag) {
-            tag = new TestTagRef(item.id, item.name, 2);
+            tag = new TestTagRef(item.id, item.name, this.version);
             this.testTageRefMap.set(globalKey, tag);
         }
 
         return tag;
     }
 
+    private createPhaseRef(item: GetPhaseResponse): TestPhaseRef {
+        return new TestPhaseRef(item.id, item.key, item. name, this.version);
+    }
+
     private createTestSession(test: GetTestSessionResponse): TestSessionResult {
         const start = new Date(test.startDateTime);
-        const result = new TestSessionResult(test.id, start, test.duration, test.commitId, test.commitOrder, test.nameRef, test.metadataId, test.streamId, test.outcome);
+        const result = new TestSessionResult(test.id, start, test.duration, test.commitId, test.commitOrder, test.nameRef, test.metadataId, test.streamId, test.outcome, test.testDataId);
         result.phaseSucceededCount = test.phasesSucceededCount;
         result.phaseFailedCount = test.phasesFailedCount;
         result.phaseUnspecifiedCount = test.phasesUndefinedCount;
@@ -106,5 +173,44 @@ export class TestDataV2 {
 
         return result;
     }
+
+    private createPhaseSession(phase: GetTestPhaseSessionResponse): PhaseSessionResult {
+        const start = new Date(phase.startDateTime);
+        const result = new PhaseSessionResult(phase.id, start, phase.duration, phase.commitId, phase.commitOrder, phase.sessionId, phase.phaseRef, phase.metadataId, phase.streamId, phase.outcome);
+        result.hasWarning = phase.hasWarning;
+        result.errorFingerprint = phase.errorFingerprint;
+        result.jobId = phase.jobId;
+        result.stepId = phase.stepId;
+
+        return result;
+    }
+
+    private createSessionDetails(testData: TestData) {
+        const testDetails = testData.data as TestSessionDetailsResponse;
+        const start = new Date(testDetails.summary.dateTime);
+        const result = new TestSessionDetails(testData.key, testDetails.summary.testName, start, testDetails.summary.timeElapseSec, testData.commitId.name, testData.commitId.order, testData.streamId, testData.jobId, testData.stepId);
+        Object.entries(testDetails.devices).forEach(([key, d]) => {
+            result.devices.set(key, new DeviceRef(d.name, d.appInstanceLogPath, d.metadata, this.version));
+        });
+        result.phases.push(...testDetails.phases.map(p => this.createPhaseStatus(p, result.devices, result)));
+
+        // search for known metadata
+        const metaFilter = new Map(Object.entries(testDetails.metadata).map(([key, value]) => [key, [value]]));
+        const metaRef = this.metadataRefMap.values().find((meta) => Object.keys(meta.entries).length === metaFilter.size &&  meta.isMatch(metaFilter)) ?? this.createMetadataRef({id: testData.id, entries: testDetails.metadata});
+        result.meta = metaRef;
+
+        // search for known test name
+        const testNameRef = this.testNameRefMap.values().find((nameRef) => nameRef.key === testData.key) ?? this.createTestNameRef({id: testData.id, key: testData.key, name: result.name});
+        result.test = testNameRef;
+
+        return result;
+    }
+
+    private createPhaseStatus(phase: PhaseSessionDetailsResponse, deviceRefs: Map<string, DeviceRef>, parent: TestSessionDetails) {
+        const devices = phase.deviceKeys.map(key => deviceRefs.get(key) ?? new DeviceRef(key));
+        const result = new TestPhaseStatusV2(phase, devices, parent);
+
+        return result;
+    }
 }
 
diff --git a/Engine/Source/Programs/Horde/HordeDashboard/plugins/test-automation/testHistoryGraph.ts b/Engine/Source/Programs/Horde/HordeDashboard/plugins/test-automation/testHistoryGraph.ts
--- a/Engine/Source/Programs/Horde/HordeDashboard/plugins/test-automation/testHistoryGraph.ts
+++ b/Engine/Source/Programs/Horde/HordeDashboard/plugins/test-automation/testHistoryGraph.ts
@@ -4,7 +4,7 @@ import { TestDataHandler, TestNameRef, MetadataRef, TestMetaStatus, TestSessionR
 import dashboard, { StatusColor } from 'horde/backend/Dashboard';
 import { getHumanTime, getShortNiceTime, msecToElapsed } from 'horde/base/utilities/timeUtils';
 import { getHordeStyling } from 'horde/styles/Styles';
-import { SessionStatusBar, SessionValues } from "./testAutomationCommon";
+import { SessionValues } from "./testAutomationCommon";
 import { renderToString } from "react-dom/server"
 import { TestOutcome } from "./api";
 
@@ -14,6 +14,7 @@ const _d3 = d3 as any;
 type SelectionType = d3.Selection<SVGGElement, unknown, null, undefined>;
 type DivSelectionType = d3.Selection<HTMLDivElement, unknown, null, undefined>;
 
+const clampMin = (x: number) => x !== 0 && x <= 0.3? 0.3: x;
 
 export class TestHistoryGraph {
 
@@ -193,7 +194,7 @@ export class TestHistoryGraph {
             .append("g")
             .attr("class", "circleScale")
             .attr("id", i => `circleScale${sessions[i as any].id}`)
-            .each((d, i, nodes) => SessionPieChart(nodes[i], sessions[d as any], 5));
+            .each((d, i, nodes) => SessionConcentricCircle(nodes[i], sessions[d as any], 5));
 
         const tooltip = this.tooltip = d3.select(container)
             .append("div")
@@ -357,10 +358,10 @@ export class TestHistoryGraph {
     tooltipSessionId?: string;
 }
 
-type PieChartStack = {value: number, name: string, color: string, title?: boolean};
+type ItemStack = {value: number, name: string, color: string, title?: boolean};
 
-const PieChart = (container: SVGGElement, data: PieChartStack[], radius: number, innerRadius: number = 0, padAngle: number = 0) => {
-    const pie = d3.pie<PieChartStack>().padAngle(padAngle).value((d) => d.value).sort(null);
+const PieChart = (container: SVGGElement, data: ItemStack[], radius: number, innerRadius: number = 0, padAngle: number = 0) => {
+    const pie = d3.pie<ItemStack>().padAngle(padAngle).value((d) => d.value).sort(null);
     const arc = d3.arc().innerRadius(innerRadius).outerRadius(radius);
     const color = d3.scaleOrdinal()
         .domain(data.map(d => d.name))
@@ -387,7 +388,7 @@ const SessionPieChart = (container: SVGGElement, session: TestSessionResult, rad
 
     const statusColors = dashboard.getStatusColors();
 
-    const stack: PieChartStack[] = [
+    const stack: ItemStack[] = [
         {
             value: metaUnspecifiedFactor * 100,
             name: "Unspecified",
@@ -407,3 +408,52 @@ const SessionPieChart = (container: SVGGElement, session: TestSessionResult, rad
 
     return PieChart(container, stack, radius, innerRadius, padAngle);
 }
+
+const ConcentricCircle = (container: SVGGElement, data: ItemStack[], radius: number) => {
+    const color = d3.scaleOrdinal()
+        .domain(data.map(d => d.name))
+        .range(data.map(d => d.color));
+    
+    const g = d3.select(container)
+        .selectAll("centric")
+        .data(data)
+        .join("circle")
+            .attr("class", "centric")
+            .attr("fill", d => color(d.name) as any)
+            .attr("r", d => radius * d.value)
+        .append("title")
+            .text(d => d.title? `${Math.ceil(d.value)}% ${d.name}` : "");
+
+    return g;
+}
+
+const SessionConcentricCircle = (container: SVGGElement, session: TestSessionResult, radius: number) => {
+
+    const sessionTotalCount = session.phaseSucceededCount + session.phaseFailedCount + session.phaseUnspecifiedCount;
+    const metaFailedFactor = Math.ceil(clampMin(session.phaseFailedCount / (sessionTotalCount || 1)) * 5) / 5;
+    const metaUnspecifiedFactor = Math.ceil(clampMin(session.phaseUnspecifiedCount / (sessionTotalCount || 1)) * 5) / 5;
+    let metaSuccessFactor = 1 - (metaFailedFactor + metaUnspecifiedFactor);
+    metaSuccessFactor = metaSuccessFactor < 0? 0 : metaSuccessFactor
+
+    const statusColors = dashboard.getStatusColors();
+
+    const stack: ItemStack[] = [
+        {
+            value: metaSuccessFactor && 1,
+            name: "Passed",
+            color: statusColors.get(StatusColor.Success)!
+        },
+        {
+            value: metaUnspecifiedFactor && ((metaSuccessFactor && (metaFailedFactor + metaUnspecifiedFactor)) || 1),
+            name: "Unspecified",
+            color: statusColors.get(StatusColor.Unspecified)!,
+        },
+        {
+            value: metaFailedFactor,
+            name: "Failure",
+            color: statusColors.get(StatusColor.Failure)!,
+        },
+    ];
+
+    return ConcentricCircle(container, stack, radius);
+}
