Fix warning in ITestDataCollection

#rnx

diff --git a/Engine/Source/Programs/Horde/Plugins/Build/HordeServer.Build/Jobs/TestData/ITestDataCollection.cs b/Engine/Source/Programs/Horde/Plugins/Build/HordeServer.Build/Jobs/TestData/ITestDataCollection.cs
--- a/Engine/Source/Programs/Horde/Plugins/Build/HordeServer.Build/Jobs/TestData/ITestDataCollection.cs
+++ b/Engine/Source/Programs/Horde/Plugins/Build/HordeServer.Build/Jobs/TestData/ITestDataCollection.cs
@@ -1,6 +1,5 @@
 // Copyright Epic Games, Inc. All Rights Reserved.
 
-using System.Threading;
 using EpicGames.Horde.Commits;
 using EpicGames.Horde.Jobs;
 using EpicGames.Horde.Streams;
