Fixed an exception that was occuring in the LocalizationIssueHandler

diff --git a/Engine/Source/Programs/Horde/Plugins/Build/HordeServer.Build.Tests/Issues/IssueHandlerTests.cs b/Engine/Source/Programs/Horde/Plugins/Build/HordeServer.Build.Tests/Issues/IssueHandlerTests.cs
--- a/Engine/Source/Programs/Horde/Plugins/Build/HordeServer.Build.Tests/Issues/IssueHandlerTests.cs
+++ b/Engine/Source/Programs/Horde/Plugins/Build/HordeServer.Build.Tests/Issues/IssueHandlerTests.cs
@@ -73,13 +73,14 @@ namespace HordeServer.Tests.Issues
 				"{ \"time\":\"2025-01-01T01:20:30.996Z\",\"level\":\"Error\",\"message\":\"LogLocTextHelper: Error: /Path/To/File6.cpp(60): Failed to add text from LOCTEXT for namespace 'NamespaceC' and key 'KeyC' with source 'Text from File6'.\",\"properties\":{ \"_channel\":{ \"$type\":\"Channel\",\"$text\":\"LogLocTextHelper\"},\"_severity\":{ \"$type\":\"Severity\",\"$text\":\"Error\"},\"location\":\"/Path/To/File6.cpp(60)\",\"from\":\" from LOCTEXT\",\"locNamespace\":\"NamespaceC\",\"locKey\":\"KeyC\",\"locKeyMetaData\":\"\",\"text\":\"Text from File6\",\"textMetaData\":\"\",\"id\":304},\"format\":\"{_channel}: {_severity}: {location}: Failed to add text{from} for namespace '{locNamespace}' and key '{locKey}'{locKeyMetaData} with source '{text}'{textMetaData}.\"}",
 				"{ \"time\":\"2025-01-01T01:20:30.997Z\",\"level\":\"Error\",\"message\":\"LogLocTextHelper: Error: Failed to add text for namespace 'NamespaceA' and key 'KeyA' with source 'TextA'.\",\"properties\":{ \"_channel\":{ \"$type\":\"Channel\",\"$text\":\"LogLocTextHelper\"},\"_severity\":{ \"$type\":\"Severity\",\"$text\":\"Error\"},\"locNamespace\":\"NamespaceA\",\"locKey\":\"KeyA\",\"locKeyMetaData\":\"\",\"text\":\"TextA\",\"textMetaData\":\"\",\"forPlatform\":\"\",\"id\":304},\"format\":\"{_channel}: {_severity}: Failed to add text for namespace '{locNamespace}' and key '{locKey}'{locKeyMetaData} with source '{text}'{textMetaData}{forPlatform}.\"}",
 				"{ \"time\":\"2025-01-01T01:20:30.998Z\",\"level\":\"Error\",\"message\":\"LogLocTextHelper: Error: Failed to add text for namespace 'NamespaceB' and key 'KeyB' with source 'TextB'.\",\"properties\":{ \"_channel\":{ \"$type\":\"Channel\",\"$text\":\"LogLocTextHelper\"},\"_severity\":{ \"$type\":\"Severity\",\"$text\":\"Error\"},\"locNamespace\":\"NamespaceB\",\"locKey\":\"KeyB\",\"locKeyMetaData\":\"\",\"text\":\"TextB\",\"textMetaData\":\"\",\"forPlatform\":\"\",\"id\":304},\"format\":\"{_channel}: {_severity}: Failed to add text for namespace '{locNamespace}' and key '{locKey}'{locKeyMetaData} with source '{text}'{textMetaData}{forPlatform}.\"}",
+				"{ \"time\":\"2025-01-01T01:20:30.999Z\",\"level\":\"Warning\",\"message\":\"LogLocTextHelper: Warning: Text conflict for namespace '' and key 'ABCDEFGHIJKLMNOPQRSTUVWXYZABCDEH'. First entry is /SamePlugin/Scripts/SubFolder/AssetName.AssetName.Sub.Message.ContextualMessages(0).ContextualMessages.Messages(0).Messages.Message:'Text in first asset.'. Conflicting entry is /SamePlugin/Scripts/SubFolder/CopiedAsset.CopiedAsset.Sub.Message.ContextualMessages(0).ContextualMessages.Messages(0).Messages.Message:'Text in second asset.'.\",\"properties\":{ \"_channel\":{ \"$type\":\"Channel\",\"$text\":\"LogLocTextHelper\"},\"_severity\":{ \"$type\":\"Severity\",\"$text\":\"Warning\"},\"from\":\"\",\"location\":\"/SamePlugin/Scripts/SubFolder/AssetName.AssetName.Sub.Message.ContextualMessages(0).ContextualMessages.Messages(0).Messages.Message\",\"locNamespace\":\"\",\"locKey\":\"ABCDEFGHIJKLMNOPQRSTUVWXYZABCDEH\",\"locKeyMetaData\":\"\",\"text\":\"Text in first asset.\",\"textMetaData\":\"\",\"conflictLocation\":\"/SamePlugin/Scripts/SubFolder/CopiedAsset.CopiedAsset.Sub.Message.ContextualMessages(0).ContextualMessages.Messages(0).Messages.Message\",\"conflictText\":\"Text in second asset.\",\"conflictTextMetaData\":\"\",\"id\":304},\"format\":\"{_channel}: {_severity}: Text conflict{from} for namespace '{locNamespace}' and key '{locKey}'{locKeyMetaData}. First entry is {location}:'{text}'{textMetaData}. Conflicting entry is {conflictLocation}:'{conflictText}'{conflictTextMetaData}.\"}", // This "location" created an exception in the code at one point
 			};
 
 			IssueHandlerContext context = new IssueHandlerContext(new StreamId(), new TemplateId(), "Localize ProjectX", new Dictionary<string, string>());
 			LocalizationIssueHandler localizationIssueHandler = new LocalizationIssueHandler(context);
 			List<IssueEventGroup> groupedIssues = TestSingleIssueHandler(localizationIssueHandler, jsonStructuredLogs);
 
-			Assert.IsTrue(groupedIssues.Count == 8);
+			Assert.IsTrue(groupedIssues.Count == 9);
 
 			Assert.IsTrue(groupedIssues[0].Keys.Count == 3);
 			Assert.IsTrue(groupedIssues[0].Keys.Contains(new IssueKey("File1", IssueKeyType.File)));
@@ -173,6 +174,18 @@ namespace HordeServer.Tests.Issues
 			Assert.IsTrue(groupedIssues[7].Metadata.Count == 0);
 			Assert.IsTrue(groupedIssues[7].SummaryTemplate == "Localization {Severity}");
 			Assert.IsTrue(groupedIssues[7].Type == "Localization");
+
+			Assert.IsTrue(groupedIssues[8].Keys.Count == 3);
+			Assert.IsTrue(groupedIssues[8].Keys.Contains(new IssueKey("AssetName", IssueKeyType.File)));
+			Assert.IsTrue(groupedIssues[8].Keys.Contains(new IssueKey("CopiedAsset", IssueKeyType.File)));
+			Assert.IsTrue(groupedIssues[8].Keys.Contains(new IssueKey(",ABCDEFGHIJKLMNOPQRSTUVWXYZABCDEH", IssueKeyType.None)));
+			Assert.IsTrue(groupedIssues[8].ChangeFilter == "...");
+			Assert.IsTrue(groupedIssues[8].Events.Count == 1);
+			Assert.IsTrue(groupedIssues[8].Events[0].EventId == KnownLogEvents.Engine_Localization);
+			Assert.IsTrue(groupedIssues[8].Events[0].Severity == Microsoft.Extensions.Logging.LogLevel.Warning);
+			Assert.IsTrue(groupedIssues[8].Metadata.Count == 1);
+			Assert.IsTrue(groupedIssues[8].SummaryTemplate == "{Meta:Node}: Localization {Severity} in {Files}");
+			Assert.IsTrue(groupedIssues[8].Type == "Localization");
 		}
 	}
 }
diff --git a/Engine/Source/Programs/Shared/EpicGames.Horde/Issues/Handlers/LocalizationIssueHandler.cs b/Engine/Source/Programs/Shared/EpicGames.Horde/Issues/Handlers/LocalizationIssueHandler.cs
--- a/Engine/Source/Programs/Shared/EpicGames.Horde/Issues/Handlers/LocalizationIssueHandler.cs
+++ b/Engine/Source/Programs/Shared/EpicGames.Horde/Issues/Handlers/LocalizationIssueHandler.cs
@@ -156,47 +156,43 @@ namespace EpicGames.Horde.Issues.Handlers
 				}
 			}
 
-			ConvertLocationToFileAndLineNumber(data);
+			ConvertLocationToFile(data);
 			RemoveFilesExtension(data);
 
 			return data;
 		}
 
 		/// <summary>
-		/// Location data in LocalizationData may contain file and line information. We can try to parse it.
+		/// Location data in LocalizationData may contain file information. We can try to parse it.
 		/// </summary>
 		/// <param name="data">The localization data</param>
-		private static void ConvertLocationToFileAndLineNumber(LocalizationData data)
+		private static void ConvertLocationToFile(LocalizationData data)
 		{
-			if (!System.String.IsNullOrEmpty(data._location))
+			try
 			{
-				// If the TextLocation contains a '(', we assume it is a format: "/Path/To/File.cpp(10)" where 10 is a line number that can be parsed
-				int startLineIndex = data._location.IndexOf('(', System.StringComparison.Ordinal);
-				// If the TextLocation contains a '.', we assume it is a format: "/Path/To/Filename.*" where the filename can be retrieved
-				int startExtensionIndex = data._location.IndexOf('.', System.StringComparison.Ordinal);
-				if (startLineIndex > 0 && data._line == -1)
+				if (!System.String.IsNullOrEmpty(data._location))
 				{
-					data._line = Int32.Parse(data._location.Substring(startLineIndex + 1, data._location.Length - startLineIndex - 2));
+					// If the TextLocation contains a '.', we assume it is a format: "/Path/To/Filename.*" where the filename can be retrieved
+					int startExtensionIndex = data._location.IndexOf('.', System.StringComparison.Ordinal);
+					if (startExtensionIndex > 0 && System.String.IsNullOrEmpty(data._file))
+					{
+						data._file = data._location.Substring(0, startExtensionIndex);
+					}
 				}
-				if (startExtensionIndex > 0 && System.String.IsNullOrEmpty(data._file))
+				if (!System.String.IsNullOrEmpty(data._conflictLocation))
 				{
-					data._file = data._location.Substring(0, startExtensionIndex);
+					// If the TextLocation contains a '.', we assume it is a format: "/Path/To/Filename.*" where the filename can be retrieved
+					int startExtensionIndex = data._conflictLocation.IndexOf('.', System.StringComparison.Ordinal);
+					if (startExtensionIndex > 0 && System.String.IsNullOrEmpty(data._conflictFile))
+					{
+						data._conflictFile = data._conflictLocation.Substring(0, startExtensionIndex);
+					}
 				}
 			}
-			if (!System.String.IsNullOrEmpty(data._conflictLocation))
+			catch
 			{
-				// If the TextLocation contains a '(', we assume it is a format: "/Path/To/File.cpp(10)" where 10 is a line number that can be parsed
-				int startLineIndex = data._conflictLocation.IndexOf('(', System.StringComparison.Ordinal);
-				// If the TextLocation contains a '.', we assume it is a format: "/Path/To/Filename.*" where the filename can be retrieved
-				int startExtensionIndex = data._conflictLocation.IndexOf('.', System.StringComparison.Ordinal);
-				if (startLineIndex > 0 && data._conflictLine == -1)
-				{
-					data._conflictLine = Int32.Parse(data._conflictLocation.Substring(startLineIndex + 1, data._conflictLocation.Length - startLineIndex - 2));
-				}
-				if (startExtensionIndex > 0 && System.String.IsNullOrEmpty(data._conflictFile))
-				{
-					data._conflictFile = data._conflictLocation.Substring(0, startExtensionIndex);
-				}
+				data._file = "";
+				data._conflictFile = "";
 			}
 		}
 
