Horde: Updating schema documentation

#rnx

diff --git a/Engine/Source/Programs/Horde/Docs/Config/Schema/AclActions.md b/Engine/Source/Programs/Horde/Docs/Config/Schema/AclActions.md
--- a/Engine/Source/Programs/Horde/Docs/Config/Schema/AclActions.md
+++ b/Engine/Source/Programs/Horde/Docs/Config/Schema/AclActions.md
@@ -37,6 +37,32 @@
 | `DeleteAccount` |  |
 | `ViewAccount` |  |
 
+## Notices
+
+| Name | Description |
+| ---- | ----------- |
+| `CreateNotice` |  |
+| `UpdateNotice` |  |
+| `DeleteNotice` |  |
+
+## ServiceAccounts
+
+| Name | Description |
+| ---- | ----------- |
+| `CreateAccount` |  |
+| `UpdateAccount` |  |
+| `DeleteAccount` |  |
+| `ViewAccount` |  |
+
+## Accounts
+
+| Name | Description |
+| ---- | ----------- |
+| `CreateAccount` |  |
+| `UpdateAccount` |  |
+| `DeleteAccount` |  |
+| `ViewAccount` |  |
+
 ## Artifacts
 
 | Name | Description |
@@ -132,6 +158,17 @@
 | `DeleteAccount` |  |
 | `ViewAccount` |  |
 
+## Agents
+
+| Name | Description |
+| ---- | ----------- |
+| `CreateAgent` | Ability to create an agent. This may be done explicitly, or granted to agents to allow them to self-register. |
+| `CreateWorkstationAgent` | Ability to create a workstation agent. Granted to users with hardware to offer. |
+| `UpdateAgent` | Update an agent's name, pools, etc... |
+| `DeleteAgent` | Soft-delete an agent |
+| `ViewAgent` | View an agent |
+| `ListAgents` | List the available agents |
+
 ## Artifacts
 
 | Name | Description |
@@ -150,6 +187,13 @@
 | `UpdateBisectTask` |  |
 | `ViewBisectTask` |  |
 
+## Compute
+
+| Name | Description |
+| ---- | ----------- |
+| `AddComputeTasks` | User can add tasks to the compute cluster |
+| `GetComputeTasks` | User can get and list tasks from the compute cluster |
+
 ## Devices
 
 | Name | Description |
@@ -168,100 +212,23 @@
 | `RetryJobStep` |  |
 | `ViewJob` |  |
 
-## Notices
+## Leases
 
 | Name | Description |
 | ---- | ----------- |
-| `CreateNotice` |  |
-| `UpdateNotice` |  |
-| `DeleteNotice` |  |
+| `ViewLeases` | View all the leases that an agent has worked on |
+| `ViewLeaseTasks` | View the task data for a lease |
 
-## Notifications
+## Logs
 
 | Name | Description |
 | ---- | ----------- |
-| `CreateSubscription` |  |
-
-## Projects
-
-| Name | Description |
-| ---- | ----------- |
-| `CreateProject` |  |
-| `DeleteProject` |  |
-| `UpdateProject` |  |
-| `ViewProject` |  |
-
-## Replicators
-
-| Name | Description |
-| ---- | ----------- |
-| `UpdateReplicator` |  |
-| `ViewReplicator` |  |
-
-## ServiceAccounts
-
-| Name | Description |
-| ---- | ----------- |
-| `CreateAccount` |  |
-| `UpdateAccount` |  |
-| `DeleteAccount` |  |
-| `ViewAccount` |  |
-
-## Streams
-
-| Name | Description |
-| ---- | ----------- |
-| `CreateStream` |  |
-| `UpdateStream` |  |
-| `DeleteStream` |  |
-| `ViewStream` |  |
-| `ViewChanges` |  |
-| `ViewTemplate` |  |
-
-## Accounts
-
-| Name | Description |
-| ---- | ----------- |
-| `CreateAccount` |  |
-| `UpdateAccount` |  |
-| `DeleteAccount` |  |
-| `ViewAccount` |  |
-
-## Artifacts
-
-| Name | Description |
-| ---- | ----------- |
-| `ReadArtifact` |  |
-| `WriteArtifact` |  |
-| `DeleteArtifact` |  |
-| `UploadArtifact` |  |
-| `DownloadArtifact` |  |
-
-## Bisect
-
-| Name | Description |
-| ---- | ----------- |
-| `CreateBisectTask` |  |
-| `UpdateBisectTask` |  |
-| `ViewBisectTask` |  |
-
-## Devices
-
-| Name | Description |
-| ---- | ----------- |
-| `DeviceRead` |  |
-| `DeviceWrite` |  |
-
-## Jobs
-
-| Name | Description |
-| ---- | ----------- |
-| `CreateJob` |  |
-| `UpdateJob` |  |
-| `DeleteJob` |  |
-| `ExecuteJob` |  |
-| `RetryJobStep` |  |
-| `ViewJob` |  |
+| `CreateLog` | Ability to create a log. Implicitly granted to agents. |
+| `UpdateLog` | Ability to update log metadata |
+| `ViewLog` | Ability to view a log contents |
+| `WriteLogData` | Ability to write log data |
+| `CreateEvent` | Ability to create events |
+| `ViewEvent` | Ability to view events |
 
 ## Notices
 
@@ -277,6 +244,16 @@
 | ---- | ----------- |
 | `CreateSubscription` |  |
 
+## Pools
+
+| Name | Description |
+| ---- | ----------- |
+| `CreatePool` | Create a global pool of agents |
+| `UpdatePool` | Modify an agent pool |
+| `DeletePool` | Delete an agent pool |
+| `ViewPool` | Ability to view a pool |
+| `ListPools` | View all the available agent pools |
+
 ## Projects
 
 | Name | Description |
@@ -302,6 +279,21 @@
 | `DeleteAccount` |  |
 | `ViewAccount` |  |
 
+## Sessions
+
+| Name | Description |
+| ---- | ----------- |
+| `CreateSession` | Granted to agents to call CreateSession, which returns a bearer token identifying themselves valid to call UpdateSesssion via gRPC. |
+| `ViewSession` | Allows viewing information about an agent session |
+
+## Software
+
+| Name | Description |
+| ---- | ----------- |
+| `UploadSoftware` | Ability to upload new versions of the agent software |
+| `DownloadSoftware` | Ability to download the agent software |
+| `DeleteSoftware` | Ability to delete agent software |
+
 ## Streams
 
 | Name | Description |
@@ -326,12 +318,12 @@
 
 | Name | Description |
 | ---- | ----------- |
-| `CreateAgent` | Ability to create an agent. This may be done explicitly, or granted to agents to allow them to self-register. |
-| `CreateWorkstationAgent` | Ability to create a workstation agent. Granted to users with hardware to offer. |
-| `UpdateAgent` | Update an agent's name, pools, etc... |
-| `DeleteAgent` | Soft-delete an agent |
-| `ViewAgent` | View an agent |
-| `ListAgents` | List the available agents |
+| `CreateAgent` |  |
+| `CreateWorkstationAgent` |  |
+| `UpdateAgent` |  |
+| `DeleteAgent` |  |
+| `ViewAgent` |  |
+| `ListAgents` |  |
 
 ## Artifacts
 
@@ -355,8 +347,21 @@
 
 | Name | Description |
 | ---- | ----------- |
-| `AddComputeTasks` | User can add tasks to the compute cluster |
-| `GetComputeTasks` | User can get and list tasks from the compute cluster |
+| `AddComputeTasks` |  |
+| `GetComputeTasks` |  |
 
+## Ddc
+
+| Name | Description |
+| ---- | ----------- |
+| `DdcReadObject` | General read access to refs / blobs and so on |
+| `DdcWriteObject` | General write access to upload refs / blobs etc |
+| `DdcDeleteObject` | Access to delete blobs / refs etc |
+| `DdcDeleteBucket` | Access to delete a particular bucket |
+| `DdcDeleteNamespace` | Access to delete a whole namespace |
+| `DdcReadTransactionLog` | Access to read the transaction log |
+| `DdcWriteTransactionLog` | Access to write the transaction log |
+| `DdcAdminAction` | Access to perform administrative task |
+
 ## Devices
 
@@ -380,19 +385,19 @@
 
 | Name | Description |
 | ---- | ----------- |
-| `ViewLeases` | View all the leases that an agent has worked on |
-| `ViewLeaseTasks` | View the task data for a lease |
+| `ViewLeases` |  |
+| `ViewLeaseTasks` |  |
 
 ## Logs
 
 | Name | Description |
 | ---- | ----------- |
-| `CreateLog` | Ability to create a log. Implicitly granted to agents. |
-| `UpdateLog` | Ability to update log metadata |
-| `ViewLog` | Ability to view a log contents |
-| `WriteLogData` | Ability to write log data |
-| `CreateEvent` | Ability to create events |
-| `ViewEvent` | Ability to view events |
+| `CreateLog` |  |
+| `UpdateLog` |  |
+| `ViewLog` |  |
+| `WriteLogData` |  |
+| `CreateEvent` |  |
+| `ViewEvent` |  |
 
 ## Notices
 
@@ -412,11 +417,11 @@
 
 | Name | Description |
 | ---- | ----------- |
-| `CreatePool` | Create a global pool of agents |
-| `UpdatePool` | Modify an agent pool |
-| `DeletePool` | Delete an agent pool |
-| `ViewPool` | Ability to view a pool |
-| `ListPools` | View all the available agent pools |
+| `CreatePool` |  |
+| `UpdatePool` |  |
+| `DeletePool` |  |
+| `ViewPool` |  |
+| `ListPools` |  |
 
 ## Projects
 
@@ -447,16 +452,16 @@
 
 | Name | Description |
 | ---- | ----------- |
-| `CreateSession` | Granted to agents to call CreateSession, which returns a bearer token identifying themselves valid to call UpdateSesssion via gRPC. |
-| `ViewSession` | Allows viewing information about an agent session |
+| `CreateSession` |  |
+| `ViewSession` |  |
 
 ## Software
 
 | Name | Description |
 | ---- | ----------- |
-| `UploadSoftware` | Ability to upload new versions of the agent software |
-| `DownloadSoftware` | Ability to download the agent software |
-| `DeleteSoftware` | Ability to delete agent software |
+| `UploadSoftware` |  |
+| `DownloadSoftware` |  |
+| `DeleteSoftware` |  |
 
 ## Streams
 
@@ -514,6 +519,19 @@
 | `AddComputeTasks` |  |
 | `GetComputeTasks` |  |
 
+## Ddc
+
+| Name | Description |
+| ---- | ----------- |
+| `DdcReadObject` |  |
+| `DdcWriteObject` |  |
+| `DdcDeleteObject` |  |
+| `DdcDeleteBucket` |  |
+| `DdcDeleteNamespace` |  |
+| `DdcReadTransactionLog` |  |
+| `DdcWriteTransactionLog` |  |
+| `DdcAdminAction` |  |
+
 ## Devices
 
 | Name | Description |
@@ -590,12 +608,6 @@
 | `UpdateReplicator` |  |
 | `ViewReplicator` |  |
 
-## Secrets
-
-| Name | Description |
-| ---- | ----------- |
-| `ViewSecret` | View a credential |
-
 ## ServiceAccounts
 
 | Name | Description |
@@ -676,6 +688,19 @@
 | `AddComputeTasks` |  |
 | `GetComputeTasks` |  |
 
+## Ddc
+
+| Name | Description |
+| ---- | ----------- |
+| `DdcReadObject` |  |
+| `DdcWriteObject` |  |
+| `DdcDeleteObject` |  |
+| `DdcDeleteBucket` |  |
+| `DdcDeleteNamespace` |  |
+| `DdcReadTransactionLog` |  |
+| `DdcWriteTransactionLog` |  |
+| `DdcAdminAction` |  |
+
 ## Devices
 
 | Name | Description |
@@ -756,7 +781,7 @@
 
 | Name | Description |
 | ---- | ----------- |
-| `ViewSecret` |  |
+| `ViewSecret` | View a credential |
 
 ## ServiceAccounts
 
@@ -782,19 +807,6 @@
 | `DownloadSoftware` |  |
 | `DeleteSoftware` |  |
 
-## Storage
-
-| Name | Description |
-| ---- | ----------- |
-| `ReadBlobs` | Ability to read blobs from the storage service |
-| `WriteBlobs` | Ability to write blobs to the storage service |
-| `ReadAliases` | Ability to read aliases from the storage service |
-| `WriteAliases` | Ability to write aliases to the storage service |
-| `DeleteAliases` | Ability to write aliases to the storage service |
-| `ReadRefs` | Ability to read refs from the storage service |
-| `WriteRefs` | Ability to write refs to the storage service |
-| `DeleteRefs` | Ability to delete refs |
-
 ## Streams
 
 | Name | Description |
@@ -855,14 +867,14 @@
 
 | Name | Description |
 | ---- | ----------- |
-| `DdcReadObject` | General read access to refs / blobs and so on |
-| `DdcWriteObject` | General write access to upload refs / blobs etc |
-| `DdcDeleteObject` | Access to delete blobs / refs etc |
-| `DdcDeleteBucket` | Access to delete a particular bucket |
-| `DdcDeleteNamespace` | Access to delete a whole namespace |
-| `DdcReadTransactionLog` | Access to read the transaction log |
-| `DdcWriteTransactionLog` | Access to write the transaction log |
-| `DdcAdminAction` | Access to perform administrative task |
+| `DdcReadObject` |  |
+| `DdcWriteObject` |  |
+| `DdcDeleteObject` |  |
+| `DdcDeleteBucket` |  |
+| `DdcDeleteNamespace` |  |
+| `DdcReadTransactionLog` |  |
+| `DdcWriteTransactionLog` |  |
+| `DdcAdminAction` |  |
 
 ## Devices
 
@@ -974,14 +986,14 @@
 
 | Name | Description |
 | ---- | ----------- |
-| `ReadBlobs` |  |
-| `WriteBlobs` |  |
-| `ReadAliases` |  |
-| `WriteAliases` |  |
-| `DeleteAliases` |  |
-| `ReadRefs` |  |
-| `WriteRefs` |  |
-| `DeleteRefs` |  |
+| `ReadBlobs` | Ability to read blobs from the storage service |
+| `WriteBlobs` | Ability to write blobs to the storage service |
+| `ReadAliases` | Ability to read aliases from the storage service |
+| `WriteAliases` | Ability to write aliases to the storage service |
+| `DeleteAliases` | Ability to write aliases to the storage service |
+| `ReadRefs` | Ability to read refs from the storage service |
+| `WriteRefs` | Ability to write refs to the storage service |
+| `DeleteRefs` | Ability to delete refs |
 
 ## Streams
 
diff --git a/Engine/Source/Programs/Horde/Docs/Config/Schema/Globals.md b/Engine/Source/Programs/Horde/Docs/Config/Schema/Globals.md
--- a/Engine/Source/Programs/Horde/Docs/Config/Schema/Globals.md
+++ b/Engine/Source/Programs/Horde/Docs/Config/Schema/Globals.md
@@ -58,16 +58,24 @@ Name | Description
 
 Name | Description
 ---- | -----------
-`build` | [BuildConfig](#buildconfig)<br>Configuration for the build plugin
 `analytics` | [AnalyticsConfig](#analyticsconfig)<br>Configuration for the analytics plugin
-`experimental` | [ExperimentalConfig](#experimentalconfig)<br>Configuration for the experimental plugin
+`build` | [BuildConfig](#buildconfig)<br>Configuration for the build plugin
 `compute` | [ComputeConfig](#computeconfig)<br>Configuration for the compute plugin
+`ddc` | [EmptyPluginConfig](#emptypluginconfig)<br>Configuration for the ddc plugin
+`experimental` | [ExperimentalConfig](#experimentalconfig)<br>Configuration for the experimental plugin
 `secrets` | [SecretsConfig](#secretsconfig)<br>Configuration for the secrets plugin
 `storage` | [StorageConfig](#storageconfig)<br>Configuration for the storage plugin
-`ddc` | [EmptyPluginConfig](#emptypluginconfig)<br>Configuration for the ddc plugin
 `symbols` | [SymbolsConfig](#symbolsconfig)<br>Configuration for the symbols plugin
 `tools` | [ToolsConfig](#toolsconfig)<br>Configuration for the tools plugin
 
+## AnalyticsConfig
+
+Config settings for analytics
+
+Name | Description
+---- | -----------
+`stores` | [TelemetryStoreConfig](Telemetry.md)`[]`<br>Metrics to aggregate on the Horde server
+
 ## BuildConfig
 
 Configuration for the build plugin
@@ -226,22 +234,6 @@ Name | Description
 `excludeActions` | `string[]`<br>Actions to exclude from the inherited actions
 `extends` | `string[]`<br>Other profiles to extend from
 
-## AnalyticsConfig
-
-Config settings for analytics
-
-Name | Description
----- | -----------
-`stores` | [TelemetryStoreConfig](Telemetry.md)`[]`<br>Metrics to aggregate on the Horde server
-
-## ExperimentalConfig
-
-Configuration for the experimental plugin
-
-Name | Description
----- | -----------
-`notifications` | [NotificationConfig](Notification.md)`[]`<br>Notifications to be sent on the Horde server
-
 ## ComputeConfig
 
 Configuration for the compute system
@@ -426,6 +418,19 @@ Name | Description
 `description` | `string`<br>Human-readable description
 `computeId` | `string`<br>Compute ID for this network (used when allocating compute resources)
 
+## EmptyPluginConfig
+
+Empty implementation of
+
+
+## ExperimentalConfig
+
+Configuration for the experimental plugin
+
+Name | Description
+---- | -----------
+`notifications` | [NotificationConfig](Notification.md)`[]`<br>Notifications to be sent on the Horde server
+
 ## SecretsConfig
 
 Configuration for the secrets system
@@ -571,11 +576,6 @@ Name | Description
 `enableAliases` | `boolean`<br>Support querying exports by their aliases
 `acl` | [AclConfig](#aclconfig)<br>Access list for this namespace
 
-## EmptyPluginConfig
-
-Empty implementation of
-
-
 ## SymbolsConfig
 
 Configuration for the tools system
diff --git a/Engine/Source/Programs/Horde/Docs/Deployment/ServerSettings.md b/Engine/Source/Programs/Horde/Docs/Deployment/ServerSettings.md
--- a/Engine/Source/Programs/Horde/Docs/Deployment/ServerSettings.md
+++ b/Engine/Source/Programs/Horde/Docs/Deployment/ServerSettings.md
@@ -127,16 +127,53 @@
 
 Name | Description
 ---- | -----------
-`build` | [BuildServerConfig](#buildserverconfig)<br>Configuration for the build plugin
 `analytics` | [AnalyticsServerConfig](#analyticsserverconfig)<br>Configuration for the analytics plugin
-`experimental` | [ExperimentalServerConfig](#experimentalserverconfig)<br>Configuration for the experimental plugin
+`build` | [BuildServerConfig](#buildserverconfig)<br>Configuration for the build plugin
 `compute` | [ComputeServerConfig](#computeserverconfig)<br>Configuration for the compute plugin
+`ddc` | [PluginServerConfig](#pluginserverconfig)<br>Configuration for the ddc plugin
+`experimental` | [ExperimentalServerConfig](#experimentalserverconfig)<br>Configuration for the experimental plugin
 `secrets` | [PluginServerConfig](#pluginserverconfig)<br>Configuration for the secrets plugin
 `storage` | [StorageServerConfig](#storageserverconfig)<br>Configuration for the storage plugin
-`ddc` | [PluginServerConfig](#pluginserverconfig)<br>Configuration for the ddc plugin
 `symbols` | [PluginServerConfig](#pluginserverconfig)<br>Configuration for the symbols plugin
 `tools` | [ToolsServerConfig](#toolsserverconfig)<br>Configuration for the tools plugin
 
+## AnalyticsServerConfig
+
+Server configuration for the analytics system
+
+Name | Description
+---- | -----------
+`sinks` | [TelemetrySinkConfig](#telemetrysinkconfig)<br>Settings for the various telemetry sinks
+`enabled` | `boolean`<br>Whether the plugin should be enabled or not
+
+## TelemetrySinkConfig
+
+Telemetry sinks
+
+Name | Description
+---- | -----------
+`epic` | [EpicTelemetryConfig](#epictelemetryconfig)<br>Settings for the Epic telemetry sink
+`mongo` | [MongoTelemetryConfig](#mongotelemetryconfig)<br>Settings for the MongoDB telemetry sink
+
+## EpicTelemetryConfig
+
+Configuration for the telemetry sink
+
+Name | Description
+---- | -----------
+`url` | `string`<br>Base URL for the telemetry server
+`appId` | `string`<br>Application name to send in the event messages
+`enabled` | `boolean`<br>Whether to enable this sink
+
+## MongoTelemetryConfig
+
+Configuration for the telemetry sink
+
+Name | Description
+---- | -----------
+`retainDays` | `number`<br>Number of days worth of telmetry events to keep
+`enabled` | `boolean`<br>Whether to enable this sink
+
 ## BuildServerConfig
 
 Static configuration for the build plugin
@@ -279,51 +316,6 @@
 `maxChildCount` | `integer`<br>Maximum number of children in each node
 `sliceThreshold` | `integer`<br>Threshold hash value for splitting interior nodes
 
-## AnalyticsServerConfig
-
-Server configuration for the analytics system
-
-Name | Description
----- | -----------
-`sinks` | [TelemetrySinkConfig](#telemetrysinkconfig)<br>Settings for the various telemetry sinks
-`enabled` | `boolean`<br>Whether the plugin should be enabled or not
-
-## TelemetrySinkConfig
-
-Telemetry sinks
-
-Name | Description
----- | -----------
-`epic` | [EpicTelemetryConfig](#epictelemetryconfig)<br>Settings for the Epic telemetry sink
-`mongo` | [MongoTelemetryConfig](#mongotelemetryconfig)<br>Settings for the MongoDB telemetry sink
-
-## EpicTelemetryConfig
-
-Configuration for the telemetry sink
-
-Name | Description
----- | -----------
-`url` | `string`<br>Base URL for the telemetry server
-`appId` | `string`<br>Application name to send in the event messages
-`enabled` | `boolean`<br>Whether to enable this sink
-
-## MongoTelemetryConfig
-
-Configuration for the telemetry sink
-
-Name | Description
----- | -----------
-`retainDays` | `number`<br>Number of days worth of telmetry events to keep
-`enabled` | `boolean`<br>Whether to enable this sink
-
-## ExperimentalServerConfig
-
-Server configuration for the experimental plugin
-
-Name | Description
----- | -----------
-`enabled` | `boolean`<br>Whether the plugin should be enabled or not
-
 ## ComputeServerConfig
 
 Static configuration for the compute plugin
@@ -377,6 +369,14 @@
 ---- | -----------
 `enabled` | `boolean`<br>Whether the plugin should be enabled or not
 
+## ExperimentalServerConfig
+
+Server configuration for the experimental plugin
+
+Name | Description
+---- | -----------
+`enabled` | `boolean`<br>Whether the plugin should be enabled or not
+
 ## StorageServerConfig
 
 Static settings for the storage system
