Horde :  Fix ScheduledDowntime GetNext to return next or active downtime

If the downtime is active return the start and finish times for the active schedule.
If the downtime is not active return the start and finish times for the next day or next week schedule.
Downtime Service Tick ignores next times that are in the past, which is for downtime that are non-repeating.

diff --git a/Engine/Source/Programs/Horde/HordeServer.Tests/Server/DowntimeServiceTests.cs b/Engine/Source/Programs/Horde/HordeServer.Tests/Server/DowntimeServiceTests.cs
--- a/Engine/Source/Programs/Horde/HordeServer.Tests/Server/DowntimeServiceTests.cs
+++ b/Engine/Source/Programs/Horde/HordeServer.Tests/Server/DowntimeServiceTests.cs
@@ -329,5 +329,130 @@ namespace HordeServer.Tests.Server
 			Assert.AreEqual(ScheduledDowntimeFrequency.Once, downtime.Frequency);
 			Assert.AreEqual(downtime.TimeZone, "Eastern Standard Time");
 		}
+
+		private static ScheduledDowntime GetWeeklyScheduledDowntime()
+		{
+			DateTimeOffset startTime = DateTimeOffset.Parse("Wednesday 1 Jan 2025 1AM");
+			DateTimeOffset finishTime = startTime + TimeSpan.FromMinutes(5);
+			ScheduledDowntime downtime = new()
+			{
+				StartTime = startTime,
+				FinishTime = finishTime,
+				Frequency = ScheduledDowntimeFrequency.Weekly
+			};
+			return downtime;
+		}
+
+		[TestMethod]
+		[DataRow("18 Dec 2024 12:00", " 1 Jan 2025 1:00", DisplayName = "Two weeks before first schedule")]
+		[DataRow("31 Dec 2024 12:00", " 1 Jan 2025 1:00", DisplayName = "One day before first schedule")]
+		[DataRow(" 1 Jan 2025  1:05", " 8 Jan 2025 1:00", DisplayName = "End of first schedule")]
+		[DataRow(" 2 Jan 2025 15:00", " 8 Jan 2025 1:00", DisplayName = "After first schedule")]
+		[DataRow("11 Jan 2025 15:00", "15 Jan 2025 1:00", DisplayName = "After schedule")]
+		[DataRow("11 Jan 2026 15:00", "14 Jan 2026 1:00", DisplayName = "One year later")]
+		public void GetNextWeekly(string nowStr, string expectedStr)
+		{
+			DateTimeOffset now = DateTimeOffset.Parse(nowStr);
+			DateTimeOffset expected = DateTimeOffset.Parse(expectedStr);
+
+			ScheduledDowntime downtime = GetWeeklyScheduledDowntime();
+			(DateTimeOffset startTime, DateTimeOffset finishTime) = downtime.GetNext(now);
+
+			Assert.AreEqual(expected, startTime);
+			Assert.AreEqual(expected + (finishTime - startTime), finishTime);
+			Assert.IsFalse(downtime.IsActive(now));
+		}
+
+		[TestMethod]
+		[DataRow(" 1 Jan 2025 1:00", " 1 Jan 2025 1:00", DisplayName = "Start of schedule")]
+		[DataRow(" 1 Jan 2025 1:01", " 1 Jan 2025 1:00", DisplayName = "During schedule")]
+		[DataRow("15 Jan 2025 1:00", "15 Jan 2025 1:00", DisplayName = "During schedule two weeks later")]
+		[DataRow("14 Jan 2026 1:00", "14 Jan 2026 1:00", DisplayName = "During schedule one year later")]
+		public void GetCurrentWeekly(string nowStr, string expectedStr)
+		{
+			DateTimeOffset now = DateTimeOffset.Parse(nowStr);
+			DateTimeOffset expected = DateTimeOffset.Parse(expectedStr);
+
+			ScheduledDowntime downtime = GetWeeklyScheduledDowntime();
+			(DateTimeOffset startTime, DateTimeOffset finishTime) = downtime.GetNext(now);
+
+			Assert.AreEqual(expected, startTime);
+			Assert.AreEqual(expected + (finishTime - startTime), finishTime);
+			Assert.IsTrue(downtime.IsActive(now));
+		}
+
+		private static ScheduledDowntime GetDailyScheduledDowntime()
+		{
+			DateTimeOffset startTime = DateTimeOffset.Parse("Wednesday 1 Jan 2025 1AM");
+			DateTimeOffset finishTime = startTime + TimeSpan.FromMinutes(5);
+			ScheduledDowntime downtime = new()
+			{
+				StartTime = startTime,
+				FinishTime = finishTime,
+				Frequency = ScheduledDowntimeFrequency.Daily
+			};
+			return downtime;
+		}
+
+		[TestMethod]
+		[DataRow("30 Dec 2024  0:00", " 1 Jan 2025 1:00", DisplayName = "Two days before schedule")]
+		[DataRow("31 Dec 2024  0:00", " 1 Jan 2025 1:00", DisplayName = "One day before schedule")]
+		[DataRow(" 1 Jan 2025  1:06", " 2 Jan 2025 1:00", DisplayName = "End of schedule")]
+		[DataRow("10 Jan 2025  1:06", "11 Jan 2025 1:00", DisplayName = "End of schedule")]
+		[DataRow(" 2 Jan 2025 15:00", " 3 Jan 2025 1:00", DisplayName = "After schedule")]
+		public void GetNextDaily(string nowStr, string expectedStr)
+		{
+			DateTimeOffset now = DateTimeOffset.Parse(nowStr);
+			DateTimeOffset expected = DateTimeOffset.Parse(expectedStr);
+
+			ScheduledDowntime downtime = GetDailyScheduledDowntime();
+			(DateTimeOffset startTime, DateTimeOffset finishTime) = downtime.GetNext(now);
+
+			Assert.AreEqual(expected, startTime);
+			Assert.AreEqual(expected + (finishTime - startTime), finishTime);
+			Assert.IsFalse(downtime.IsActive(now));
+		}
+
+		[TestMethod]
+		[DataRow(" 1 Jan 2025  1:00", " 1 Jan 2025 1:00", DisplayName = "Start of schedule")]
+		[DataRow(" 1 Jan 2025  1:03", " 1 Jan 2025 1:00", DisplayName = "During schedule")]
+		[DataRow(" 1 Jan 2026  1:03", " 1 Jan 2026 1:00", DisplayName = "During schedule one year later")]
+		public void GetCurrentDaily(string nowStr, string expectedStr)
+		{
+			DateTimeOffset now = DateTimeOffset.Parse(nowStr);
+			DateTimeOffset expected = DateTimeOffset.Parse(expectedStr);
+
+			ScheduledDowntime downtime = GetDailyScheduledDowntime();
+			(DateTimeOffset startTime, DateTimeOffset finishTime) = downtime.GetNext(now);
+
+			Assert.AreEqual(expected, startTime);
+			Assert.AreEqual(expected + (finishTime - startTime), finishTime);
+			Assert.IsTrue(downtime.IsActive(now));
+		}
+
+		[TestMethod]
+		public void GetSuccessiveNext()
+		{
+			ScheduledDowntime downtime = new()
+			{
+				StartTime = DateTimeOffset.Parse("2024-11-07T10:10:00"),
+				FinishTime = DateTimeOffset.Parse("2024-11-07T14:30:00"),
+				Frequency = ScheduledDowntimeFrequency.Daily
+			};
+
+			DateTimeOffset now = DateTimeOffset.Parse("2021-11-07T10:40:00");
+			(DateTimeOffset startTime0, DateTimeOffset finishTime0) = downtime.GetNext(now);
+			(DateTimeOffset startTime1, DateTimeOffset finishTime1) = downtime.GetNext(now + TimeSpan.FromMinutes(1));
+			(DateTimeOffset startTime2, DateTimeOffset finishTime2) = downtime.GetNext(now + TimeSpan.FromMinutes(2));
+
+			Assert.AreEqual(DateTimeOffset.Parse("7 Nov 2024 10:10"), startTime0);
+			Assert.AreEqual(DateTimeOffset.Parse("7 Nov 2024 14:30"), finishTime0);
+
+			Assert.AreEqual(DateTimeOffset.Parse("7 Nov 2024 10:10"), startTime1);
+			Assert.AreEqual(DateTimeOffset.Parse("7 Nov 2024 14:30"), finishTime1);
+
+			Assert.AreEqual(DateTimeOffset.Parse("7 Nov 2024 10:10"), startTime2);
+			Assert.AreEqual(DateTimeOffset.Parse("7 Nov 2024 14:30"), finishTime2);
+		}
 	}
 }
diff --git a/Engine/Source/Programs/Horde/HordeServer/Server/DowntimeService.cs b/Engine/Source/Programs/Horde/HordeServer/Server/DowntimeService.cs
--- a/Engine/Source/Programs/Horde/HordeServer/Server/DowntimeService.cs
+++ b/Engine/Source/Programs/Horde/HordeServer/Server/DowntimeService.cs
@@ -77,7 +77,7 @@ namespace HordeServer.Server
 			foreach (ScheduledDowntime schedule in globalConfig.Downtime)
 			{
 				DateTimeOffset start = schedule.GetNext(now).StartTime;
-				if (next == null || start < next)
+				if (start >= now && (next == null || start < next))
 				{
 					next = start;
 				}
diff --git a/Engine/Source/Programs/Horde/HordeServer/Server/GlobalConfig.cs b/Engine/Source/Programs/Horde/HordeServer/Server/GlobalConfig.cs
--- a/Engine/Source/Programs/Horde/HordeServer/Server/GlobalConfig.cs
+++ b/Engine/Source/Programs/Horde/HordeServer/Server/GlobalConfig.cs
@@ -338,30 +338,36 @@ namespace HordeServer.Server
 		}
 
 		/// <summary>
-		/// Gets the next scheduled downtime
+		/// Gets the next scheduled downtime or the current one if it is active
 		/// </summary>
 		/// <param name="now">The current time</param>
 		/// <returns>Start and finish time</returns>
 		public (DateTimeOffset StartTime, DateTimeOffset FinishTime) GetNext(DateTimeOffset now)
 		{
 			TimeSpan offset = TimeSpan.Zero;
-			if (Frequency == ScheduledDowntimeFrequency.Daily)
+			if (Frequency is ScheduledDowntimeFrequency.Daily or ScheduledDowntimeFrequency.Weekly)
 			{
 				double days = (now - StartTime).TotalDays;
-				if (days >= 1.0)
+				if (IsActive(now))
 				{
-					days -= days % 1.0;
+					// Clamp active schedule to "now"
+					offset = TimeSpan.FromDays(Math.Floor(days));
 				}
-				offset = TimeSpan.FromDays(days);
-			}
-			else if (Frequency == ScheduledDowntimeFrequency.Weekly)
-			{
-				double days = (now - StartTime).TotalDays;
-				if (days >= 7.0)
+				else
 				{
-					days -= days % 7.0;
+					// Round up to the next schedule window after "now"
+					if (days > 0.0)
+					{
+						if (Frequency == ScheduledDowntimeFrequency.Daily)
+						{
+							offset = TimeSpan.FromDays(Math.Ceiling(days));
+						}
+						else if (Frequency == ScheduledDowntimeFrequency.Weekly)
+						{
+							offset = TimeSpan.FromDays(Math.Ceiling(days / 7.0) * 7.0);
+						}
+					}
 				}
-				offset = TimeSpan.FromDays(days);
 			}
 			return (StartTime + offset, FinishTime + offset);
 		}
