Horde: Add support for allocating a UBA cache server session for compute resource requests

diff --git a/Engine/Source/Programs/Horde/Plugins/Build/HordeServer.Build.Tests/Compute/ComputeServiceTest.cs b/Engine/Source/Programs/Horde/Plugins/Build/HordeServer.Build.Tests/Compute/ComputeServiceTest.cs
--- a/Engine/Source/Programs/Horde/Plugins/Build/HordeServer.Build.Tests/Compute/ComputeServiceTest.cs
+++ b/Engine/Source/Programs/Horde/Plugins/Build/HordeServer.Build.Tests/Compute/ComputeServiceTest.cs
@@ -1,7 +1,9 @@
 // Copyright Epic Games, Inc. All Rights Reserved.
 
+using System.Diagnostics.CodeAnalysis;
 using System.Diagnostics.Metrics;
 using System.Net;
+using System.Text.RegularExpressions;
 using EpicGames.Horde.Agents;
 using EpicGames.Horde.Agents.Pools;
 using EpicGames.Horde.Compute;
@@ -19,8 +21,18 @@ namespace HordeServer.Tests.Compute
 	[TestClass]
 	public class ComputeServiceTest : BuildTestSetup
 	{
+		private class StubMessageHandler(HttpStatusCode statusCode, string content) : HttpMessageHandler
+		{
+			private string Content { get; set; } = content;
+			protected override Task<HttpResponseMessage> SendAsync(HttpRequestMessage request, CancellationToken cancellationToken)
+			{
+				return Task.FromResult(new HttpResponseMessage(statusCode) { Content = new StringContent(Content) });
+			}
+		}
+		
 		private readonly ClusterId _cluster1 = new("cluster1");
 
+		#region Resource needs
 		[TestMethod]
 		public async Task ResourceNeedsMetricAsync()
 		{
@@ -173,6 +185,7 @@ namespace HordeServer.Tests.Compute
 			Assert.AreEqual(3, result["poolA"]);
 			Assert.AreEqual(2, result["poolB"]);
 		}
+		#endregion Resource needs
 
 		[TestMethod]
 		public async Task PoolNameTemplatingAsync()
@@ -312,6 +325,7 @@ namespace HordeServer.Tests.Compute
 			}
 		}
 
+		#region Cluster
 		[TestMethod]
 		[DataRow("11.0.0.1", "11.0.0.1", null, false)]
 		[DataRow("11.0.0.1", "11.0.0.1", "200.20.20.20", false)]
@@ -348,7 +362,9 @@ namespace HordeServer.Tests.Compute
 			await globalConfigCatchAll.PostLoadAsync(configOptions);
 			Assert.AreEqual(new ClusterId("catchAll"), ComputeService.FindBestComputeClusterId(globalConfigCatchAll, IPAddress.Parse("123.123.123.123")));
 		}
-
+		#endregion Cluster
+		
+		#region Connection
 		[TestMethod]
 		public async Task Connection_Direct_IpConnection_Async()
 		{
@@ -406,13 +422,42 @@ namespace HordeServer.Tests.Compute
 			Assert.AreEqual(new ComputeResourcePort(18656, 12000), cr.Ports["myPort"]);
 			Assert.AreEqual(new ComputeResourcePort(23151, 13000), cr.Ports["myOtherPort"]);
 		}
-
-		private async Task<ComputeService> CreateComputeServiceAsync(string? tunnelAddress, string pool, ComputeClusterConfig ccc)
+		
+		[TestMethod]
+		public async Task Connection_UbaCache_CacheServerDefined_Async()
 		{
+			ComputeResource? cr = await AllocateWithAgentAsync(useUbaCache: true, ubaCacheEndpoint: "1.2.3.4:3333");
+			Assert.AreEqual("1.2.3.4:3333", cr!.Uba!.CacheEndpoint);
+			Assert.IsTrue(Regex.IsMatch(cr.Uba!.CacheSessionKey, "^[0-9a-fA-F]{32}$"));
+		}
+		
+		[TestMethod]
+		[Ignore] // Uncomment to manually test allocation with a local UBA cache server
+		public async Task Connection_UbaCache_RealServer_Async()
+		{
+			const string Endpoint = "localhost:7010";
+			ComputeResource? cr = await AllocateWithAgentAsync(useUbaCache: true, ubaCacheEndpoint: Endpoint);
+			Assert.AreEqual(Endpoint, cr!.Uba!.CacheEndpoint);
+			Assert.IsTrue(Regex.IsMatch(cr.Uba!.CacheSessionKey, "^[0-9a-fA-F]{32}$"));
+		}
+		
+		[TestMethod]
+		public async Task Connection_UbaCache_CacheServerNotDefined_Async()
+		{
+			ComputeServiceException cse = await Assert.ThrowsExceptionAsync<ComputeServiceException>(() => AllocateWithAgentAsync(useUbaCache: true, ubaCacheEndpoint: null));
+			Assert.IsTrue(cse.Message.StartsWith("No UBA cache server defined for compute cluster", StringComparison.Ordinal));
+		}
+		
+		#endregion Connection
+		
+		[SuppressMessage("Reliability", "CA2000:Dispose objects before losing scope", Justification = "Ignore to simplify test")]
+		private async Task<ComputeService> CreateComputeServiceWithAgentAsync(string? tunnelAddress, string pool, ComputeClusterConfig ccc, HttpClient? httpClient = null)
+		{
+			StubMessageHandler ubaCacheOk = new(HttpStatusCode.OK, "");
 			IOptionsMonitor<ComputeConfig> computeConfig = ServiceProvider.GetRequiredService<IOptionsMonitor<ComputeConfig>>();
 			ComputeServerConfig ss = new() { ComputeTunnelAddress = tunnelAddress };
 			ComputeService cs = new(AgentCollection, ServiceProvider.GetRequiredService<IAgentScheduler>(), LogCollection, AgentService, AgentRelayService, GetRedisServiceSingleton(),
-				new TestOptionsMonitor<ComputeServerConfig>(ss), computeConfig, Clock, Tracer, Meter,
+				httpClient ?? new HttpClient(ubaCacheOk), new TestOptionsMonitor<ComputeServerConfig>(ss), computeConfig, Clock, Tracer, Meter,
 				NullLogger<ComputeService>.Instance);
 			await CreateAgentAsync(new PoolId(pool), properties: ["ComputeIp=11.0.0.1", "ComputePort=5000"]);
 			computeConfig.CurrentValue.Clusters = [ccc];
@@ -427,15 +472,18 @@ namespace HordeServer.Tests.Compute
 			bool usePublicIp = false,
 			string[]? relayIps = null,
 			string? tunnelAddress = null,
+			bool useUbaCache = false,
+			string? ubaCacheEndpoint = null,
 			ComputeProtocol protocol = ComputeProtocol.Latest,
 			string pool = "myPool")
 		{
-			await using ComputeService cs = await CreateComputeServiceAsync(tunnelAddress, pool, new ComputeClusterConfig { Id = _cluster1 });
+			await using ComputeService cs = await CreateComputeServiceWithAgentAsync(tunnelAddress, pool, new ComputeClusterConfig { Id = _cluster1, UbaCacheEndpoint = ubaCacheEndpoint });
 			AllocateResourceParams arp = new(_cluster1, protocol, new Requirements() { Pool = pool })
 			{
 				ConnectionMode = connectionMode,
 				Ports = ports ?? new Dictionary<string, int>(),
-				UsePublicIp = usePublicIp
+				UsePublicIp = usePublicIp,
+				UseUbaCache = useUbaCache
 			};
 			IPAddress[] defaultRelayIps = { IPAddress.Parse("192.168.1.1") };
 			await AgentRelayService.UpdateAgentHeartbeatAsync(_cluster1, "myrelay", relayIps?.Select(IPAddress.Parse) ?? defaultRelayIps);
diff --git a/Engine/Source/Programs/Horde/Plugins/Compute/HordeServer.Compute/Compute/ComputeController.cs b/Engine/Source/Programs/Horde/Plugins/Compute/HordeServer.Compute/Compute/ComputeController.cs
--- a/Engine/Source/Programs/Horde/Plugins/Compute/HordeServer.Compute/Compute/ComputeController.cs
+++ b/Engine/Source/Programs/Horde/Plugins/Compute/HordeServer.Compute/Compute/ComputeController.cs
@@ -120,6 +120,7 @@ namespace HordeServer.Compute
 				UsePublicIp = request.Connection?.PreferPublicIp,
 				Encryption = ComputeService.ConvertEncryptionToProto(request.Connection?.Encryption),
 				InactivityTimeoutMs = request.Connection?.InactivityTimeoutMs,
+				UseUbaCache = request.UseUbaCache,
 				UserId = User.GetUserId(),
 			};
 
@@ -145,23 +146,24 @@ namespace HordeServer.Compute
 				responsePorts[name] = new ConnectionMetadataPort(crp.Port, crp.AgentPort);
 			}
 
-			AssignComputeResponse response = new AssignComputeResponse();
-			response.Ip = computeResource.Ip.ToString();
-			response.Port = computeResource.Ports[ConnectionMetadataPort.ComputeId].Port;
-			response.ConnectionMode = computeResource.ConnectionMode;
-			response.ConnectionAddress = computeResource.ConnectionAddress;
-			response.Ports = responsePorts;
-			response.Encryption = ComputeService.ConvertEncryptionFromProto(computeResource.Task.Encryption);
-			response.Nonce = StringUtils.FormatHexString(computeResource.Task.Nonce.Span);
-			response.Key = StringUtils.FormatHexString(computeResource.Task.Key.Span);
-			response.Certificate = StringUtils.FormatHexString(computeResource.Task.Certificate.Span);
-			response.ClusterId = clusterId;
-			response.AgentId = computeResource.AgentId;
-			response.AgentVersion = computeResource.AgentVersion;
-			response.LeaseId = computeResource.LeaseId;
-			response.Properties = computeResource.Properties;
-			response.Protocol = computeResource.Task.Protocol;
-
+			AssignComputeResponse response = new()
+			{
+				Ip = computeResource.Ip.ToString(), Port = computeResource.Ports[ConnectionMetadataPort.ComputeId].Port, ConnectionMode = computeResource.ConnectionMode,
+				ConnectionAddress = computeResource.ConnectionAddress,
+				Ports = responsePorts,
+				Encryption = ComputeService.ConvertEncryptionFromProto(computeResource.Task.Encryption),
+				Nonce = StringUtils.FormatHexString(computeResource.Task.Nonce.Span),
+				Key = StringUtils.FormatHexString(computeResource.Task.Key.Span),
+				Certificate = StringUtils.FormatHexString(computeResource.Task.Certificate.Span),
+				Uba = computeResource.Uba == null ? null : new UbaConfig(computeResource.Uba.CacheEndpoint, computeResource.Uba.CacheSessionKey),
+				ClusterId = clusterId,
+				AgentId = computeResource.AgentId,
+				AgentVersion = computeResource.AgentVersion,
+				LeaseId = computeResource.LeaseId,
+				Properties = computeResource.Properties,
+				Protocol = computeResource.Task.Protocol
+			};
+			
 			foreach (KeyValuePair<string, int> pair in computeResource.Task.Resources)
 			{
 				response.AssignedResources.Add(pair.Key, pair.Value);
diff --git a/Engine/Source/Programs/Horde/Plugins/Compute/HordeServer.Compute/Compute/ComputeService.cs b/Engine/Source/Programs/Horde/Plugins/Compute/HordeServer.Compute/Compute/ComputeService.cs
--- a/Engine/Source/Programs/Horde/Plugins/Compute/HordeServer.Compute/Compute/ComputeService.cs
+++ b/Engine/Source/Programs/Horde/Plugins/Compute/HordeServer.Compute/Compute/ComputeService.cs
@@ -76,7 +76,7 @@ namespace HordeServer.Compute
 	}
 
 	/// <summary>
-	/// Parameters for allocating a resource
+	/// Internal parameters for allocating a resource
 	/// </summary>
 	public class AllocateResourceParams
 	{
@@ -132,6 +132,9 @@ namespace HordeServer.Compute
 		
 		/// <inheritdoc cref="ConnectionMetadataRequest.InactivityTimeoutMs" />
 		public int? InactivityTimeoutMs { get; init; }
+		
+		/// <inheritdoc cref="AssignComputeResponse.Uba" />
+		public bool? UseUbaCache { get; init; }
 
 		/// <summary>
 		/// Constructor
@@ -185,6 +188,7 @@ namespace HordeServer.Compute
 		readonly AgentService _agentService;
 		readonly AgentRelayService _agentRelayService;
 		readonly IRedisService _redisService;
+		readonly HttpClient _httpClient;
 		readonly IOptionsMonitor<ComputeServerConfig> _staticComputeConfig;
 		readonly IOptionsMonitor<ComputeConfig> _computeConfig;
 		readonly IClock _clock;
@@ -208,6 +212,7 @@ namespace HordeServer.Compute
 			AgentService agentService,
 			AgentRelayService agentRelayService,
 			IRedisService redisService,
+			HttpClient httpClient,
 			IOptionsMonitor<ComputeServerConfig> staticComputeConfig,
 			IOptionsMonitor<ComputeConfig> computeConfig,
 			IClock clock,
@@ -221,6 +226,7 @@ namespace HordeServer.Compute
 			_agentService = agentService;
 			_agentRelayService = agentRelayService;
 			_redisService = redisService;
+			_httpClient = httpClient;
 			_staticComputeConfig = staticComputeConfig;
 			_computeConfig = computeConfig;
 			_clock = clock;
@@ -375,6 +381,7 @@ namespace HordeServer.Compute
 			span.SetAttribute("usePublicIp", arp.UsePublicIp.ToString());
 			span.SetAttribute("encryption", arp.Encryption.ToString());
 			span.SetAttribute("inactivityTimeout", arp.InactivityTimeoutMs);
+			span.SetAttribute("useUbaCache", arp.UseUbaCache ?? false);
 			span.SetAttribute("parentLeaseId", arp.ParentLeaseId?.ToString());
 			span.SetAttribute("req.pool", arp.Requirements.Pool);
 			span.SetAttribute("req.condition", arp.Requirements.Condition?.ToString());
@@ -408,6 +415,17 @@ namespace HordeServer.Compute
 				// With UnrealBuildAccelerator, it already makes full use of the local machine.
 				conditions.Add(Condition.Parse($"{KnownPropertyNames.ComputeIp} != '{arp.RequesterIp.ToString()}'"));
 			}
+			
+			Task<UbaComputeResource>? ubaCache = null;
+			if (arp.UseUbaCache ?? false)
+			{
+				if (clusterConfig.UbaCacheEndpoint == null)
+				{
+					throw new ComputeServiceException($"No UBA cache server defined for compute cluster {arp.ClusterId}") { ShowToUser = true };
+				}
+				// Allocate UBA cache session key in the background, await the task later
+				ubaCache = AllocateUbaCacheServerAsync(_httpClient, clusterConfig.UbaCacheEndpoint, cancellationToken);
+			}
 
 			try
 			{
@@ -460,6 +478,11 @@ namespace HordeServer.Compute
 								span.SetAttribute("allocatedAgentId", newAgent.Id.ToString());
 								span.SetAttribute("allocatedConnectionAddress", resource.ConnectionAddress);
 								span.SetAttribute("allocatedIp", resource.Ip.ToString());
+								
+								if (ubaCache != null)
+								{
+									resource.Uba = await ubaCache;
+								}
 
 								await LogRequestAsync(AllocationOutcome.Accepted, arp.RequestId, arp.Requirements, arp.ParentLeaseId, span, cancellationToken);
 								return resource;
@@ -854,6 +877,30 @@ namespace HordeServer.Compute
 			// TODO: Implement proper lookup of best relay IP to use
 			return relayIps[0];
 		}
+		
+		internal static async Task<UbaComputeResource> AllocateUbaCacheServerAsync(HttpClient httpClient, string ubaCacheEndpoint, CancellationToken cancellationToken)
+		{
+			try
+			{
+				using CancellationTokenSource timeoutCts = CancellationTokenSource.CreateLinkedTokenSource(cancellationToken);
+				timeoutCts.CancelAfter(TimeSpan.FromSeconds(3));
+				
+				// Generate a session encryption key that we activate in the UBA cache server
+				string key = RandomNumberGenerator.GetHexString(32, true);
+				HttpResponseMessage res = await httpClient.GetAsync(new Uri($"http://{ubaCacheEndpoint}/addcrypto?{key}"), timeoutCts.Token);
+				res.EnsureSuccessStatusCode();
+				
+				return new UbaComputeResource(ubaCacheEndpoint, key);
+			}
+			catch (TaskCanceledException ex)
+			{
+				throw new ComputeServiceException("UBA cache server allocation timed out or was cancelled", ex);
+			}
+			catch (Exception ex)
+			{
+				throw new ComputeServiceException("Unexpected error during UBA cache server allocation", ex);
+			}
+		}
 
 		static ComputeTask CreateComputeTask(Dictionary<string, int> assignedResources, LogId? logId, ComputeEncryption encryption,
 			byte[] certificateData, LeaseId? parentLeaseId, ComputeProtocol protocol, int? inactivityTimeoutMs)
diff --git a/Engine/Source/Programs/Horde/Plugins/Compute/HordeServer.Compute/Compute/ComputeTaskSource.cs b/Engine/Source/Programs/Horde/Plugins/Compute/HordeServer.Compute/Compute/ComputeTaskSource.cs
--- a/Engine/Source/Programs/Horde/Plugins/Compute/HordeServer.Compute/Compute/ComputeTaskSource.cs
+++ b/Engine/Source/Programs/Horde/Plugins/Compute/HordeServer.Compute/Compute/ComputeTaskSource.cs
@@ -15,6 +15,16 @@ using Microsoft.Extensions.Options;
 
 namespace HordeServer.Compute
 {
+	/// <inheritdoc cref="UbaConfig" />
+	public class UbaComputeResource(string cacheEndpoint, string cacheSessionKey)
+	{
+		/// <inheritdoc cref="UbaConfig.CacheEndpoint" />
+		public string CacheEndpoint { get; } = cacheEndpoint;
+		
+		/// <inheritdoc cref="UbaConfig.CacheSessionKey" />
+		public string CacheSessionKey { get; } = cacheSessionKey;
+	}
+	
 	/// <summary>
 	/// A TCP/IP port used by a compute resource, and how it is mapped externally -> internally
 	/// </summary>
@@ -78,7 +88,7 @@ namespace HordeServer.Compute
 	}
 
 	/// <summary>
-	/// Information about a compute 
+	/// Information about a compute resource
 	/// </summary>
 	public class ComputeResource
 	{
@@ -95,6 +105,9 @@ namespace HordeServer.Compute
 
 		/// <inheritdoc cref="AssignComputeResponse.Ports" />
 		public IReadOnlyDictionary<string, ComputeResourcePort> Ports { get; }
+		
+		/// <inheritdoc cref="UbaComputeResource" />
+		public UbaComputeResource? Uba { get; set;  }
 
 		/// <summary>
 		/// Information about the compute task
diff --git a/Engine/Source/Programs/Horde/Plugins/Compute/HordeServer.Compute/ComputeConfig.cs b/Engine/Source/Programs/Horde/Plugins/Compute/HordeServer.Compute/ComputeConfig.cs
--- a/Engine/Source/Programs/Horde/Plugins/Compute/HordeServer.Compute/ComputeConfig.cs
+++ b/Engine/Source/Programs/Horde/Plugins/Compute/HordeServer.Compute/ComputeConfig.cs
@@ -270,6 +270,13 @@ namespace HordeServer
 		/// Access control list
 		/// </summary>
 		public AclConfig Acl { get; set; } = new AclConfig();
+		
+		/// <summary>
+		/// Network endpoint of an Unreal Build Accelerator (UBA) cache server in IP:port format.
+		/// Must be accessible to both agents and initiators running UBA / Unreal Build Tool.
+		/// If null, use UBA cache functionality will be disabled for this compute cluster.
+		/// </summary>
+		public string? UbaCacheEndpoint { get; set; }
 
 		/// <summary>
 		/// Callback post loading this config file
diff --git a/Engine/Source/Programs/Shared/EpicGames.Horde/Compute/ComputeMessages.cs b/Engine/Source/Programs/Shared/EpicGames.Horde/Compute/ComputeMessages.cs
--- a/Engine/Source/Programs/Shared/EpicGames.Horde/Compute/ComputeMessages.cs
+++ b/Engine/Source/Programs/Shared/EpicGames.Horde/Compute/ComputeMessages.cs
@@ -42,6 +42,23 @@ namespace EpicGames.Horde.Compute
 			AgentPort = agentPort;
 		}
 	}
+	
+	/// <summary>
+	/// Configuration settings for Unreal Build Accelerator (UBA) for a compute task invocation
+	/// Contains cache server connectivity and other UBA-related settings.
+	/// </summary>
+	public class UbaConfig(string cacheEndpoint, string cacheSessionKey)
+	{
+		/// <summary>
+		/// Network endpoint (IP address and port) for connecting to the UBA cache server
+		/// </summary>
+		public string CacheEndpoint { get; } = cacheEndpoint;
+		
+		/// <summary>
+		/// AES encryption key used for secure client-server communication with the UBA cache server
+		/// </summary>
+		public string CacheSessionKey { get; } = cacheSessionKey;
+	}
 
 	/// <summary>
 	/// Request a machine to execute compute requests
@@ -68,6 +85,12 @@ namespace EpicGames.Horde.Compute
 		/// Details for making an agent connection
 		/// </summary>
 		public ConnectionMetadataRequest? Connection { get; set; }
+		
+		/// <summary>
+		/// Whether to allocate an Unreal Build Accelerator (UBA) cache server session for this compute request.
+		/// If requested, a UBA cache server must be defined for the compute cluster. Null is treated as false.
+		/// </summary>
+		public bool? UseUbaCache { get; init; }
 	}
 
 	/// <summary>
@@ -180,6 +203,9 @@ namespace EpicGames.Horde.Compute
 		/// </summary>
 		public string Certificate { get; set; } = String.Empty;
 		
+		/// <inheritdoc cref="UbaConfig" />
+		public UbaConfig? Uba { get; set; }
+		
 		/// <summary>
 		/// Which cluster this remote machine belongs to
 		/// </summary>
