Horde: Use Razor template for /account

diff --git a/Engine/Source/Programs/Horde/HordeServer/Server/AccountController.cs b/Engine/Source/Programs/Horde/HordeServer/Server/AccountController.cs
--- a/Engine/Source/Programs/Horde/HordeServer/Server/AccountController.cs
+++ b/Engine/Source/Programs/Horde/HordeServer/Server/AccountController.cs
@@ -2,11 +2,9 @@
 
 using System;
 using System.Collections.Generic;
-using System.Globalization;
 using System.Linq;
 using System.Net;
 using System.Security.Claims;
-using System.Text;
 using System.Threading.Tasks;
 using EpicGames.Horde.Accounts;
 using EpicGames.Horde.Acls;
@@ -42,6 +40,47 @@ namespace HordeServer.Server
 		/// </summary>
 		public string? ErrorMessage { get; set; }
 	}
+	
+	/// <summary>
+	/// Model for Horde account state view
+	/// </summary>
+	public class AccountStateViewModel
+	{
+		/// <summary>Whether the current user is authenticated</summary>
+		public bool IsAuthenticated { get; set; }
+		
+		/// <summary>Authenticated user's display name</summary>
+		public string? UserName { get; set; }
+		
+		/// <summary>Whether the user has administrative write permissions</summary>
+		public bool IsAdmin { get; set; }
+		
+		/// <summary>List of security claims associated with the authenticated user</summary>
+		public List<ClaimViewModel> Claims { get; set; } = new();
+		
+		/// <summary>
+		/// Gets or sets the user's ACL permissions grouped by scope.
+		/// Each scope contains a list of permissions showing which actions are authorized.
+		/// Only scopes with at least one permission are displayed in the view.
+		/// </summary>
+		public Dictionary<string, List<UserAclPermission>> ScopeToPermissions { get; set; } = new();
+	}
+	
+	/// <summary>
+	/// Model for Horde account state claim view
+	/// </summary>
+	public class ClaimViewModel
+	{
+		/// <summary>
+		/// Gets or sets the claim type identifier (e.g., "name", "role", "sub").
+		/// </summary>
+		public string Type { get; set; } = "";
+		
+		/// <summary>
+		/// Gets or sets the claim value associated with the claim type.
+		/// </summary>
+		public string Value { get; set; } = "";
+	}
 
 	/// <summary>
 	/// Controller managing account status
@@ -98,87 +137,29 @@ namespace HordeServer.Server
 		}
 
 		/// <summary>
-		/// Gets the current login status
+		/// Displays the current user's authentication status and permissions
 		/// </summary>
-		/// <returns>The current login state</returns>
 		[HttpGet]
 		[Route("/account")]
 		public ActionResult State()
 		{
-			StringBuilder content = new StringBuilder();
-			content.Append($"<html><style>{StyleSheet}</style><h1>Horde Server</h1>");
-			if (User.Identity?.IsAuthenticated ?? false)
+			AccountStateViewModel model = new()
 			{
-				content.Append(CultureInfo.InvariantCulture, $"<p>User <b>{User.Identity?.Name}</b> is logged in. <a href=\"/account/logout\">Log out</a></p>");
-				if (_globalConfig.Value.Authorize(AdminAclAction.AdminWrite, User))
-				{
-					content.Append("<p>");
-					content.Append("<a href=\"/api/v1/admin/token\">Get bearer token</a><br/>");
-					content.Append("<a href=\"/api/v1/admin/registrationtoken\">Get agent registration token</a><br/>");
-					content.Append("<a href=\"/api/v1/admin/softwaretoken\">Get agent software upload token</a><br/>");
-					content.Append("<a href=\"/api/v1/admin/softwaredownloadtoken\">Get agent software download token</a><br/>");
-					content.Append("<a href=\"/api/v1/admin/configtoken\">Get configuration token</a><br/>");
-					content.Append("<a href=\"/api/v1/admin/chainedjobtoken\">Get chained job token</a><br/>");
-					content.Append("</p>");
-				}
-				content.AppendLine(CultureInfo.InvariantCulture, $"<h2>Claims for {User.Identity?.Name}:</h2>");
-				content.AppendLine("<table>");
-				content.AppendLine("<tr><th>Type</th><th>Value</th></tr>");
-				foreach (System.Security.Claims.Claim claim in User.Claims)
-				{
-					content.AppendLine(CultureInfo.InvariantCulture, $"<tr><td>{claim.Type}</td><td>{claim.Value}</td></tr>");
-				}
-				content.AppendLine("</table>");
-				
-				content.AppendLine("<h2>ACL Permissions</h2>");
-				content.AppendLine("<p>Only scopes with at least one authorized action are shown.</p>");
-				content.AppendLine("<table class=\"acl-scopes\">");
-				content.AppendLine("<tr><th>Scope</th><th>Action</th><th>Authorized</th></tr>");
-				
+				IsAuthenticated = User.Identity?.IsAuthenticated ?? false,
+				UserName = User.Identity?.Name,
+				IsAdmin = _globalConfig.Value.Authorize(AdminAclAction.AdminWrite, User),
+			};
+			
+			if (model.IsAuthenticated)
+			{
+				model.Claims = User.Claims.Select(claim => new ClaimViewModel { Type = claim.Type, Value = claim.Value }).ToList();
 				List<UserAclPermission> aclPermissions = UserController.GetUserAclPermissions(_globalConfig.Value, User);
-				Dictionary<string, List<UserAclPermission>> scopeToPermissions = aclPermissions
+				model.ScopeToPermissions = aclPermissions
 					.GroupBy(permission => permission.Scope)
 					.ToDictionary(group => group.Key, group => group.ToList());
-				
-				foreach ((string scope, List<UserAclPermission> permissions) in scopeToPermissions)
-				{
-					if (permissions.Count > 0)
-					{
-						for (int i = 0; i < permissions.Count; i++)
-						{
-							UserAclPermission uap = permissions[i];
-							if (i == 0)
-							{
-								content.Append("<tr>");
-								content.AppendLine($"<tr><td rowspan=\"{permissions.Count + 1}\">{scope}</td>");
-							}
-							else
-							{
-								string color = uap.IsAuthorized ? "#c8ffc8" : "#ffc8c8";
-								string authorized = uap.IsAuthorized ? "Yes" : "No";
-								content.AppendLine($"<td>{uap.Action}</td><td style=\"background-color: {color}; text-align: center;\">{authorized}</td>");
-							}
-							content.Append("</tr>");
-						}
-					}
-					content.AppendLine("<tr style=\"border-bottom:1px solid black\"><td colspan=\"100%\"></td></tr>");
-					content.AppendLine("<tr><td colspan=\"100%\" style=\"display: block; height 20px;\"></td></tr>");
-				}
-				
-				if (scopeToPermissions.Count == 0)
-				{
-					content.AppendLine("<tr><td colspan=\"100%\">No scopes</td></tr>");
-				}
-				content.AppendLine("</table>");
-
-				content.AppendLine("<p>Built from Perforce</p>");
 			}
-			else
-			{
-				content.AppendLine("<p><a href=\"/account/login\"><b>Login</b></a></p>");
-			}
-			content.AppendLine("</html>");
-			return new ContentResult { ContentType = "text/html", StatusCode = (int)HttpStatusCode.OK, Content = content.ToString() };
+
+			return View("~/Server/HordeAccountState.cshtml", model);
 		}
 
 		/// <summary>
