Horde Dashboard: Error Retry Callback

diff --git a/Engine/Source/Programs/Horde/HordeDashboard/src/App.tsx b/Engine/Source/Programs/Horde/HordeDashboard/src/App.tsx
--- a/Engine/Source/Programs/Horde/HordeDashboard/src/App.tsx
+++ b/Engine/Source/Programs/Horde/HordeDashboard/src/App.tsx
@@ -16,7 +16,8 @@ import { AutomationView } from './components/AutomationView';
 import { DashboardView } from './components/DashboardView';
 import { DebugView } from './components/DebugView';
 import { DeviceView } from './components/devices/DeviceView';
-import { ErrorDialog, ErrorHandler } from './components/ErrorHandler';
+import { ErrorDialog } from './components/error/ErrorDialog';
+import { errorDialogStore } from './components/error/ErrorStore';
 import { JobRedirector } from './components/JobRedirector';
 import { LogView } from './components/LogView';
 import { NoticeView } from './components/NoticeView';
@@ -44,6 +45,7 @@ import hordePlugins from './legacyPlugins';
 import { preloadFonts } from './styles/Styles';
 import { darkTheme } from './styles/darkTheme';
 import { lightTheme } from './styles/lightTheme';
+import { ErrorBar } from './components/error/ErrorBar';
 
 let router: any;
 
@@ -79,12 +81,12 @@ const Main: React.FC = () => {
             return null;
 
          }).catch(reason => {
-            ErrorHandler.set({ title: "Error initializing site, unable to get user", reason: reason }, true);
+            errorDialogStore.set({ title: "Error initializing site, unable to get user", reason: reason }, true);
          })
 
 
       }).catch((reason) => {
-         ErrorHandler.set({ title: "Error initializing site", reason: reason }, true);
+         errorDialogStore.set({ title: "Error initializing site", reason: reason }, true);
       });
 
       return (<ThemeProvider applyTo='element' theme={dashboard.darktheme ? darkTheme : lightTheme}>
diff --git a/Engine/Source/Programs/Horde/HordeDashboard/src/backend/Fetch.ts b/Engine/Source/Programs/Horde/HordeDashboard/src/backend/Fetch.ts
--- a/Engine/Source/Programs/Horde/HordeDashboard/src/backend/Fetch.ts
+++ b/Engine/Source/Programs/Horde/HordeDashboard/src/backend/Fetch.ts
@@ -1,6 +1,6 @@
 // Copyright Epic Games, Inc. All Rights Reserved.
 
-import { ErrorHandler, ErrorInfo } from "../components/ErrorHandler";
+import { errorDialogStore, ErrorInfo, errorBarStore, ErrorReportType } from "../components/error/ErrorStore";
 import { GetDashboardChallengeResponse } from "./Api";
 import userInactivity from "./UserInactivity";
 //import { setDatadogUser } from './Datadog';
@@ -26,12 +26,26 @@ export type FetchResponse = {
     status: number;
 }
 
-function handleError(errorIn: ErrorInfo, force?: boolean) {
+function handleError(
+    errorIn: ErrorInfo, 
+    force?: boolean,
+    reportType?: ErrorReportType,
+    retryCallback?: () => boolean | Promise<boolean>
+) {
 
     const e = { ...errorIn, hordeError: true } as ErrorInfo;
 
-    ErrorHandler.set(e, force ?? false);
-
+    switch(reportType){
+        case ErrorReportType.Dialog:
+            errorDialogStore.set(e, force ?? false, retryCallback);
+            break;
+        case ErrorReportType.Bar:
+            errorBarStore.set(e, force ?? false, retryCallback);
+            break;
+        default:
+            errorDialogStore.set(e, force ?? false, retryCallback); // Make default Bar after dialog investigation
+            break;
+    }
 }
 
 /** Fetch convenience wrapper */
@@ -41,11 +55,16 @@ export class Fetch {
 
         url = this.buildUrl(url, config);
 
+        if(this.checkDebugMockStatus(url)){
+            return Promise.reject(new Error("Vite endpoint status debugging"))
+        }
+
         return new Promise<FetchResponse>(async (resolve, reject) => {
 
             try {
                 
                 const response = await fetch(url, this.buildRequest("GET")).then(response => {
+                    // Internal Server Error Response
                     if (!response.ok && response.status === 500) {
                         response.text().then(text => {
                             handleError({
@@ -59,7 +78,7 @@ export class Fetch {
                         return null;
                     }
 
-
+                    // 404 Not Found Response
                     if (!response.ok && response.status === 404 && config?.show404Error) {
                         response.json().then(o => {
                             handleError({
@@ -76,6 +95,31 @@ export class Fetch {
                         return null;
                     }
 
+                    // 401 Unauthorized Response - No Credentials
+                    if(!response.ok && response.status === 401){
+                        handleError({
+                            response,
+                            title: "401 Unauthorized",
+                            reason: "401 Unauthorized",
+                            mode: "GET",
+                            url,
+                            message: "You don't have access because you are not logged in. Try logging in again."
+                        }, true, ErrorReportType.Bar)
+                        return null;
+                    }
+
+                    // 403 Forbidden Response - Insufficient Permissions
+                    if(!response.ok && response.status === 403){
+                        handleError({
+                            response,
+                            title: "403 Forbidden",
+                            reason: "403 Forbidden",
+                            mode: "GET",
+                            url,
+                            message: "You don't have access to this page or some of its features"
+                        }, true, ErrorReportType.Bar)
+                        return null;
+                    }
 
                     return response;
                 });
@@ -128,6 +172,92 @@ export class Fetch {
         });
     }
 
+    // DEBUG code, used to test response status handling.
+    checkDebugMockStatus(url: string): boolean {
+        if(!import.meta.env.VITE_HORDE_DEBUG_MOCK_STATUS_ENDPOINTS){
+            // No debug settings found, return false to continue the calling function
+            return false;
+        }
+        
+        const mockStatusEndpoints: Record<string, [number, ErrorReportType]> = {};
+        // Parse the vite mock status envrionment variable into a record map-like object
+        import.meta.env.VITE_HORDE_DEBUG_MOCK_STATUS_ENDPOINTS
+            .split(",")
+            .forEach((triplet, index) => {
+                try {
+                    const [endpointRaw, statusRaw, reportTypeRaw] = triplet.split(":");
+                    
+                    // Endpoint parsing
+                    if(!endpointRaw){
+                        throw new Error(`Missing endpoint in triplet "${triplet}" at index ${index}`);
+                    }
+                    const endpoint = endpointRaw.trim();
+                    
+                    // Status parsing
+                    let status = 401;
+                    if(statusRaw){
+                        const parsedStatus = Number(statusRaw.trim());
+                        if(isNaN(parsedStatus) || parsedStatus < 100 || parsedStatus > 599) {
+                            console.warn(`Invalid status "${statusRaw}" in triplet "${triplet}" at index ${index}. Defaulting to status 401.`);
+                        }
+                        else {
+                            status = parsedStatus;
+                        }
+                    }
+                    else {
+                        console.warn(`No status provided in triplet "${triplet}" at index ${index}. Defaulting to status 401. If this was unintended, please use the following format: "endpoint:status:reportType".`);
+                    }
+
+                    // Report type parsing
+                    let reportType = ErrorReportType.Bar;
+                    if(reportTypeRaw){
+                        const parsedReportType = ErrorReportType[reportTypeRaw.trim() as keyof ErrorReportType];
+                        if(!parsedReportType){
+                            console.warn(`Invalid report type "${reportTypeRaw}" in triplet "${triplet}" at index ${index}. Defaulting to ErrorReportType "Bar".`);
+                        } 
+                        else {
+                            reportType = parsedReportType;
+                        }
+                    }
+                    else {
+                        console.warn(`No report type provided in triplet "${triplet}" at index ${index}. Defaulting to ErrorReportType "Bar". If this was unintended, please use the following format: "endpoint:status:reportType".`);
+                    }
+
+                    mockStatusEndpoints[endpoint] = [status, reportType];
+                } catch (e) {
+                    console.error("Issue with VITE Debug Env Var: VITE_HORDE_DEBUG_MOCK_STATUS_ENDPOINTS", e)
+                }
+            }); // No need for super careful list construction
+
+        // Determine if any entries match the current url
+        const match = Object.entries(mockStatusEndpoints).find(([endpoint]) => url.includes(endpoint));
+        // If a match is found, simulate a bad response with the given status and mock an error with the given report type
+        if(match) {
+            const [endpoint, [status, reportType]] = match;
+
+            const response = new Response(JSON.stringify({error: `Error ${status} (Mock)`}), {
+                status,
+                headers: {"Content-Type": "application/json"}
+            });
+
+            if (!response.ok) {
+                handleError({
+                    response,
+                    title: `Error ${status}`,
+                    reason: "Endpoint Status Debugging",
+                    mode: "GET",
+                    url,
+                }, true, reportType);
+
+                // A match was found, return true to return from the calling function early
+                return true;
+            }
+        }
+
+        // A match was not found, return false to continue calling function
+        return false;
+    }
+
     post(url: string, data?: any, config?: FetchRequestConfig) {
 
         url = this.buildUrl(url, config);
diff --git a/Engine/Source/Programs/Horde/HordeDashboard/src/components/AutoSubmit.tsx b/Engine/Source/Programs/Horde/HordeDashboard/src/components/AutoSubmit.tsx
--- a/Engine/Source/Programs/Horde/HordeDashboard/src/components/AutoSubmit.tsx
+++ b/Engine/Source/Programs/Horde/HordeDashboard/src/components/AutoSubmit.tsx
@@ -7,8 +7,8 @@ import backend from '../backend';
 import { JobState } from '../backend/Api';
 import dashboard from '../backend/Dashboard';
 
-import { ErrorHandler } from './ErrorHandler';
 import { JobDetailsV2 } from './jobDetailsV2/JobDetailsViewCommon';
+import { errorDialogStore } from './error/ErrorStore';
 
 export const AutosubmitInfo: React.FC<{ jobDetails: JobDetailsV2 }> = observer(({ jobDetails }) => {
 
@@ -52,7 +52,7 @@ export const AutosubmitInfo: React.FC<{ jobDetails: JobDetailsV2 }> = observer((
                   setState({ autoSubmit: value, inflight: false })
                }).catch(reason => {
 
-                  ErrorHandler.set({
+                  errorDialogStore.set({
                      reason: reason,
                      title: `Error Setting Auto-submit`,
                      message: `There was an error setting job to autosubmit, reason: "${reason}"`
diff --git a/Engine/Source/Programs/Horde/HordeDashboard/src/components/AutomationView.tsx b/Engine/Source/Programs/Horde/HordeDashboard/src/components/AutomationView.tsx
--- a/Engine/Source/Programs/Horde/HordeDashboard/src/components/AutomationView.tsx
+++ b/Engine/Source/Programs/Horde/HordeDashboard/src/components/AutomationView.tsx
@@ -16,7 +16,7 @@ import { getHumanTime, getShortNiceTime, msecToElapsed } from '../base/utilities
 import { AutomationSuiteDetails } from './AutomationSuiteDetails';
 import { AutomationViewSummary } from './AutomationViewSummary';
 import { Breadcrumbs } from './Breadcrumbs';
-import ErrorBoundary from './ErrorBoundary';
+import ErrorBoundary from './error/ErrorBoundary';
 import { TopNav } from './TopNav';
 import { getHordeStyling } from '../styles/Styles';
 
diff --git a/Engine/Source/Programs/Horde/HordeDashboard/src/components/Breadcrumbs.tsx b/Engine/Source/Programs/Horde/HordeDashboard/src/components/Breadcrumbs.tsx
--- a/Engine/Source/Programs/Horde/HordeDashboard/src/components/Breadcrumbs.tsx
+++ b/Engine/Source/Programs/Horde/HordeDashboard/src/components/Breadcrumbs.tsx
@@ -1,6 +1,6 @@
 // Copyright Epic Games, Inc. All Rights Reserved.
 
-import { mergeStyleSets, MessageBar, MessageBarType, Separator, Spinner, SpinnerSize, Stack, Text } from '@fluentui/react';
+import { divProperties, mergeStyleSets, MessageBar, MessageBarType, Separator, Spinner, SpinnerSize, Stack, Text } from '@fluentui/react';
 import { observer } from 'mobx-react-lite';
 import React, { useState } from 'react';
 import { Link } from 'react-router-dom';
@@ -9,6 +9,7 @@ import { Markdown } from "../base/components/Markdown";
 import notices from '../backend/Notices';
 import { getHordeTheme } from '../styles/theme';
 import { getHordeStyling } from '../styles/Styles';
+import { ErrorBar } from './error/ErrorBar';
 
 export type BreadcrumbItem = {
    text: string;
@@ -138,6 +139,12 @@ export const Breadcrumbs: React.FC<{ items: BreadcrumbItem[], title?: string, su
                <Stack grow />
             </Stack>
          </Stack>}
+
+         {!alert &&
+            <div style={{margin: "auto", position: "absolute", width: "fit-content", left: 0, right: 0}}>
+               <ErrorBar/>
+            </div>
+         }
       </Stack>
    </Stack>
    );
diff --git a/Engine/Source/Programs/Horde/HordeDashboard/src/components/ErrorBar.tsx b/Engine/Source/Programs/Horde/HordeDashboard/src/components/ErrorBar.tsx
deleted file mode 100644
diff --git a/Engine/Source/Programs/Horde/HordeDashboard/src/components/ErrorBoundary.tsx b/Engine/Source/Programs/Horde/HordeDashboard/src/components/ErrorBoundary.tsx
deleted file mode 100644
diff --git a/Engine/Source/Programs/Horde/HordeDashboard/src/components/ErrorHandler.tsx b/Engine/Source/Programs/Horde/HordeDashboard/src/components/ErrorHandler.tsx
deleted file mode 100644
diff --git a/Engine/Source/Programs/Horde/HordeDashboard/src/components/ErrorPane.tsx b/Engine/Source/Programs/Horde/HordeDashboard/src/components/ErrorPane.tsx
deleted file mode 100644
diff --git a/Engine/Source/Programs/Horde/HordeDashboard/src/components/IssueViewV2.tsx b/Engine/Source/Programs/Horde/HordeDashboard/src/components/IssueViewV2.tsx
--- a/Engine/Source/Programs/Horde/HordeDashboard/src/components/IssueViewV2.tsx
+++ b/Engine/Source/Programs/Horde/HordeDashboard/src/components/IssueViewV2.tsx
@@ -14,7 +14,7 @@ import { useWindowSize } from "../base/utilities/hooks";
 import { displayTimeZone, getHumanTime, getShortNiceTime } from "../base/utilities/timeUtils";
 import { getHordeStyling } from "../styles/Styles";
 import { getHordeTheme } from "../styles/theme";
-import { ErrorHandler } from "./ErrorHandler";
+import { errorDialogStore } from "./error/ErrorStore";
 import { useQuery } from "./JobDetailCommon";
 import { renderLine } from "./LogRender";
 import { UserSelect } from "./UserSelect";
@@ -3335,7 +3335,7 @@ const TestFixModal: React.FC<{ onClose: () => void }> = ({ onClose }) => {
 
                   if (errorReason) {
 
-                     ErrorHandler.set({
+                     errorDialogStore.set({
 
                         reason: `${errorReason}`,
                         title: `Error Creating Job`,
diff --git a/Engine/Source/Programs/Horde/HordeDashboard/src/components/JobRedirector.tsx b/Engine/Source/Programs/Horde/HordeDashboard/src/components/JobRedirector.tsx
--- a/Engine/Source/Programs/Horde/HordeDashboard/src/components/JobRedirector.tsx
+++ b/Engine/Source/Programs/Horde/HordeDashboard/src/components/JobRedirector.tsx
@@ -4,7 +4,7 @@ import React, { useState } from 'react';
 import { useNavigate } from 'react-router-dom';
 import backend from '../backend';
 import { useQuery } from './JobDetailCommon';
-import ErrorHandler from './ErrorHandler';
+import {errorDialogStore} from "./error/ErrorStore";
 import { JobStepState } from 'horde/backend/Api';
 
 function setError(message: string) {
@@ -13,7 +13,7 @@ function setError(message: string) {
 
    message = `${message}\n\nReferring URL: ${encodeURI(window.location.href)}`
 
-   ErrorHandler.set({
+   errorDialogStore.set({
 
       title: `Error handling Preflight`,
       message: message
diff --git a/Engine/Source/Programs/Horde/HordeDashboard/src/components/NewBuild.tsx b/Engine/Source/Programs/Horde/HordeDashboard/src/components/NewBuild.tsx
--- a/Engine/Source/Programs/Horde/HordeDashboard/src/components/NewBuild.tsx
+++ b/Engine/Source/Programs/Horde/HordeDashboard/src/components/NewBuild.tsx
@@ -7,7 +7,7 @@ import React, { useState } from 'react';
 import backend, { useBackend } from '../backend';
 import { BoolParameterData, CreateJobRequest, GetJobsTabResponse, GroupParameterData, JobsTabData, ListParameterData, ListParameterItemData, ListParameterStyle, ParameterData, ParameterType, Priority, TabType, GetTemplateRefResponse, TextParameterData, ChangeQueryConfig } from '../backend/Api';
 import templateCache from '../backend/TemplateCache';
-import { ErrorHandler } from "../components/ErrorHandler";
+import { errorDialogStore } from "./error/ErrorStore";
 import { useQuery } from './JobDetailCommon';
 import { JobDetailsV2 } from './jobDetailsV2/JobDetailsViewCommon';
 import dashboard from '../backend/Dashboard';
@@ -918,7 +918,7 @@ export const NewBuild: React.FC<{ streamId: string; show: boolean; onClose: (new
             }
 
             if (errorReason) {
-               ErrorHandler.set({
+               errorDialogStore.set({
                   reason: `${errorReason}`,
                   title: `Preflight Template Error`,
                   message: `There was an issue with the specified preflight template.\n\nReason: ${errorReason}\n\nTime: ${moment.utc().format("MMM Do, HH:mm z")}`
@@ -1572,7 +1572,7 @@ export const NewBuild: React.FC<{ streamId: string; show: boolean; onClose: (new
                         errorReason += ".  Perforce edge server replication for the change may be in progress."
                      }
 
-                     ErrorHandler.set({
+                     errorDialogStore.set({
 
                         reason: `${errorReason}`,
                         title: `Error Creating Job`,
diff --git a/Engine/Source/Programs/Horde/HordeDashboard/src/components/Preflight.tsx b/Engine/Source/Programs/Horde/HordeDashboard/src/components/Preflight.tsx
--- a/Engine/Source/Programs/Horde/HordeDashboard/src/components/Preflight.tsx
+++ b/Engine/Source/Programs/Horde/HordeDashboard/src/components/Preflight.tsx
@@ -4,7 +4,7 @@ import React, { useState } from 'react';
 import { useNavigate } from 'react-router-dom';
 import backend, { useBackend } from '../backend';
 import { useQuery } from './JobDetailCommon';
-import ErrorHandler from './ErrorHandler';
+import {errorDialogStore} from './error/ErrorStore';
 import { GetJobsTabResponse, JobsTabData, TabType } from '../backend/Api';
 
 function setError(message: string) {
@@ -13,7 +13,7 @@ function setError(message: string) {
 
    message = `${message}\n\nReferring URL: ${encodeURI(window.location.href)}`
 
-   ErrorHandler.set({
+   errorDialogStore.set({
 
       title: `Error handling Preflight`,
       message: message
diff --git a/Engine/Source/Programs/Horde/HordeDashboard/src/components/accounts/AccountsView.tsx b/Engine/Source/Programs/Horde/HordeDashboard/src/components/accounts/AccountsView.tsx
--- a/Engine/Source/Programs/Horde/HordeDashboard/src/components/accounts/AccountsView.tsx
+++ b/Engine/Source/Programs/Horde/HordeDashboard/src/components/accounts/AccountsView.tsx
@@ -10,7 +10,7 @@ import { useWindowSize } from "../../base/utilities/hooks";
 import { getHordeStyling } from "../../styles/Styles";
 import { Breadcrumbs } from "../Breadcrumbs";
 import { TopNav } from "../TopNav";
-import ErrorHandler from "../ErrorHandler";
+import { errorDialogStore } from "../error/ErrorStore";
 import moment from "moment";
 
 class AccountHandler extends PollBase {
@@ -246,7 +246,7 @@ const AccountEditor: React.FC<{ accountIn?: GetAccountResponse, onClose: () => v
 
       } catch (reason) {
          console.error(reason);
-         ErrorHandler.set({
+         errorDialogStore.set({
             reason: `${reason}`,
             title: `User Modification Error`,
             message: `There was an issue modifying the user.\n\nReason: ${reason}\n\nTime: ${moment.utc().format("MMM Do, HH:mm z")}`
diff --git a/Engine/Source/Programs/Horde/HordeDashboard/src/components/accounts/ServiceAccountsView.tsx b/Engine/Source/Programs/Horde/HordeDashboard/src/components/accounts/ServiceAccountsView.tsx
--- a/Engine/Source/Programs/Horde/HordeDashboard/src/components/accounts/ServiceAccountsView.tsx
+++ b/Engine/Source/Programs/Horde/HordeDashboard/src/components/accounts/ServiceAccountsView.tsx
@@ -10,7 +10,7 @@ import { useWindowSize } from "../../base/utilities/hooks";
 import { getHordeStyling } from "../../styles/Styles";
 import { Breadcrumbs } from "../Breadcrumbs";
 import { TopNav } from "../TopNav";
-import ErrorHandler from "../ErrorHandler";
+import { errorDialogStore } from "../error/ErrorStore";
 import moment from "moment";
 
 class ServiceAccountHandler extends PollBase {
@@ -249,7 +249,7 @@ const ServiceAccountEditor: React.FC<{ accountIn?: GetServiceAccountResponse, on
 
       } catch (reason) {
          console.error(reason);
-         ErrorHandler.set({
+         errorDialogStore.set({
             reason: `${reason}`,
             title: `User Modification Error`,
             message: `There was an issue modifying the user.\n\nReason: ${reason}\n\nTime: ${moment.utc().format("MMM Do, HH:mm z")}`
diff --git a/Engine/Source/Programs/Horde/HordeDashboard/src/components/agents/AgentRequestsView.tsx b/Engine/Source/Programs/Horde/HordeDashboard/src/components/agents/AgentRequestsView.tsx
--- a/Engine/Source/Programs/Horde/HordeDashboard/src/components/agents/AgentRequestsView.tsx
+++ b/Engine/Source/Programs/Horde/HordeDashboard/src/components/agents/AgentRequestsView.tsx
@@ -10,7 +10,7 @@ import { useWindowSize } from "../../base/utilities/hooks";
 import { getHordeStyling } from "../../styles/Styles";
 import { Breadcrumbs } from "../Breadcrumbs";
 import { TopNav } from "../TopNav";
-import ErrorHandler from "../ErrorHandler";
+import { errorDialogStore } from "../error/ErrorStore";
 import { useNavigate } from "react-router-dom";
 import moment from "moment";
 import { agentStore } from "./AgentStore"
@@ -111,7 +111,7 @@ const AgentsPanel: React.FC = observer(() => {
       } catch (reason) {
          console.error(reason);
 
-         ErrorHandler.set({
+         errorDialogStore.set({
             reason: reason,
             title: `Error Enrolling Agents`,
             message: `There was an error enrolling agents, reason: "${reason}"`
diff --git a/Engine/Source/Programs/Horde/HordeDashboard/src/components/artifacts/ArtifactsRedirector.tsx b/Engine/Source/Programs/Horde/HordeDashboard/src/components/artifacts/ArtifactsRedirector.tsx
--- a/Engine/Source/Programs/Horde/HordeDashboard/src/components/artifacts/ArtifactsRedirector.tsx
+++ b/Engine/Source/Programs/Horde/HordeDashboard/src/components/artifacts/ArtifactsRedirector.tsx
@@ -2,7 +2,7 @@
 
 import React from 'react';
 import { Navigate, useNavigate, useParams } from 'react-router-dom';
-import ErrorHandler from '../ErrorHandler';
+import { errorDialogStore } from '../error/ErrorStore';
 import backend from '../../backend';
 
 function setError(artifactId: string, message: string) {
@@ -11,7 +11,7 @@ function setError(artifactId: string, message: string) {
 
    message = `${message}\n\nReferring URL: ${encodeURI(window.location.href)}`
 
-   ErrorHandler.set({
+   errorDialogStore.set({
 
       title: `Error Getting Artifact ${artifactId}`,
       message: message
diff --git a/Engine/Source/Programs/Horde/HordeDashboard/src/components/build/NewBuildV2.tsx b/Engine/Source/Programs/Horde/HordeDashboard/src/components/build/NewBuildV2.tsx
--- a/Engine/Source/Programs/Horde/HordeDashboard/src/components/build/NewBuildV2.tsx
+++ b/Engine/Source/Programs/Horde/HordeDashboard/src/components/build/NewBuildV2.tsx
@@ -13,7 +13,7 @@ import { ProjectStore } from "../../backend/ProjectStore";
 import templateCache from "../../backend/TemplateCache";
 import { copyToClipboard } from "../../base/utilities/clipboard";
 import { getHordeStyling } from "../../styles/Styles";
-import ErrorHandler from "../ErrorHandler";
+import { errorDialogStore } from "../error/ErrorStore";
 import { NewBuild } from "../NewBuild";
 import { JobDetailsV2 } from "../jobDetailsV2/JobDetailsViewCommon";
 
@@ -899,7 +899,7 @@ class BuildOptions {
             }
 
             if (errorReason) {
-               ErrorHandler.set({
+               errorDialogStore.set({
                   reason: `${errorReason}`,
                   title: `Preflight Template Error`,
                   message: `There was an issue with the specified preflight template.\n\nReason: ${errorReason}\n\nTime: ${moment.utc().format("MMM Do, HH:mm z")}`
@@ -2016,7 +2016,7 @@ const BuildModal: React.FC<{ setUseLegacyDialog: (value: boolean) => void }> = o
                      errorReason += ".\nPerforce edge server replication for the change may be in progress. Please try again in a few minutes."
                   }
 
-                  ErrorHandler.set({
+                  errorDialogStore.set({
 
                      reason: `${errorReason}`,
                      title: `Error Creating Job`,
diff --git a/Engine/Source/Programs/Horde/HordeDashboard/src/components/error/ErrorBar.tsx b/Engine/Source/Programs/Horde/HordeDashboard/src/components/error/ErrorBar.tsx
new file mode 100644
--- /dev/null
+++ b/Engine/Source/Programs/Horde/HordeDashboard/src/components/error/ErrorBar.tsx
@@ -0,0 +1,53 @@
+// Copyright Epic Games, Inc. All Rights Reserved.
+
+import React from 'react';
+import { MessageBar, MessageBarType } from '@fluentui/react';
+import { errorBarStore } from './ErrorStore';
+import { observer } from 'mobx-react-lite';
+import dashboard from 'horde/backend/Dashboard';
+
+export const ErrorBar: React.FC = observer(() => {
+    const {error, clear} = errorBarStore;
+
+    const helpEmail = dashboard.helpEmail;
+    const helpSlack = dashboard.helpSlack;
+    const anyHelp = !!helpEmail || !!helpSlack;
+
+    const getConsoleErrorMessage = (): string => {
+        if(!error) return "";
+
+        let msg = `${error.title}: ${errorBarStore.message}`
+        msg += error.url ? ` - Occurred while trying to access ${error.url}` : "";
+        msg += anyHelp ? ` - If the issue persists, please contact us for assistance at ${helpEmail} OR the ${helpSlack} Slack channel.` : "";
+        return msg;
+    }
+
+    // Report extra information in the broswer console (leave room to breathe in the UI)
+    if(error){
+        console.error(getConsoleErrorMessage());
+    }
+
+    if(!error) {
+        return <div />
+    }
+
+    return (
+        <MessageBar
+            messageBarType={MessageBarType.error} 
+            isMultiline={false}
+            hidden={!error}
+            truncated={true}
+            onDismiss={clear}
+            dismissButtonAriaLabel="Close"
+            actions={
+                <div>
+                    {/*Opportunity for retry or filter implementation if needed.*/}
+                    {/* <MessageBarButton>Never show this again</MessageBarButton> */}
+                    {/* <MessageBarButton>Retry</MessageBarButton> */}
+                </div>
+            }
+        >
+            <span>{error.title}: {errorBarStore.message}. {(error.url || anyHelp) && "Check your browser's developer console."}</span>
+        </MessageBar>
+    );
+});
\ No newline at end of file
diff --git a/Engine/Source/Programs/Horde/HordeDashboard/src/components/error/ErrorBoundary.tsx b/Engine/Source/Programs/Horde/HordeDashboard/src/components/error/ErrorBoundary.tsx
new file mode 100644
--- /dev/null
+++ b/Engine/Source/Programs/Horde/HordeDashboard/src/components/error/ErrorBoundary.tsx
@@ -0,0 +1,45 @@
+import { Stack, Text } from "@fluentui/react";
+import { Component, ErrorInfo, ReactNode } from "react";
+import { errorDialogStore } from "./ErrorStore";
+
+interface Props {
+   children?: ReactNode;
+}
+
+interface State {
+   hasError: boolean;
+}
+
+class ErrorBoundary extends Component<Props, State> {
+   public state: State = {
+      hasError: false
+   };
+
+   public static getDerivedStateFromError(_: Error): State {
+      return { hasError: true };
+   }
+
+   public componentDidCatch(error: Error, errorInfo: ErrorInfo) {      
+      console.error("Uncaught Exception:", error, errorInfo);
+      errorDialogStore.set({
+         reason: `${error}`,
+         title: `Uncaught Exception`,
+         message: `Error: ${error}\n\n${errorInfo?.componentStack}`
+
+      }, true);
+   }
+
+   public render() {
+
+      if (this.state.hasError) {
+         return <Stack style={{padding: 128}}><Text variant="mediumPlus">Uncaught Exception, please see console for more information</Text></Stack>;
+      }
+                  
+      return this.props.children;
+   }
+
+   static error?: Error;
+   static errorInfo?: ErrorInfo;
+}
+
+export default ErrorBoundary;
\ No newline at end of file
diff --git a/Engine/Source/Programs/Horde/HordeDashboard/src/components/error/ErrorDialog.tsx b/Engine/Source/Programs/Horde/HordeDashboard/src/components/error/ErrorDialog.tsx
new file mode 100644
--- /dev/null
+++ b/Engine/Source/Programs/Horde/HordeDashboard/src/components/error/ErrorDialog.tsx
@@ -0,0 +1,143 @@
+// Copyright Epic Games, Inc. All Rights Reserved.
+
+import { observer } from 'mobx-react-lite';
+import { BaseButton, Checkbox, DefaultButton, Dialog, DialogType, IButtonStyles, mergeStyleSets, PrimaryButton, Spinner, SpinnerSize, Stack, Text } from '@fluentui/react';
+import React from 'react';
+import { errorDialogStore } from './ErrorStore';
+import dashboard, { StatusColor } from '../../backend/Dashboard';
+import { getErrorDialogStyles } from './ErrorDialogStyles';
+
+export const ErrorDialog: React.FC = observer(() => {
+   const {
+      canRetry,
+      retry,
+      clear,
+      error,
+   } = errorDialogStore;
+
+   if (!error) {
+      return <div />;
+   }
+
+   let message = errorDialogStore.message;
+   const unauthorized = errorDialogStore.unauthorized();
+
+   let title = error.title;
+   if (!title) {
+      title = unauthorized ? "Permission Error" : "Oops, there was a problem...";
+   }
+
+   const dialogContentProps = {
+      type: DialogType.normal,
+      title: title
+   };
+
+   const helpEmail = dashboard.helpEmail;
+   const helpSlack = dashboard.helpSlack;
+
+   let fontSize = 13;
+
+   let width = 800;
+   if (message?.length) {
+      try {
+         const exception = JSON.parse(message);
+         if (exception?.exception?.message && exception?.exception?.trace) {
+            width = 1380;
+            fontSize = 11
+            message = `${exception?.exception?.message}\n${exception?.exception?.trace}`
+         }
+      } catch (error) {
+
+      }
+   }
+
+   const anyHelp = !!helpEmail || !!helpSlack;
+
+   const errorDialogStyles = getErrorDialogStyles();
+
+   const {
+      isRetrying,
+      succeeded: retrySucceeded,
+      failed: retryFailed,
+      hasResult: retryHasResult,
+   } = errorDialogStore.retryInfo;
+
+   const getRetryButtonText = () => {
+      if (retrySucceeded) return "Success";
+      if (retryFailed) return "Retry Failed";
+      if (isRetrying) return "Retrying...";
+      if (retryHasResult) return "Try Again?";
+      return "Retry";
+   };
+
+   const getRetryButtonClass = () => {
+      if (retrySucceeded) return errorDialogStyles.successButton;
+      if (retryFailed) return errorDialogStyles.errorButtonDisabled;
+      if (retryHasResult) return errorDialogStyles.errorButton;
+      return errorDialogStyles.button;
+   };
+
+   return (
+      <Dialog
+         minWidth={width}
+         hidden={!error}
+         onDismiss={clear}
+         modalProps={{
+            isBlocking: true
+         }}
+         dialogContentProps={dialogContentProps}>
+
+         <Stack tokens={{ childrenGap: 12 }}>
+            <Text styles={{ root: { whiteSpace: "pre", fontSize: fontSize, color: "#EC4C47" } }}>{message}</Text>
+            {!!error.url && <Text variant="medium">{`URL: ${error.url}`}</Text>}
+            <Stack>
+               {anyHelp && <Stack tokens={{ childrenGap: 18 }} style={{ paddingTop: 12 }}>
+                  <Text style={{ fontSize: 15 }}>If the issue persists, please contact us for assistance.</Text>
+                  <Stack tokens={{ childrenGap: 14 }} style={{ paddingLeft: 12, paddingTop: 2 }}>
+                     {!!helpEmail && 
+                        <Stack horizontal tokens={{ childrenGap: 12 }}>
+                           <Text className={errorDialogStyles.supportTextKey}>Email:</Text>
+                           <Text className={errorDialogStyles.supportTextValue}>{helpEmail}</Text>
+                        </Stack>
+                     }
+                     {!!helpSlack && 
+                        <Stack horizontal tokens={{ childrenGap: 12 }}>
+                           <Text className={errorDialogStyles.supportTextKey}>Slack:</Text>
+                           <Text className={errorDialogStyles.supportTextValue}>{helpSlack}</Text>
+                        </Stack>
+                     }
+                  </Stack>
+               </Stack>}
+
+               <Stack horizontal style={{ paddingTop: 24, gap: 24, alignItems: "center", justifyContent: "end" }}>
+                  <Checkbox 
+                     label="Don't show errors like this again" 
+                     defaultChecked={false} 
+                     onChange={(e, b) => { errorDialogStore.filterError = !!b; }} 
+                  />
+                  {!canRetry && <PrimaryButton text="Dismiss" onClick={clear} />}
+                  {canRetry && (
+                     <>
+                        <DefaultButton text="Dismiss" onClick={clear} />
+                        
+                        <PrimaryButton
+                           text={getRetryButtonText()}
+                           onClick={retry}
+                           disabled={isRetrying}
+                           className={getRetryButtonClass()}
+                        >
+                           {(isRetrying && !retryHasResult) && <Spinner size={SpinnerSize.small}/>}
+                        </PrimaryButton>
+                     </>
+                  )}
+               </Stack>
+               {retrySucceeded && (
+                  <Stack.Item align='end'>
+                     <p style={{marginBottom: 0}}>You can safely dismiss this dialog.</p>
+                  </Stack.Item>
+               )}
+            </Stack>
+         </Stack>
+      </Dialog>
+   );
+});
\ No newline at end of file
diff --git a/Engine/Source/Programs/Horde/HordeDashboard/src/components/error/ErrorDialogStyles.tsx b/Engine/Source/Programs/Horde/HordeDashboard/src/components/error/ErrorDialogStyles.tsx
new file mode 100644
--- /dev/null
+++ b/Engine/Source/Programs/Horde/HordeDashboard/src/components/error/ErrorDialogStyles.tsx
@@ -0,0 +1,64 @@
+import { mergeStyleSets } from "@fluentui/react"
+import dashboard, { StatusColor } from "horde/backend/Dashboard";
+
+export const getErrorDialogStyles = () => {
+    const colors = dashboard.getDefaultStatusColors();
+
+    const errorDialogStyles = mergeStyleSets({
+        button: {
+            width: "117px",
+        },
+        
+        errorButton: {
+            backgroundColor: colors.get(StatusColor.Failure),
+            borderStyle: "hidden",
+            width: "117px",
+            color: "#FFFFFF",
+            selectors: {
+            ':active,:hover': {
+                filter: "brightness(120%)",
+                backgroundColor: colors.get(StatusColor.Failure),
+                width: "117px",
+                borderStyle: "hidden",
+                color: "#FFFFFF",
+            }
+            }
+        },
+
+        errorButtonDisabled: {
+            backgroundColor: colors.get(StatusColor.Failure),
+            borderStyle: "hidden",
+            width: "117px",
+            color: "#FFFFFF",
+            selectors: {
+            ':active,:hover': {
+                backgroundColor: colors.get(StatusColor.Failure),
+                borderStyle: "hidden",
+                width: "117px",
+                color: "#FFFFFF",
+            }
+            }
+        },
+
+        successButton: {
+            backgroundColor: colors.get(StatusColor.Success),
+            borderStyle: "hidden",
+            width: "117px",
+            color: "#FFFFFF",
+            selectors: {
+            ':active,:hover': {
+                backgroundColor: colors.get(StatusColor.Success),
+                borderStyle: "hidden",
+                color: "#FFFFFF",
+                width: "117px",
+            }
+            }
+        },
+
+        supportTextKey: {fontSize: 15},
+
+        supportTextValue: { fontSize: 15, fontWeight: 400, fontFamily: "Horde Open Sans SemiBold" }
+    })
+
+    return errorDialogStyles;
+}
\ No newline at end of file
diff --git a/Engine/Source/Programs/Horde/HordeDashboard/src/components/error/ErrorPane.tsx b/Engine/Source/Programs/Horde/HordeDashboard/src/components/error/ErrorPane.tsx
new file mode 100644
--- /dev/null
+++ b/Engine/Source/Programs/Horde/HordeDashboard/src/components/error/ErrorPane.tsx
@@ -0,0 +1,320 @@
+// Copyright Epic Games, Inc. All Rights Reserved.
+
+import { observer } from 'mobx-react-lite';
+import { IColumn, List, Stack, Text } from '@fluentui/react';
+import { getFocusStyle, mergeStyleSets } from '@fluentui/react/lib/Styling';
+import React from 'react';
+import { Link, useNavigate } from 'react-router-dom';
+import { EventData, EventSeverity, GetLabelStateResponse, JobStepState } from '../../backend/Api';
+import { JobDetails } from '../../backend/JobDetails';
+import { JobEventHandler } from '../../backend/JobEventHandler';
+import { renderLine } from '../LogRender';
+import { StepStatusIcon } from '../StatusIcon';
+import { getHordeStyling } from '../../styles/Styles';
+import { getHordeTheme } from '../../styles/theme';
+
+let _styles: any;
+
+const getStyles = () => {
+
+   const theme = getHordeTheme();
+
+   const styles = _styles ?? mergeStyleSets({
+      gutter: [
+         {
+            borderLeftStyle: 'solid',
+            borderLeftColor: "#EC4C47",
+            borderLeftWidth: 6,
+            padding: 0,
+            margin: 0,
+            paddingTop: 8,
+            paddingBottom: 8,
+            paddingRight: 8,
+            marginTop: 0,
+            marginBottom: 0
+         }
+      ],
+      gutterWarning: [
+         {
+            borderLeftStyle: 'solid',
+            borderLeftColor: "rgb(247, 209, 84)",
+            borderLeftWidth: 6,
+            padding: 0,
+            margin: 0,
+            paddingTop: 8,
+            paddingBottom: 8,
+            paddingRight: 8,
+            marginTop: 0,
+            marginBottom: 0
+         }
+      ],
+      itemCell: [
+         getFocusStyle(theme, { inset: -1 }),
+         {
+            selectors: {
+               '&:hover': { background: theme.palette.neutralLight }
+            }
+         }
+      ],
+
+   });
+
+   _styles = styles;
+
+   return styles;
+}
+
+
+export const ErrorPane: React.FC<{ jobDetails: JobDetails; stepId: string; showErrors: boolean; count?: number }> = ({ jobDetails, stepId, showErrors, count }) => {
+
+   const navigate = useNavigate();
+   const styles = getStyles();
+
+   if (!stepId) {
+      return (<div />);
+   }
+
+   const events = jobDetails.eventsByStep(stepId);
+   const step = jobDetails.stepById(stepId);
+
+   if (!step) {
+      return (<div />);
+   }
+
+   let errors = events.filter(e => e.severity === EventSeverity.Error).sort((a, b) => {
+      return a.lineIndex < b.lineIndex ? -1 : 1;
+   });
+
+   if (count) {
+      errors = errors.slice(0, count);
+   }
+
+   let warnings = events.filter(e => e.severity === EventSeverity.Warning).sort((a, b) => {
+      return a.lineIndex < b.lineIndex ? -1 : 1;
+   });
+
+   if (count) {
+      warnings = warnings.slice(0, count);
+   }
+
+   if (showErrors && !errors.length) {
+      return <div />;
+   }
+
+   if (!showErrors && !warnings.length) {
+      return <div />;
+   }
+
+   const onRenderCell = (item?: EventData, index?: number, isScrolling?: boolean): JSX.Element => {
+
+      if (!item) {
+         return <div>???</div>;
+      }
+
+      const url = `/log/${item.logId}?lineIndex=${item.lineIndex + 1}`;
+
+      const lines = item.lines.filter(line => line.message?.trim().length).map(line => <Stack styles={{ root: { paddingLeft: 8, paddingRight: 8, lineBreak: "anywhere", whiteSpace: "pre-wrap", lineHeight: 18, fontSize: 10, fontFamily: "Horde Cousine Regular, monospace, monospace" } }}> <Link className="log-link" to={url}>{renderLine(navigate, line, undefined, {})}</Link></Stack>);
+
+      return (<Stack className={styles.itemCell} styles={{ root: { padding: 8, marginRight: 8 } }}><Stack className={item.severity === EventSeverity.Warning ? styles.gutterWarning : styles.gutter} styles={{ root: { padding: 0, margin: 0 } }}>
+         <Stack styles={{ root: { paddingLeft: 14 } }}>
+            {!!lines.length && lines}
+            {!lines.length && <Text>{step.state === JobStepState.Running ? "Generating Log Data" : "Missing Log Data"}</Text>}
+         </Stack>
+      </Stack>
+      </Stack>
+      );
+   };
+
+   return (<Stack>
+      <Stack tokens={{ padding: 4 }}>
+         {showErrors && errors.length && <Stack>
+            <div style={{ overflowY: 'auto', overflowX: 'visible', maxHeight: "400px" }} data-is-scrollable={true}>
+               <List
+                  items={errors}
+                  data-is-focusable={false}
+                  onRenderCell={onRenderCell} />
+               <div style={{ padding: 8 }} />
+            </div>
+         </Stack>}
+         {!showErrors && warnings.length && <Stack>
+            <div style={{ overflowY: 'auto', overflowX: 'visible', maxHeight: "400px" }} data-is-scrollable={true}>
+               <List
+                  items={warnings}
+                  data-is-focusable={false}
+                  onRenderCell={onRenderCell} />
+               <div style={{ padding: 8 }} />
+            </div>
+
+         </Stack>}
+      </Stack>
+   </Stack>);
+};
+
+export const JobEventListPanel: React.FC<{ jobDetails: JobDetails, stepIds: string[], eventHandler: JobEventHandler }> = ({ jobDetails, stepIds, eventHandler }) => {
+
+   const navigate = useNavigate();
+   const styles = getStyles();
+
+   if (!stepIds.length) {
+      return null;
+   }
+
+   type EventItem = {
+
+      stepId: string;
+      event?: EventData;
+
+   }
+
+   if (!stepIds.length) {
+      return null;
+   }
+
+   const items: EventItem[] = [];
+
+   stepIds.forEach(stepId => {
+
+      items.push({
+         stepId: stepId,
+      });
+
+      eventHandler.stepEvents.get(stepId)?.events?.sort((a, b) => a.lineIndex - b.lineIndex).forEach(e => {
+
+         items.push({
+            stepId: stepId,
+            event: e
+         });
+      })
+   })
+
+   const job = jobDetails.jobdata!;
+
+   const renderEventRow = (item: EventItem, index?: number, column?: IColumn) => {
+
+      if (!item) {
+         return <div>???</div>;
+      }
+
+      const step = jobDetails.stepById(item.stepId)!;
+      const errors = eventHandler.stepEvents.get(item.stepId)?.events?.filter(e => e.severity === EventSeverity.Error) ?? [];
+      const warnings = eventHandler.stepEvents.get(item.stepId)?.events?.filter(e => e.severity === EventSeverity.Warning) ?? [];
+
+      if (!errors.length && !warnings.length) {
+         return null;
+      }
+
+      const logId = step.logId!;
+      const event = item.event!;
+
+      const url = `/log/${logId}?lineIndex=${event.lineIndex + 1}`;
+
+      const lines = event.lines.map(line => <Stack styles={{ root: { paddingLeft: 8, paddingRight: 8, lineBreak: "anywhere", whiteSpace: "pre-wrap", lineHeight: 18, fontSize: 10, fontFamily: "Horde Cousine Regular, monospace, monospace" } }}> {renderLine(navigate, line, undefined, {})}</Stack>);
+      return (
+         <Link className="log-link" to={url}>
+            <Stack className={styles.itemCell} styles={{ root: { padding: 8, paddingLeft: 24, marginRight: 8 } }}><Stack className={event.severity === EventSeverity.Warning ? styles.gutterWarning : styles.gutter} styles={{ root: { padding: 0, margin: 0 } }}>
+               <Stack styles={{ root: { paddingLeft: 8 } }}>
+                  {!!lines.length && lines}
+                  {!lines.length && <Text>{step.state === JobStepState.Running ? "Generating Log Data" : "Missing Log Data"}</Text>}
+               </Stack>
+            </Stack>
+            </Stack>
+         </Link>);
+   };
+
+   const onRenderCell = (item?: EventItem, index?: number, isScrolling?: boolean): JSX.Element | null => {
+
+      if (!item) {
+         return null;
+      }
+
+      if (item.event) {
+         return renderEventRow(item);
+      }
+
+      const url = `/job/${job.id}?step=${item.stepId}`;
+
+      const stepId = item.stepId;
+
+      const step = jobDetails.stepById(stepId)!;
+
+      return <Stack styles={{ root: { background: 'rgb(233, 232, 231)', marginRight: 8, height: 42, fontSize: 12, selectors: { "a, a:hover, a:visited": { color: "#FFFFFF" }, ":hover": { background: 'rgb(223, 222, 221)' } } } }} verticalFill={true} verticalAlign="center" >
+         <Link to={url}><Stack style={{ paddingLeft: 12 }} horizontal><StepStatusIcon step={step} /><Text>{step?.name}</Text>
+         </Stack>
+         </Link>
+      </Stack>
+   }
+
+   return (<Stack style={{ maxHeight: 460 }}>
+      <div style={{ overflowY: 'auto', overflowX: 'auto' }} data-is-scrollable={true}>
+         <List
+            items={items}
+            data-is-focusable={false}
+            onRenderCell={onRenderCell} />
+      </div>
+   </Stack>);
+};
+
+
+export const JobEventPanel: React.FC<{ jobDetails: JobDetails, label?: GetLabelStateResponse, eventHandler: JobEventHandler }> = observer(({ jobDetails, label, eventHandler }) => {
+
+   const { hordeClasses } = getHordeStyling();
+
+   eventHandler.set(jobDetails);
+
+   // subscribe, this is a little ugly, should look into mobx pattern for this
+   let updated = jobDetails.updated;
+   updated = eventHandler.updated;
+   // silence warning
+   if (updated) { }
+
+   let stepIds = jobDetails.getSteps().map(s => s.id);
+
+   if (label) {
+      stepIds = stepIds.filter(stepId => label.steps.indexOf(stepId) !== -1);
+   }
+
+   stepIds = stepIds.filter(s => {
+      return !!eventHandler.stepEvents.get(s)?.events?.length;
+   })
+
+   if (!stepIds.length) {
+      return null;
+   }
+
+   let hasErrors = false;
+   let hasWarnings = false;
+
+   eventHandler.stepEvents.forEach(e => {
+
+      if (stepIds.indexOf(e.stepId) === -1) {
+         return;
+      }
+
+      if (e.events?.find(event => event.severity === EventSeverity.Error)) hasErrors = true;
+      if (e.events?.find(event => event.severity === EventSeverity.Warning)) hasWarnings = true;
+
+   });
+
+   let title = hasErrors ? "Errors" : "";
+   if (hasWarnings) {
+      if (title) {
+         title += " & Warnings";
+      } else {
+         title = "Warnings";
+      }
+   }
+
+   return (<Stack styles={{ root: { paddingTop: 18, paddingRight: 12 } }}>
+      <Stack className={hordeClasses.raised}>
+         <Stack tokens={{ childrenGap: 12 }}>
+            <Text variant="mediumPlus" styles={{ root: { fontFamily: "Horde Open Sans SemiBold" } }}>{title}</Text>
+            <Stack styles={{ root: { paddingLeft: 4, paddingRight: 0, paddingTop: 8, paddingBottom: 4 } }}>
+               <JobEventListPanel jobDetails={jobDetails} stepIds={stepIds} eventHandler={eventHandler} />
+            </Stack>
+         </Stack>
+      </Stack>
+   </Stack>);
+
+});
+
+
diff --git a/Engine/Source/Programs/Horde/HordeDashboard/src/components/error/ErrorStore.ts b/Engine/Source/Programs/Horde/HordeDashboard/src/components/error/ErrorStore.ts
new file mode 100644
--- /dev/null
+++ b/Engine/Source/Programs/Horde/HordeDashboard/src/components/error/ErrorStore.ts
@@ -0,0 +1,272 @@
+import { getSiteConfig } from "../../backend/Config";
+import { action, makeObservable, observable, override, runInAction } from "mobx";
+
+// Explicit error message
+export type ErrorFormatted = {
+   time: Date | string;
+   level: "Error" | "Warning";
+   message: string;
+   id: number;
+   format: string;
+   properties: Record<string, string | number | boolean>
+}
+
+export type ErrorInfo = {
+
+   // general project rejection
+   reason?: any;
+
+   // fetch response object
+   response?: Response;
+
+   // request url (should match url in fetch response if any)
+   url?: string;
+
+   // REST mode
+   mode?: "GET" | "PUT" | "DELETE" | "POST" | "PATCH" | string;
+
+   // custom title override
+   title?: string;
+
+   // custom message override
+   message?: string;
+
+   format?: ErrorFormatted;
+
+}
+
+export enum ErrorReportType {
+   Dialog = "Dialog",
+   Bar = "Bar"
+}
+
+type RetryInfo = {
+   isRetrying: boolean,
+   succeeded?: boolean,
+   failed?: boolean,
+   hasResult: boolean
+}
+
+class ErrorStore {
+
+   constructor() {
+      makeObservable(this);
+   }
+
+   @observable
+   error?: ErrorInfo;
+
+   // Retry logic
+   @observable
+   retryInfo: RetryInfo = {
+      isRetrying: false,
+      succeeded: undefined,
+      failed: undefined,
+      hasResult: false,
+   };
+   
+   retryCallback?: () => boolean | Promise<boolean>;
+
+   @action
+   set(
+      infoIn: ErrorInfo, 
+      update?: boolean,
+      retryCallback?: () => boolean | Promise<boolean>
+   ): void {
+      // this.print(infoIn);
+
+      const config = getSiteConfig();
+
+      if (!update) {// && !this.unauthorized(infoIn)) {
+         if (config.environment === "production") {
+            return;
+         }
+      }
+
+      if (infoIn.response?.status === 502) {
+         return;
+      }
+
+      const hash = this.hash(infoIn);
+      if (hash && this.filter.has(hash)) {
+         return;
+      }
+
+      if (this.error && !update) {
+         return;
+      }
+
+      if(retryCallback) this.retryCallback = retryCallback;
+
+      const info = { ...infoIn };
+      this.error = info;
+   }
+
+   @action.bound
+   async retry() {
+      if(!this.retryCallback) return;
+
+      this.retryInfo.isRetrying = true;
+      this.retryInfo.succeeded = false;
+      this.retryInfo.failed = false;
+      this.retryInfo.hasResult = false;
+
+      try {
+         // Callbacks may be synchronous or asynchronous.
+         const result = this.retryCallback();
+         const success = result instanceof Promise ? await result : result;
+
+         // In the case of async callbacks, this kicks the rest of the function out of
+         // being a MobX action. Thus, further observable mutations need to be in a runInAction
+         runInAction(() => {
+            if(success) {
+               this.retryInfo.succeeded = true;
+               // this.clear();
+            } else {
+               this.retryInfo.failed = true;
+            }
+         });
+      } catch (e) {
+         runInAction(() => this.retryInfo.failed = true)
+      } finally {
+         runInAction(() => {
+               this.retryInfo.hasResult = true;
+               // Let the user try again
+               if(this.retryInfo.failed) setTimeout(this.clearRetry, 1250);
+            }
+         )
+      }
+   }
+
+   get canRetry(): boolean {
+      return !!this.retryCallback;
+   }
+
+   get message(): string {
+
+      const error = this.error;
+
+      if (!error) {
+         return "";
+      }
+
+      if (error.format?.message) {
+         return error.format?.message;
+      }
+
+      if (error.message) {
+         return error.message;
+      }
+
+      if (error.reason) {
+         return error.reason.toString();
+      }
+
+      if (error.response) {
+         return `${error.response.status}: ${error.response.statusText}`;
+      }
+
+      return "";
+
+   }
+
+   unauthorized(infoIn?: ErrorInfo): boolean {
+
+      let error = infoIn;
+
+      if (!error) {
+         error = this.error;
+      }
+
+      if (!error) {
+         return false;
+      }
+
+      if (!error.response) {
+         return false;
+      }
+
+      if (error.response.status === 401 || error.response.status === 403 || error.response.url.toLowerCase().indexOf("accessdenied") !== -1) {
+         return true;
+      }
+
+      return false;
+
+   }
+
+   get reject(): string {
+
+      return this.message;
+
+   }
+
+   hash(error: ErrorInfo): string {
+
+      const url = error.url ?? error?.response?.url;
+
+      // filter on end point, without query strings
+      if (url) {
+         return url.split(/[?#]/)[0];
+      }
+
+      // filter on reason
+      if (error.reason) {
+         return error.reason;
+      }
+
+      return "";
+   }
+
+   print(info: ErrorInfo) {
+
+      // this will also go out to datadog if configured
+      const message = `${info.response ? JSON.stringify(info.response) : ""} ${info.reason ? info.reason : ""}`;
+
+      if (!message) {
+         return;
+      }
+      console.error(message);
+   }
+
+   @action.bound
+   clear(): void {
+
+      if (!this.error) {
+         return;
+      }
+
+      if (this.filterError) {
+         const hash = this.hash(this.error);
+         if (hash) {
+            this.filter.add(hash);
+         }
+      }
+
+      this.filterError = false;
+      
+      this.error = undefined;
+
+      this.clearRetry(true);
+   }
+
+   // Can clear retry entirely or just for subsequent attempts
+   @action.bound
+   clearRetry(clearFull?: boolean): void {
+      this.retryInfo.succeeded = undefined;
+      this.retryInfo.failed = undefined;
+      this.retryInfo.isRetrying = false;
+      
+      if(clearFull) {
+         this.retryCallback = undefined;
+         this.retryInfo.hasResult = false;
+      }
+   }
+
+   filter = new Set<string>();
+
+   filterError = false;
+
+}
+
+export const errorDialogStore = new ErrorStore();
+export const errorBarStore = new ErrorStore();
\ No newline at end of file
diff --git a/Engine/Source/Programs/Horde/HordeDashboard/src/components/jobDetailsV2/JobDetailPreflight.tsx b/Engine/Source/Programs/Horde/HordeDashboard/src/components/jobDetailsV2/JobDetailPreflight.tsx
--- a/Engine/Source/Programs/Horde/HordeDashboard/src/components/jobDetailsV2/JobDetailPreflight.tsx
+++ b/Engine/Source/Programs/Horde/HordeDashboard/src/components/jobDetailsV2/JobDetailPreflight.tsx
@@ -7,7 +7,7 @@ import backend from '../../backend';
 import { JobState } from '../../backend/Api';
 import dashboard from '../../backend/Dashboard';
 import { ISideRailLink } from '../../base/components/SideRail';
-import { ErrorHandler } from '../ErrorHandler';
+import { errorDialogStore } from '../error/ErrorStore';
 import { JobDataView, JobDetailsV2 } from "./JobDetailsViewCommon";
 import { getHordeStyling } from '../../styles/Styles';
 
@@ -86,7 +86,7 @@ const AutosubmitInfo: React.FC<{ jobDetails: JobDetailsV2 }> = observer(({ jobDe
                         setState({ autoSubmit: value, inflight: false })
                      }).catch(reason => {
 
-                        ErrorHandler.set({
+                        errorDialogStore.set({
                            reason: reason,
                            title: `Error Setting Auto-submit`,
                            message: `There was an error setting job to autosubmit, reason: "${reason}"`
diff --git a/Engine/Source/Programs/Horde/HordeDashboard/src/components/jobDetailsV2/StepRetryModal.tsx b/Engine/Source/Programs/Horde/HordeDashboard/src/components/jobDetailsV2/StepRetryModal.tsx
--- a/Engine/Source/Programs/Horde/HordeDashboard/src/components/jobDetailsV2/StepRetryModal.tsx
+++ b/Engine/Source/Programs/Horde/HordeDashboard/src/components/jobDetailsV2/StepRetryModal.tsx
@@ -6,7 +6,7 @@ import { CreateJobRequest } from '../../backend/Api';
 import { Link, useNavigate } from 'react-router-dom';
 import { JobDetailsV2 } from './JobDetailsViewCommon';
 import { getHordeStyling } from '../../styles/Styles';
-import ErrorHandler from '../ErrorHandler';
+import { errorDialogStore } from '../error/ErrorStore';
 import moment from 'moment';
 
 export enum StepRetryType {
@@ -294,7 +294,7 @@ export const StepRetryModal: React.FC<{ stepId: string; jobDetails: JobDetailsV2
          navigate(`/job/${data.id}`);
       }).catch(reason => {
 
-         ErrorHandler.set({
+         errorDialogStore.set({
 
             reason: `${reason}`,
             title: `Error Creating Job`,
diff --git a/Engine/Source/Programs/Horde/HordeDashboard/src/components/preflights/PreflightConfigCheckRedirector.tsx b/Engine/Source/Programs/Horde/HordeDashboard/src/components/preflights/PreflightConfigCheckRedirector.tsx
--- a/Engine/Source/Programs/Horde/HordeDashboard/src/components/preflights/PreflightConfigCheckRedirector.tsx
+++ b/Engine/Source/Programs/Horde/HordeDashboard/src/components/preflights/PreflightConfigCheckRedirector.tsx
@@ -3,7 +3,7 @@
 import React from 'react';
 import { Navigate } from 'react-router-dom';
 import { useQuery } from '../JobDetailCommon';
-import ErrorHandler from '../ErrorHandler';
+import { errorDialogStore } from '../error/ErrorStore';
 
 function setError(message: string) {
 
@@ -11,7 +11,7 @@ function setError(message: string) {
 
    message = `${message}\n\nReferring URL: ${encodeURI(window.location.href)}`
 
-   ErrorHandler.set({
+   errorDialogStore.set({
 
       title: `Error handling Preflight Configuration`,
       message: message
diff --git a/Engine/Source/Programs/Horde/HordeDashboard/src/components/tools/ToolManagement.tsx b/Engine/Source/Programs/Horde/HordeDashboard/src/components/tools/ToolManagement.tsx
--- a/Engine/Source/Programs/Horde/HordeDashboard/src/components/tools/ToolManagement.tsx
+++ b/Engine/Source/Programs/Horde/HordeDashboard/src/components/tools/ToolManagement.tsx
@@ -9,7 +9,7 @@ import { GetToolDeploymentResponse, GetToolResponse, GetToolSummaryResponse, Too
 import dashboard from "../../backend/Dashboard";
 import { PollBase } from "../../backend/PollBase";
 import { getHordeStyling } from "../../styles/Styles";
-import ErrorHandler from "../ErrorHandler";
+import { errorDialogStore } from "../error/ErrorStore";
 
 class ToolHandler extends PollBase {
 
@@ -90,7 +90,7 @@ const ActionConfirmation: React.FC<{ handler: ToolHandler, request: UpdateDeploy
 
          if (errorReason) {
 
-            ErrorHandler.set({
+            errorDialogStore.set({
 
                reason: `${errorReason}`,
                title: `Error Updating Deployment`,
