Horde Dashboard: Improved modal warning for downloading artifacts

diff --git a/Engine/Source/Programs/Horde/HordeDashboard/src/components/artifacts/ArtifactButton.tsx b/Engine/Source/Programs/Horde/HordeDashboard/src/components/artifacts/ArtifactButton.tsx
--- a/Engine/Source/Programs/Horde/HordeDashboard/src/components/artifacts/ArtifactButton.tsx
+++ b/Engine/Source/Programs/Horde/HordeDashboard/src/components/artifacts/ArtifactButton.tsx
@@ -241,21 +241,22 @@ const ZipDownloadWarning: React.FC<{ onContinue: () => void, onClose: () => void
 
    const { hordeClasses } = getHordeStyling();
 
-   return <Modal className={hordeClasses.modal} isOpen={true} isBlocking={true} topOffsetFixed={true} styles={{ main: { padding: 8, height: 200, width: 600, hasBeenOpened: false, top: 0, bottom: 150, left: 0, right: 0, margin: "auto", position: "absolute" } }} >
+   return <Modal className={hordeClasses.modal} isOpen={true} isBlocking={true} topOffsetFixed={true} styles={{ main: { padding: 8, height: "fit-content", width: 600, hasBeenOpened: false, top: 0, bottom: 150, left: 0, right: 0, margin: "auto", position: "absolute" } }} >
 
       <Stack style={{ padding: 8 }}>
          <Stack style={{ paddingBottom: 24 }}>
             <Text variant="mediumPlus" style={{ fontFamily: "Horde Open Sans SemiBold" }}>Warning: Slow Download</Text>
          </Stack>
-         <Stack style={{ paddingBottom: 38 }}>
+         <Stack style={{ paddingBottom: 32 }} tokens={{childrenGap: 16}}>
             <Text variant="medium" style={{ fontFamily: "Horde Open Sans SemiBold" }}>Downloading artifacts in zip format is significantly slower than downloading with UGS.</Text>
+            <Text variant="medium" style={{ fontFamily: "Horde Open Sans SemiBold" }}>Click the download button dropdown and select "Download with UGS." Open the downloaded uartifact file with UnrealGameSync for faster download speeds.</Text>
          </Stack>
          <Stack>
             <Checkbox label="Don't show this warning again" onChange={(ev?: any, checked?: boolean | undefined) => { localStorage.setItem("horde_hide_zip_warning", checked!.toString()) }} />
          </Stack>
          <Stack horizontal style={{ paddingTop: 16 }}>
             <Stack grow />
-            <Stack horizontal tokens={{ childrenGap: 20 }}>
+            <Stack horizontal tokens={{ childrenGap: 16 }}>
                <PrimaryButton onClick={() => onContinue()} text="Continue" />
                <DefaultButton onClick={() => onClose()} text="Cancel" />
             </Stack>
