Horde: Updating release notes
#skipci

diff --git a/Engine/Source/Programs/Horde/Docs/ReleaseNotes.md b/Engine/Source/Programs/Horde/Docs/ReleaseNotes.md
--- a/Engine/Source/Programs/Horde/Docs/ReleaseNotes.md
+++ b/Engine/Source/Programs/Horde/Docs/ReleaseNotes.md
@@ -1,5 +1,50 @@
 # Release Notes
 
+## 2025-06-02
+
+* Cloud DDC - Added cross bucket search endpoint (43145204)
+* Horde Dashboard: Improved modal warning for downloading artifacts (43112786)
+* Ensure clang and clang++ errors are matched (43108898)
+* Horde Dashboard: Error Retry Callback (43108710)
+* Use Razor template for /account (43107401)
+* Add support for allocating a UBA cache server session for compute resource requests (43103819)
+* Horde :  Fix ScheduledDowntime GetNext to return next or active downtime
+  If the downtime is active return the start and finish times for the active schedule.
+  If the downtime is not active return the start and finish times for the next day or next week schedule.
+  Downtime Service Tick ignores next times that are in the past, which is for downtime that are non-repeating. (43099696)
+* Fixed an exception that was occuring in the LocalizationIssueHandler (43075154)
+* [UBA] Fixed C# static analyzer error in SchedulerImpl.cs (43053139)
+* [UBA] Changed UbaExecutor to use Uba's scheduler. Reason is to be able to implement shadow processes and also dynamically add more processes when things like llvm pgo wants to distribute internal processes.  Moved logging out of ImmediateActionQueue and into base class ActionLogger. This made it possible to reuse logging logic for UbaExecutor (42998941)
+* Horde Dashboard: Fixed HistoryModal.tsx hot reload issue (42743008)
+* Horde Dashboard: Alt-text for project logos (42649235)
+* Horde Dashboard: Click-to-dimiss for modals (42648635)
+* Fix issue with agent modal not refreshing when modified, tweak disable comment attribution (42641626)
+* Horde Dashboard: Require agent disable reason and auto postpend username (42603556)
+* Horde Dashboard: Improve job template selection when on an exiting job (42596903)
+* Horde Server: Don't require a user from user collection when creating devices, for service accounts (42593578)
+* Display warning to surface rendering issues on large log files
+* Disable trace level logging of agent channel data (42527178)
+* Prevent per-cluster Perforce errors from aborting changelist replication loop.  With an invalid Perforce cluster or no available servers for that cluster, such an exception would bubble all the way up to the tick method. This lead to the replication loop restarting. (42483815)
+* Update Remote Shader Compilation Horde docs with the new ini configuration. (42451263)
+* Do not include steps with warnings in failed steps for job retries (42430419)
+* Horde Dashboard: allow option for UserSelect to autofocus (42397046)
+* Horde Dashboard: Fix autosubmit link when swarm url isn't defined (42384179)
+* Fix publishing step for Horde installer considerations (42375246)
+* Hprde Job Driver: Downgrade TempStorage retry error to warning, so it does not fail step (42342223)
+* Horde : Add missing properties to agent.  Add RAM size for all platforms.  Add core count for ARM based Macs (42333756)
+* Horde Storage: Downgrade a file missing exception to a warning as will retry and should not fail step (42299162)
+* add support for engine version comparisons as preprocessor conditions (42297433)
+* Make it so non-admin users can see, though not edit notices (42288524)
+* Fix issue with default incremental pcb artifact (42249038)
+* Ensure agent registration token is used in dedicated mode (42176590)
+* Horde Server: Marshal metadata job step id by string (42102607)
+* Horde Server: Adding job and step meta data tagging (42099956)
+* Horde Dashboard: Add support for multiple plugin mount points (42096956)
+* Add support for multiple report summaries (42058201)
+* Horde : The Notices API Get method will return the start and finish times when values are set (42048584)
+* Silence stderr from 'security' calls on Mac as these interfere with OidcToken --ResultToConsole. (42034179)
+* Add KnownLogEvents.Horde_BuildHealth_Ignore to suppress build health issue generation regardless of log level (41916679)
+
 ## 2025-04-24
 
 * Pass telemetry store id to Epic Telemetry backend. (41897043)
