Enable Test automation hub v2

#rnx

diff --git a/Engine/Source/Programs/Horde/HordeDashboard/plugins/test-automation/index.tsx b/Engine/Source/Programs/Horde/HordeDashboard/plugins/test-automation/index.tsx
--- a/Engine/Source/Programs/Horde/HordeDashboard/plugins/test-automation/index.tsx
+++ b/Engine/Source/Programs/Horde/HordeDashboard/plugins/test-automation/index.tsx
@@ -2,7 +2,7 @@ import { registerHordePlugin, MountType } from "..";
 import { TestAutomationView } from "./testAutomationView";
 
 registerHordePlugin({
-   id: "test-automation",
+   id: "build",
    routes: [{ path: "test-automation", element: <TestAutomationView /> }],
    mount: {
       type: MountType.TopNav,
