Horde: Use preferred colors for device status
Slightly increased brightness threshold

diff --git a/Engine/Source/Programs/Horde/HordeDashboard/src/base/utilities/colors.ts b/Engine/Source/Programs/Horde/HordeDashboard/src/base/utilities/colors.ts
--- a/Engine/Source/Programs/Horde/HordeDashboard/src/base/utilities/colors.ts
+++ b/Engine/Source/Programs/Horde/HordeDashboard/src/base/utilities/colors.ts
@@ -20,5 +20,5 @@ export const isBright = (hex: string): boolean => {
 
     let luminosity: number = ((r * 0.2126) + (g * 0.7152) + (b * 0.0722)) / 255
 
-    return luminosity > 0.6
+    return luminosity > 0.65
 }
diff --git a/Engine/Source/Programs/Horde/HordeDashboard/src/components/devices/DeviceEditor.tsx b/Engine/Source/Programs/Horde/HordeDashboard/src/components/devices/DeviceEditor.tsx
--- a/Engine/Source/Programs/Horde/HordeDashboard/src/components/devices/DeviceEditor.tsx
+++ b/Engine/Source/Programs/Horde/HordeDashboard/src/components/devices/DeviceEditor.tsx
@@ -1,9 +1,11 @@
 import { SpinnerSize, DialogType, IDropdownOption, MessageBarType, Text, Checkbox, DefaultButton, Dialog, DialogFooter, Dropdown, MessageBar, Modal, PrimaryButton, Spinner, Stack, TextField } from "@fluentui/react";
 import React, { useState } from "react";
 import backend from "../../backend";
-import { GetDeviceResponse } from "../../backend/Api";
+import { DashboardPreference, GetDeviceResponse } from "../../backend/Api";
 import { getHordeStyling } from "../../styles/Styles";
 import { DeviceHandler, DeviceStatus } from "../../backend/DeviceHandler";
+import dashboard from "horde/backend/Dashboard";
+import { isBright } from "horde/base/utilities/colors";
 
 type DeviceEditData = {
    // if defined, existing device
@@ -275,6 +277,9 @@ export const DeviceEditor: React.FC<{ handler: DeviceHandler, deviceIn?: GetDevi
 
    }
 
+   const deleteButtonColor = dashboard.getPreference(DashboardPreference.ColorError) ? dashboard.getPreference(DashboardPreference.ColorError) : "#FF0000";
+   const deleteButtonTextColor = isBright(deleteButtonColor!) ? "#2D3F5F" : "#FFFFFF";
+
    return <Modal className={hordeClasses.modal} isOpen={true} styles={{ main: { padding: 8, width: 700 } }} onDismiss={() => onClose()}>
       {!!error && <MessageBar
          messageBarType={MessageBarType.error}
@@ -367,7 +372,7 @@ export const DeviceEditor: React.FC<{ handler: DeviceHandler, deviceIn?: GetDevi
          </Stack>
 
          <Stack horizontal style={{ paddingTop: 64 }}>
-            {!!existing && !editNote && !reservation && <PrimaryButton style={{ backgroundColor: "red", border: 0 }} onClick={() => setConfirmDelete(true)} text="Delete Device" />}
+            {!!existing && !editNote && !reservation && <PrimaryButton style={{ color: deleteButtonTextColor, backgroundColor: deleteButtonColor, border: 0 }} onClick={() => setConfirmDelete(true)} text="Delete Device" />}
             <Stack grow />
             <Stack horizontal tokens={{ childrenGap: 28 }}>
                <PrimaryButton onClick={() => onSave()} text="Save" />
diff --git a/Engine/Source/Programs/Horde/HordeDashboard/src/components/devices/DeviceView.tsx b/Engine/Source/Programs/Horde/HordeDashboard/src/components/devices/DeviceView.tsx
--- a/Engine/Source/Programs/Horde/HordeDashboard/src/components/devices/DeviceView.tsx
+++ b/Engine/Source/Programs/Horde/HordeDashboard/src/components/devices/DeviceView.tsx
@@ -5,7 +5,7 @@ import moment from "moment";
 import React, { useEffect, useState } from "react";
 import { Link, useNavigate, useSearchParams } from "react-router-dom";
 import backend from "../../backend";
-import { DevicePoolType, GetDeviceResponse } from "../../backend/Api";
+import { DashboardPreference, DevicePoolType, GetDeviceResponse } from "../../backend/Api";
 import dashboard from "../../backend/Dashboard";
 import { useWindowSize } from "../../base/utilities/hooks";
 import { getNiceTime } from "../../base/utilities/timeUtils";
@@ -18,6 +18,7 @@ import { DevicePoolTelemetryModal } from "./DevicePoolTelemetry";
 import { TopNav } from "../TopNav";
 import { getHordeStyling } from "../../styles/Styles";
 import { getHordeTheme } from "../../styles/theme";
+import { isBright } from "horde/base/utilities/colors";
 
 const handler = new DeviceHandler();
 
@@ -310,12 +311,18 @@ const DevicePanel: React.FC = observer(() => {
       </Stack>
    }
 
+
+   const warningColorPref = dashboard.getPreference(DashboardPreference.ColorWarning);
+   const errorColorPref = dashboard.getPreference(DashboardPreference.ColorError);
+   const successColorPref = dashboard.getPreference(DashboardPreference.ColorSuccess);
+   const runningColorPref = dashboard.getPreference(DashboardPreference.ColorRunning);
+
    const StatusColors = new Map<DeviceStatus, string>([
-      [DeviceStatus.Available, "#52C705"],
-      [DeviceStatus.Problem, "#DE4522"],
-      [DeviceStatus.Reserved, "#00BCF2"],
+      [DeviceStatus.Available, successColorPref ? successColorPref : "#52C705"],
+      [DeviceStatus.Problem, errorColorPref ? errorColorPref : "#DE4522"],
+      [DeviceStatus.Reserved, warningColorPref ? warningColorPref : "#00BCF2"],
       [DeviceStatus.Disabled, dashboard.darktheme ? "#D3D2D1" : "#F3F2F1"],
-      [DeviceStatus.Maintenance, "#0078D4"]
+      [DeviceStatus.Maintenance, runningColorPref ? runningColorPref : "#0078D4"]
    ]);
 
 
@@ -379,7 +386,6 @@ const DevicePanel: React.FC = observer(() => {
       const device = item.device;
       const status = handler.getDeviceStatus(device);
 
-      let color = status === DeviceStatus.Disabled ? "#777777 !important" : "#ffffff";
       let backgroundColor = StatusColors.get(status)!
       let text = DeviceStatusLabels.get(status)!;
 
@@ -390,13 +396,18 @@ const DevicePanel: React.FC = observer(() => {
          }
       }
 
+      let color = status === DeviceStatus.Disabled ? "#777777 !important" : isBright(backgroundColor) ? "#2D3F5F" : "#FFFFFF";
+      const hoveredPressedStyle = { backgroundColor: backgroundColor, border: '0px', filter: isBright(backgroundColor) ? "brightness(0.9)" : "brightness(1.2)" }
+
       if (item.device.cleanStartTime && !item.device.cleanFinishTime) {
          text = "Cleaning";
       }
 
       // @todo: support deleting reservations
       return <Stack verticalFill={true} verticalAlign="center" horizontal>
-         <PrimaryButton styles={{ root: { border: "0px", width: 84, height: 24, fontSize: 12, fontFamily: "Horde Open Sans SemiBold !important", selectors: { ".ms-Button-label": { color: color } }, backgroundColor: backgroundColor } }} text={text}
+         <PrimaryButton 
+            styles={{ root: { border: "0px", width: 84, height: 24, fontSize: 12, fontFamily: "Horde Open Sans SemiBold !important", selectors: { ".ms-Button-label": { color: color } }, backgroundColor: backgroundColor }, rootHovered: hoveredPressedStyle, rootPressed: hoveredPressedStyle }}
+            text={text}
             onClick={(ev) => {
                ev.stopPropagation(); ev.preventDefault();
                setEditState({ shown: true, device: item })
@@ -485,7 +496,7 @@ const DevicePanel: React.FC = observer(() => {
 
          if (user.id === dashboard.userId) {
             return <Stack verticalFill={true} verticalAlign="center">
-               <PrimaryButton styles={{ root: { border: "0px", width: 90, height: 24, fontSize: 12, fontFamily: "Horde Open Sans SemiBold !important", backgroundColor: StatusColors.get(DeviceStatus.Reserved)! } }} text="Check In"
+               <PrimaryButton styles={{ root: { border: "0px", width: 90, height: 24, fontSize: 12, fontFamily: "Horde Open Sans SemiBold !important" } }} text="Check In"
                   onClick={async (ev) => {
                      ev.stopPropagation(); ev.preventDefault();
                      setCheckoutState({ checkinId: device.id });
