Horde: Add a "No Suspects" emoji indication to build health slack threads that lack suspects

diff --git a/Engine/Source/Programs/Horde/Plugins/Build/HordeServer.Build/Notifications/Sinks/SlackNotificationSink.cs b/Engine/Source/Programs/Horde/Plugins/Build/HordeServer.Build/Notifications/Sinks/SlackNotificationSink.cs
--- a/Engine/Source/Programs/Horde/Plugins/Build/HordeServer.Build/Notifications/Sinks/SlackNotificationSink.cs
+++ b/Engine/Source/Programs/Horde/Plugins/Build/HordeServer.Build/Notifications/Sinks/SlackNotificationSink.cs
@@ -1051,6 +1051,7 @@ namespace HordeServer.Notifications.Sinks
 			FixFailed = 4,
 			Resolved = 8,
 			AssignedButNotAcknowledged = 16,
+			NoSuspects = 32,
 		}
 
 		async Task CreateOrUpdateWorkflowThreadAsync(string triageChannel, IIssue issue, IIssueSpan span, IReadOnlyList<IIssueSpan> spans, WorkflowConfig workflow, CancellationToken cancellationToken)
@@ -1401,7 +1402,7 @@ namespace HordeServer.Notifications.Sinks
 				}
 			}
 
-			await UpdateReactionsAsync(triageChannel, issue, fixFailedSpan, eventId, threadId, cancellationToken);
+			await UpdateReactionsAsync(triageChannel, issue, fixFailedSpan, suspects, eventId, threadId, cancellationToken);
 
 			if (issue.ExternalIssueKey != null)
 			{
@@ -1411,7 +1412,7 @@ namespace HordeServer.Notifications.Sinks
 			}
 		}
 
-		private async Task UpdateReactionsAsync(string triageChannel, IIssue issue, IIssueSpan? fixFailedSpan, string eventId, SlackMessageId threadId, CancellationToken cancellationToken)
+		private async Task UpdateReactionsAsync(string triageChannel, IIssue issue, IIssueSpan? fixFailedSpan, IReadOnlyList<IIssueSuspect> suspects, string eventId, SlackMessageId threadId, CancellationToken cancellationToken)
 		{
 			ReactionFlags reactions = ReactionFlags.None;
 			if (issue.AcknowledgedAt != null)
@@ -1434,6 +1435,10 @@ namespace HordeServer.Notifications.Sinks
 			{
 				reactions |= ReactionFlags.AssignedButNotAcknowledged;
 			}
+			if (issue.OwnerId == null && suspects.Count == 0)
+			{
+				reactions |= ReactionFlags.NoSuspects;
+			}
 
 			string reactionEventId = $"{eventId}_reactions";
 			string reactionStateDigest = ((int)reactions).ToString();
@@ -1486,6 +1491,15 @@ namespace HordeServer.Notifications.Sinks
 					await _slackClient.RemoveReactionAsync(threadId, "mailbox", cancellationToken);
 				}
 
+				if ((reactions & ReactionFlags.NoSuspects) != 0)
+				{
+					await _slackClient.AddReactionAsync(threadId, "question", cancellationToken);
+				}
+				else
+				{
+					await _slackClient.RemoveReactionAsync(threadId, "question", cancellationToken);
+				}
+
 				await AddOrUpdateMessageStateAsync(triageChannel, reactionEventId, null, reactionStateDigest, threadId, cancellationToken);
 			}
 		}
