Horde: Update Secrets plugin with a new server configuration option WithAws that controls if the secret provider AwsParameterStore is to be enabled

diff --git a/Engine/Source/Programs/Horde/Docs/Config/Secrets.md b/Engine/Source/Programs/Horde/Docs/Config/Secrets.md
--- a/Engine/Source/Programs/Horde/Docs/Config/Secrets.md
+++ b/Engine/Source/Programs/Horde/Docs/Config/Secrets.md
@@ -76,9 +76,25 @@ An example secret may be configured as follows:
 
 Secrets are queried from the external provider when requested by a user, and are not cached by Horde.
 
-Secrets from external providers that require additional options, for example HashiCorp Vault to set the endpoint and
-credentials, will use a `providerConfig` property instead of `name` property to provide the data. For example to fetch
-three secrets from HashiCorp Vault using different authentication methods:
+Secrets from external providers may require additional configuration.
+
+### AWS
+
+To use AWS Parameter Store in the [`server.json`](../Deployment/ServerSettings.md) file add the following to enable AWS
+features for this provider.
+
+   ```json
+    "plugins": {
+        "secrets": {
+            "withAws": true
+        }
+    }
+   ```
+
+### HashiCorp Vault
+
+For HashiCorp Vault to set the endpoint and credentials, will use a `providerConfig` property instead of `name` property to provide the data.
+For example to fetch three secrets from HashiCorp Vault using different authentication methods:
 
    * Use an AWS ARN to pass to AWS Security Token Service AssumeRole
    * Use AWS default credential search
diff --git a/Engine/Source/Programs/Horde/HordeServer/ServerApp.cs b/Engine/Source/Programs/Horde/HordeServer/ServerApp.cs
--- a/Engine/Source/Programs/Horde/HordeServer/ServerApp.cs
+++ b/Engine/Source/Programs/Horde/HordeServer/ServerApp.cs
@@ -391,6 +391,9 @@ namespace HordeServer
 			KeyValuePair.Create("Horde:ComputeTunnelAddress", "Horde:Plugins:Compute:ComputeTunnelAddress"),
 			KeyValuePair.Create("Horde:WithAws", "Horde:Plugins:Compute:WithAws"),
 
+			// Secrets plugin
+			KeyValuePair.Create("Horde:WithAws", "Horde:Plugins:Secrets:WithAws"),
+
 			// Storage plugin
 			KeyValuePair.Create("Horde:BundleCacheDir", "Horde:Plugins:Storage:BundleCacheDir"),
 			KeyValuePair.Create("Horde:BundleCacheSize", "Horde:Plugins:Storage:BundleCacheSize"),
diff --git a/Engine/Source/Programs/Horde/Plugins/Secrets/HordeServer.Secrets/Secrets/Providers/AwsParameterStoreSecretProvider.cs b/Engine/Source/Programs/Horde/Plugins/Secrets/HordeServer.Secrets/Secrets/Providers/AwsParameterStoreSecretProvider.cs
--- a/Engine/Source/Programs/Horde/Plugins/Secrets/HordeServer.Secrets/Secrets/Providers/AwsParameterStoreSecretProvider.cs
+++ b/Engine/Source/Programs/Horde/Plugins/Secrets/HordeServer.Secrets/Secrets/Providers/AwsParameterStoreSecretProvider.cs
@@ -1,6 +1,5 @@
 // Copyright Epic Games, Inc. All Rights Reserved.
 
-using Amazon.Extensions.NETCore.Setup;
 using Amazon.SimpleSystemsManagement;
 using Amazon.SimpleSystemsManagement.Model;
 
@@ -19,10 +18,9 @@ namespace HordeServer.Secrets.Providers
 		/// <summary>
 		/// Constructor
 		/// </summary>
-		public AwsParameterStoreSecretProvider()
+		public AwsParameterStoreSecretProvider(IAmazonSimpleSystemsManagement systemsManagement)
 		{
-			AWSOptions awsOptions = new AWSOptions();
-			_systemsManagement = awsOptions.CreateServiceClient<IAmazonSimpleSystemsManagement>();
+			_systemsManagement = systemsManagement;
 		}
 
 		/// <inheritdoc/>
diff --git a/Engine/Source/Programs/Horde/Plugins/Secrets/HordeServer.Secrets/SecretsPlugin.cs b/Engine/Source/Programs/Horde/Plugins/Secrets/HordeServer.Secrets/SecretsPlugin.cs
--- a/Engine/Source/Programs/Horde/Plugins/Secrets/HordeServer.Secrets/SecretsPlugin.cs
+++ b/Engine/Source/Programs/Horde/Plugins/Secrets/HordeServer.Secrets/SecretsPlugin.cs
@@ -1,10 +1,12 @@
 // Copyright Epic Games, Inc. All Rights Reserved.
 
+using Amazon.SimpleSystemsManagement;
 using EpicGames.Horde.Secrets;
 using HordeServer.Plugins;
 using HordeServer.Secrets;
 using HordeServer.Secrets.Providers;
 using Microsoft.AspNetCore.Builder;
+using Microsoft.Extensions.Configuration;
 using Microsoft.Extensions.DependencyInjection;
 
 namespace HordeServer
@@ -12,9 +14,21 @@ namespace HordeServer
 	/// <summary>
 	/// Entry point for the secrets plugin
 	/// </summary>
-	[Plugin("Secrets", GlobalConfigType = typeof(SecretsConfig))]
+	[Plugin("Secrets", GlobalConfigType = typeof(SecretsConfig), ServerConfigType = typeof(SecretsServerConfig))]
 	public class SecretsPlugin : IPluginStartup
 	{
+		private readonly IServerInfo _serverInfo;
+		private readonly SecretsServerConfig _staticConfig;
+
+		/// <summary>
+		/// Constructor
+		/// </summary>
+		public SecretsPlugin(IServerInfo serverInfo, SecretsServerConfig staticConfig)
+		{
+			_serverInfo = serverInfo;
+			_staticConfig = staticConfig;
+		}
+
 		/// <inheritdoc/>
 		public void Configure(IApplicationBuilder app)
 		{ }
@@ -25,7 +39,14 @@ namespace HordeServer
 			services.AddHttpClient();
 			services.AddSingleton<SecretCollectionInternal>();
 			services.AddScoped<ISecretCollection, SecretCollection>();
-			services.AddSingleton<ISecretProvider, AwsParameterStoreSecretProvider>();
+
+			if (_staticConfig.WithAws)
+			{
+				services.AddDefaultAWSOptions(_serverInfo.Configuration.GetAWSOptions());
+				services.AddAWSService<IAmazonSimpleSystemsManagement>();
+				services.AddSingleton<ISecretProvider, AwsParameterStoreSecretProvider>();
+			}
+
 			services.AddSingleton<ISecretProvider, HcpVaultSecretProvider>();
 		}
 	}
@@ -40,5 +61,11 @@ namespace HordeServer
 		/// </summary>
 		public static void AddSecretsConfig(this IDictionary<PluginName, IPluginConfig> dictionary, SecretsConfig secretsConfig)
 			=> dictionary[new PluginName("Secrets")] = secretsConfig;
+
+		/// <summary>
+		/// Gets configuration for the secrets plugin
+		/// </summary>
+		public static SecretsConfig GetSecretsConfig(this IDictionary<PluginName, IPluginConfig> dictionary)
+			=> (SecretsConfig)dictionary[new PluginName("Secrets")];
 	}
 }
diff --git a/Engine/Source/Programs/Horde/Plugins/Secrets/HordeServer.Secrets/SecretsServerConfig.cs b/Engine/Source/Programs/Horde/Plugins/Secrets/HordeServer.Secrets/SecretsServerConfig.cs
new file mode 100644
--- /dev/null
+++ b/Engine/Source/Programs/Horde/Plugins/Secrets/HordeServer.Secrets/SecretsServerConfig.cs
@@ -0,0 +1,17 @@
+﻿// Copyright Epic Games, Inc. All Rights Reserved.
+
+using HordeServer.Plugins;
+
+namespace HordeServer
+{
+	/// <summary>
+	/// Static configuration for the secrets plugin
+	/// </summary>
+	public class SecretsServerConfig : PluginServerConfig
+	{
+		/// <summary>
+		/// Whether to enable Amazon Web Services (AWS) specific features
+		/// </summary>
+		public bool WithAws { get; set; }
+	}
+}
