Horde: Add API to resolve a string to a secret property value

diff --git a/Engine/Source/Programs/Horde/Plugins/Secrets/HordeServer.Secrets.Tests/Secrets/SecretProviderTests.cs b/Engine/Source/Programs/Horde/Plugins/Secrets/HordeServer.Secrets.Tests/Secrets/SecretProviderTests.cs
--- a/Engine/Source/Programs/Horde/Plugins/Secrets/HordeServer.Secrets.Tests/Secrets/SecretProviderTests.cs
+++ b/Engine/Source/Programs/Horde/Plugins/Secrets/HordeServer.Secrets.Tests/Secrets/SecretProviderTests.cs
@@ -50,7 +50,7 @@ namespace HordeServer.Tests.Secrets
 			SecretCollectionInternal secretCollectionInternal = new([
 				new InMemoryProvider(),
 				new HcpVaultSecretProvider(new StubHttpClientFactory(), NullLogger<HcpVaultSecretProvider>.Instance)
-			]);
+			], NullLogger<SecretCollectionInternal>.Instance);
 			string json = """
 				{
 					"secrets": [
@@ -279,7 +279,7 @@ namespace HordeServer.Tests.Secrets
 				}
 			""";
 			SecretsConfig secretsConfig = JsonSerializer.Deserialize<SecretsConfig>(json, JsonUtils.DefaultSerializerOptions)!;
-			SecretCollectionInternal secretCollectionInternal = new([]);
+			SecretCollectionInternal secretCollectionInternal = new([], NullLogger<SecretCollectionInternal>.Instance);
 			SecretCollection secretCollection = new(secretCollectionInternal, new OptionsSnapshotStub<SecretsConfig>(secretsConfig));
 			GlobalConfig config = new();
 			config.Plugins.AddSecretsConfig(secretsConfig);
diff --git a/Engine/Source/Programs/Horde/Plugins/Secrets/HordeServer.Secrets.Tests/Secrets/SecretResolveTests.cs b/Engine/Source/Programs/Horde/Plugins/Secrets/HordeServer.Secrets.Tests/Secrets/SecretResolveTests.cs
new file mode 100644
--- /dev/null
+++ b/Engine/Source/Programs/Horde/Plugins/Secrets/HordeServer.Secrets.Tests/Secrets/SecretResolveTests.cs
@@ -0,0 +1,108 @@
+// Copyright Epic Games, Inc. All Rights Reserved.
+
+using System.Text.Json;
+using EpicGames.Horde.Secrets;
+using HordeServer.Plugins;
+using HordeServer.Secrets;
+using HordeServer.Server;
+using HordeServer.Utilities;
+using Microsoft.Extensions.Logging;
+using Microsoft.Extensions.Options;
+using Moq;
+
+namespace HordeServer.Tests.Secrets
+{
+	[TestClass]
+	public class SecretResolveTests
+	{
+		private static ISecretCollection s_secretCollection = default!;
+		private static GlobalConfig s_globalConfig = default!;
+		private static Mock<ILogger<SecretCollectionInternal>> s_logger = new();
+
+		private class OptionsSnapshotStub<T>(T value) : IOptionsSnapshot<T> where T : class
+		{
+			public T Value { get; } = value;
+			public T Get(string? name)
+			{
+				return Value;
+			}
+		}
+
+		[ClassInitialize]
+		public static async Task CreateSecretCollectionAsync(TestContext testContext)
+		{
+			string json = """
+				{
+					"secrets": [
+						{ "id": "mySecret", "data": { "username": "j.doe", "password": "password123" } },
+						{ "id": "misconfiguredSecret", "sources": [ { "provider": "secretProviderNotFound" } ] }
+					]
+				}
+			""";
+			SecretsConfig secretsConfig = JsonSerializer.Deserialize<SecretsConfig>(json, JsonUtils.DefaultSerializerOptions)!;
+			SecretCollection secretCollection = new(new SecretCollectionInternal([], s_logger.Object), new OptionsSnapshotStub<SecretsConfig>(secretsConfig));
+			GlobalConfig config = new();
+			config.Plugins.AddSecretsConfig(secretsConfig);
+			await config.PostLoadAsync(new ServerSettings(), [], []);
+			s_globalConfig = config;
+			s_secretCollection = secretCollection;
+		}
+
+		[TestMethod]
+		[DataRow("horde:secret:mySecret.username", "username", "j.doe", DisplayName = "Get Username")]
+		[DataRow("horde:secret:mySecret.password", "password", "password123", DisplayName = "Get Password")]
+		public async Task ResolveSecretAsync(string value, string expectedName, string expectedValue)
+		{
+			ISecretProperty? secret = await s_secretCollection.ResolveAsync(value);
+			Assert.IsNotNull(secret);
+			Assert.AreEqual(expectedName, secret.Name);
+			Assert.AreEqual(expectedValue, secret.Value);
+		}
+
+		[TestMethod]
+		[DataRow("horde:secret:mySecret.username", "username", "j.doe", DisplayName = "Get Username")]
+		[DataRow("horde:secret:mySecret.password", "password", "password123", DisplayName = "Get Password")]
+		public async Task ResolveSecretWithConfigAsync(string value, string expectedName, string expectedValue)
+		{
+			ISecretProperty? secret = await s_secretCollection.ResolveAsync(value, s_globalConfig.Plugins[new PluginName("Secrets")]);
+			Assert.IsNotNull(secret);
+			Assert.AreEqual(expectedName, secret.Name);
+			Assert.AreEqual(expectedValue, secret.Value);
+		}
+
+		[TestMethod]
+		[ExpectedException(typeof(InvalidCastException))]
+		public async Task ResolveSecretWithInvalidConfigAsync()
+		{
+			_ = await s_secretCollection.ResolveAsync("horde:secret:mySecret.password", new BuildConfig());
+		}
+
+		[TestMethod]
+		[DataRow("", DisplayName = "Empty string")]
+		[DataRow(null, DisplayName = "Null string")]
+		[DataRow("invalid string", DisplayName = "Incorrect format")]
+		[DataRow("horde:secret:doesNotExist.password", DisplayName = "Secret does not exist")]
+		[DataRow("horde:secret:mySecret.doesNotExist", DisplayName = "Secret property does not exist")]
+		public async Task ResolveInvalidSecretAsync(string value)
+		{
+			ISecretProperty? secret = await s_secretCollection.ResolveAsync(value);
+			Assert.IsNull(secret);
+		}
+
+		[TestMethod]
+		public async Task ResolveSecretWithMisconfiguredConfigAsync()
+		{
+			ISecretProperty? secret = await s_secretCollection.ResolveAsync("horde:secret:misconfiguredSecret.property");
+			Assert.IsNull(secret);
+			s_logger.Verify(x =>
+				x.Log(
+					LogLevel.Error,
+					It.IsAny<EventId>(),
+					It.Is<It.IsAnyType>((value, _) => value.ToString()!.Contains("Unable to resolve the string")),
+					It.IsAny<KeyNotFoundException>(),
+					It.IsAny<Func<It.IsAnyType, Exception?, string>>()
+				)
+			);
+		}
+	}
+}
diff --git a/Engine/Source/Programs/Horde/Plugins/Secrets/HordeServer.Secrets/Secrets/SecretCollection.cs b/Engine/Source/Programs/Horde/Plugins/Secrets/HordeServer.Secrets/Secrets/SecretCollection.cs
--- a/Engine/Source/Programs/Horde/Plugins/Secrets/HordeServer.Secrets/Secrets/SecretCollection.cs
+++ b/Engine/Source/Programs/Horde/Plugins/Secrets/HordeServer.Secrets/Secrets/SecretCollection.cs
@@ -1,7 +1,9 @@
 // Copyright Epic Games, Inc. All Rights Reserved.
 
 using System.Text.Json;
+using System.Text.RegularExpressions;
 using EpicGames.Horde.Secrets;
+using Microsoft.Extensions.Logging;
 using Microsoft.Extensions.Options;
 
 namespace HordeServer.Secrets
@@ -24,19 +26,22 @@ namespace HordeServer.Secrets
 		}
 
 		/// <inheritdoc/>
-		public async Task<ISecret?> GetAsync(SecretId secretId, CancellationToken cancellationToken)
+		public async Task<ISecretProperty?> ResolveAsync(string value, CancellationToken cancellationToken)
 		{
-			SecretConfig? secretConfig;
-			if (!_secretsConfig.Value.TryGetSecret(secretId, out secretConfig))
-			{
-				return null;
-			}
+			return await _secretCollectionInternal.ResolveAsync(value, _secretsConfig.Value, cancellationToken);
+		}
 
-			// secret provider configs are optional
-			Dictionary<string, SecretProviderConfig>? secretProviderConfigs;
-			_secretsConfig.Value.TryGetSecretProviderConfigs(secretId, out secretProviderConfigs);
+		/// <inheritdoc/>
+		public async Task<ISecretProperty?> ResolveAsync(string value, object config, CancellationToken cancellationToken)
+		{
+			SecretsConfig secretsConfig = config as SecretsConfig ?? throw new InvalidCastException($"{nameof(config)} is not a SecretsConfig");
+			return await _secretCollectionInternal.ResolveAsync(value, secretsConfig, cancellationToken);
+		}
 
-			return await _secretCollectionInternal.GetAsync(secretConfig, secretProviderConfigs, cancellationToken);
+		/// <inheritdoc/>
+		public async Task<ISecret?> GetAsync(SecretId secretId, CancellationToken cancellationToken)
+		{
+			return await _secretCollectionInternal.GetAsync(secretId, _secretsConfig.Value, cancellationToken);
 		}
 	}
 
@@ -57,14 +62,49 @@ namespace HordeServer.Secrets
 			}
 		}
 
+		class SecretProperty : ISecretProperty
+		{
+			public SecretId Id { get; }
+			public string Name { get; }
+			public string Value { get; }
+
+			public SecretProperty(SecretId id, string name, string value)
+			{
+				Id = id;
+				Name = name;
+				Value = value;
+			}
+		}
+
 		readonly Dictionary<string, ISecretProvider> _providers;
+		readonly ILogger _logger;
 
 		/// <summary>
 		/// Constructor
 		/// </summary>
-		public SecretCollectionInternal(IEnumerable<ISecretProvider> providers)
+		public SecretCollectionInternal(IEnumerable<ISecretProvider> providers, ILogger<SecretCollectionInternal> logger)
 		{
 			_providers = providers.ToDictionary(x => x.Name, x => x);
+			_logger = logger;
+		}
+
+		/// <summary>
+		/// Resolve a secret to concrete values
+		/// </summary>
+		/// <param name="secretId">The secret to resolve</param>
+		/// <param name="secretsConfig">Configuration for the secret</param>
+		/// <param name="cancellationToken">Cancellation token for the operation</param>
+		public async Task<ISecret?> GetAsync(SecretId secretId, SecretsConfig secretsConfig, CancellationToken cancellationToken)
+		{
+			if (!secretsConfig.TryGetSecret(secretId, out SecretConfig? secretConfig))
+			{
+				return null;
+			}
+
+			// secret provider configs are optional
+			secretsConfig.TryGetSecretProviderConfigs(secretId, out Dictionary<string, SecretProviderConfig>? secretProviderConfigs);
+
+			return await GetAsync(secretConfig, secretProviderConfigs, cancellationToken);
 		}
 
 		/// <summary>
@@ -124,5 +164,49 @@ namespace HordeServer.Secrets
 
 			return new Secret(config.Id, data);
 		}
+
+		/// <summary>
+		/// Resolves a string to a property of a secret
+		/// </summary>
+		/// <param name="value">A string that contains a secret ID and property name</param>
+		/// <param name="secretsConfig">Configuration for the secret</param>
+		/// <param name="cancellationToken">Cancellation token for the operation</param>
+		public async Task<ISecretProperty?> ResolveAsync(string value, SecretsConfig secretsConfig, CancellationToken cancellationToken)
+		{
+			if (String.IsNullOrEmpty(value))
+			{
+				return null;
+			}
+
+			Match match = Regex.Match(value, @"^horde:secret:(?<SecretId>[^\.]+)\.(?<PropertyName>[^\.]+)?$");
+			if (!match.Success)
+			{
+				return null;
+			}
+
+			SecretId secretId = new(match.Groups["SecretId"].Value);
+			ISecret? secret = null;
+			try
+			{
+				secret = await GetAsync(secretId, secretsConfig, cancellationToken);
+			}
+			catch (Exception e)
+			{
+				_logger.LogError(e, "Unable to resolve the string '{Value}'. The string was not recognized as a valid secret. Check the configuration of the secrets config: '{Message}'", value, e.Message);
+			}
+			if (secret == null)
+			{
+				return null;
+			}
+
+			string propertyName = match.Groups["PropertyName"].Value;
+			if (!secret.Data.TryGetValue(propertyName, out string? propertyValue))
+			{
+				return null;
+			}
+
+			ISecretProperty secretProperty = new SecretProperty(secretId, propertyName, propertyValue);
+			return secretProperty;
+		}
 	}
 }
diff --git a/Engine/Source/Programs/Horde/Plugins/Secrets/HordeServer.Secrets/Secrets/SecretsController.cs b/Engine/Source/Programs/Horde/Plugins/Secrets/HordeServer.Secrets/Secrets/SecretsController.cs
--- a/Engine/Source/Programs/Horde/Plugins/Secrets/HordeServer.Secrets/Secrets/SecretsController.cs
+++ b/Engine/Source/Programs/Horde/Plugins/Secrets/HordeServer.Secrets/Secrets/SecretsController.cs
@@ -70,5 +70,30 @@ namespace HordeServer.Secrets
 
 			return new GetSecretResponse(secret.Id, secret.Data.ToDictionary(x => x.Key, x => x.Value)).ApplyFilter(filter);
 		}
+
+		/// <summary>
+		/// Converts the string representation of a secret and property to a concrete value
+		/// The format of the string is 'horde:secret:secret-id.property-name'
+		/// </summary>
+		/// <param name="value">A string that contains a secret ID and property name</param>
+		/// <param name="cancellationToken">Cancellation token for the operation</param>
+		/// <returns>Information about the requested secret</returns>
+		[HttpGet]
+		[Route("/api/v1/secrets/resolve/{value}")]
+		[ProducesResponseType(typeof(GetSecretPropertyResponse), 200)]
+		public async Task<ActionResult<object>> ResolveSecretAsync(string value, CancellationToken cancellationToken = default)
+		{
+			ISecretProperty? secret = await _secretCollection.ResolveAsync(value, cancellationToken);
+			if (secret == null)
+			{
+				return NotFound(value);
+			}
+			if (!_secretsConfig.Value.Authorize(secret.Id, SecretAclAction.ViewSecret, User))
+			{
+				return Forbid(SecretAclAction.ViewSecret, secret.Id);
+			}
+
+			return new GetSecretPropertyResponse(secret.Id, secret.Name, secret.Value);
+		}
 	}
 }
diff --git a/Engine/Source/Programs/Shared/EpicGames.Horde/HordeHttpClient.cs b/Engine/Source/Programs/Shared/EpicGames.Horde/HordeHttpClient.cs
--- a/Engine/Source/Programs/Shared/EpicGames.Horde/HordeHttpClient.cs
+++ b/Engine/Source/Programs/Shared/EpicGames.Horde/HordeHttpClient.cs
@@ -344,16 +344,28 @@ namespace EpicGames.Horde
 		}
 
 		/// <summary>
-		/// Retrieve information about a specific project
+		/// Retrieve information about a specific secret
 		/// </summary>
 		/// <param name="secretId">Id of the secret to retrieve</param>
 		/// <param name="cancellationToken">Cancellation token for the operation</param>
-		/// <returns>Information about the requested project</returns>
+		/// <returns>Information about the requested secret</returns>
 		public Task<GetSecretResponse> GetSecretAsync(SecretId secretId, CancellationToken cancellationToken = default)
 		{
 			return GetAsync<GetSecretResponse>(_httpClient, $"api/v1/secrets/{secretId}", cancellationToken);
 		}
 
+		/// <summary>
+		/// Retrieve information about a specific secret and property
+		/// </summary>
+		/// <param name="value">A string representation of a secret to retrieve.
+		/// A string that contains the "horde:secret:" prefix followed by secret id and property name e.g. 'horde:secret:my-secret.property' </param>
+		/// <param name="cancellationToken">Cancellation token for the operation</param>
+		/// <returns>Information about the requested secret</returns>
+		public Task<GetSecretPropertyResponse> ResolveSecretAsync(string value, CancellationToken cancellationToken = default)
+		{
+			return GetAsync<GetSecretPropertyResponse>(_httpClient, $"api/v1/secrets/resolve/{value}", cancellationToken);
+		}
+
 		#endregion
 
 		#region Server
diff --git a/Engine/Source/Programs/Shared/EpicGames.Horde/Secrets/HttpSecretCollection.cs b/Engine/Source/Programs/Shared/EpicGames.Horde/Secrets/HttpSecretCollection.cs
--- a/Engine/Source/Programs/Shared/EpicGames.Horde/Secrets/HttpSecretCollection.cs
+++ b/Engine/Source/Programs/Shared/EpicGames.Horde/Secrets/HttpSecretCollection.cs
@@ -1,5 +1,6 @@
 ﻿// Copyright Epic Games, Inc. All Rights Reserved.
 
+using System;
 using System.Collections.Generic;
 using System.Threading;
 using System.Threading.Tasks;
@@ -19,6 +20,18 @@ namespace EpicGames.Horde.Secrets
 				=> _response = response;
 		}
 
+		class SecretProperty : ISecretProperty
+		{
+			private readonly GetSecretPropertyResponse _response;
+
+			public SecretId Id => _response.Id;
+			public string Name => _response.Name;
+			public string Value => _response.Value;
+
+			public SecretProperty(GetSecretPropertyResponse response)
+				=> _response = response;
+		}
+
 		readonly IHordeClient _hordeClient;
 
 		/// <summary>
@@ -34,5 +47,19 @@ namespace EpicGames.Horde.Secrets
 			GetSecretResponse response = await hordeHttpClient.GetSecretAsync(secretId, cancellationToken);
 			return new Secret(response);
 		}
+
+		/// <inheritdoc/>
+		public async Task<ISecretProperty?> ResolveAsync(string value, CancellationToken cancellationToken = default)
+		{
+			HordeHttpClient hordeHttpClient = _hordeClient.CreateHttpClient();
+			GetSecretPropertyResponse response = await hordeHttpClient.ResolveSecretAsync(value, cancellationToken);
+			return new SecretProperty(response);
+		}
+
+		/// <inheritdoc/>
+		public Task<ISecretProperty?> ResolveAsync(string value, object config, CancellationToken cancellationToken = default)
+		{
+			throw new NotImplementedException();
+		}
 	}
 }
diff --git a/Engine/Source/Programs/Shared/EpicGames.Horde/Secrets/ISecret.cs b/Engine/Source/Programs/Shared/EpicGames.Horde/Secrets/ISecret.cs
--- a/Engine/Source/Programs/Shared/EpicGames.Horde/Secrets/ISecret.cs
+++ b/Engine/Source/Programs/Shared/EpicGames.Horde/Secrets/ISecret.cs
@@ -19,4 +19,25 @@ namespace EpicGames.Horde.Secrets
 		/// </summary>
 		IReadOnlyDictionary<string, string> Data { get; }
 	}
+
+	/// <summary>
+	/// Information about a single property of a secret
+	/// </summary>
+	public interface ISecretProperty
+	{
+		/// <summary>
+		/// Identifier for the secret
+		/// </summary>
+		SecretId Id { get; }
+
+		/// <summary>
+		/// The name of the property from the secret
+		/// </summary>
+		string Name { get; }
+
+		/// <summary>
+		/// The value of the property from the secret
+		/// </summary>
+		string Value { get; }
+	}
 }
diff --git a/Engine/Source/Programs/Shared/EpicGames.Horde/Secrets/ISecretCollection.cs b/Engine/Source/Programs/Shared/EpicGames.Horde/Secrets/ISecretCollection.cs
--- a/Engine/Source/Programs/Shared/EpicGames.Horde/Secrets/ISecretCollection.cs
+++ b/Engine/Source/Programs/Shared/EpicGames.Horde/Secrets/ISecretCollection.cs
@@ -16,5 +16,22 @@ namespace EpicGames.Horde.Secrets
 		/// <param name="secretId">Identifier for the secret</param>
 		/// <param name="cancellationToken">Cancellation token for the operation</param>
 		Task<ISecret?> GetAsync(SecretId secretId, CancellationToken cancellationToken = default);
+
+		/// <summary>
+		/// Converts the string representation of a secret and property to a concrete value
+		/// The format of the string is 'horde:secret:secret-id.property-name'
+		/// </summary>
+		/// <param name="value">A string that contains a secret ID and property name</param>
+		/// <param name="cancellationToken">Cancellation token for the operation</param>
+		Task<ISecretProperty?> ResolveAsync(string value, CancellationToken cancellationToken = default);
+
+		/// <summary>
+		/// Converts the string representation of a secret and property to a concrete value
+		/// </summary>
+		/// <param name="value">A string that contains a secret ID and property name</param>
+		/// <param name="config">A SecretsConfig object</param>
+		/// <param name="cancellationToken">Cancellation token for the operation</param>
+		/// <remarks>The method should only be used when the config cannot be retrieved from global config object, for example when it as not initialized inside a IPluginConfig.PostLoadAsync method</remarks>
+		Task<ISecretProperty?> ResolveAsync(string value, object config, CancellationToken cancellationToken = default);
 	}
 }
diff --git a/Engine/Source/Programs/Shared/EpicGames.Horde/Secrets/SecretMessages.cs b/Engine/Source/Programs/Shared/EpicGames.Horde/Secrets/SecretMessages.cs
--- a/Engine/Source/Programs/Shared/EpicGames.Horde/Secrets/SecretMessages.cs
+++ b/Engine/Source/Programs/Shared/EpicGames.Horde/Secrets/SecretMessages.cs
@@ -48,4 +48,36 @@ namespace EpicGames.Horde.Secrets
 			Data = data;
 		}
 	}
+
+	/// <summary>
+	/// Gets data for a particular secret and property
+	/// </summary>
+	[DebuggerDisplay("{Id}")]
+	public class GetSecretPropertyResponse
+	{
+		/// <summary>
+		/// Id of the secret
+		/// </summary>
+		public SecretId Id { get; }
+
+		/// <summary>
+		/// Name of the property in the secret
+		/// </summary>
+		public string Name { get;  }
+
+		/// <summary>
+		/// Value of the property in the secret
+		/// </summary>
+		public string Value { get; }
+
+		/// <summary>
+		/// Constructor
+		/// </summary>
+		public GetSecretPropertyResponse(SecretId id, string name, string value)
+		{
+			Id = id;
+			Name = name;
+			Value = value;
+		}
+	}
 }
