Horde: Add HTTP config values for UBA cache in compute cluster config

diff --git a/Engine/Source/Programs/Horde/Plugins/Build/HordeServer.Build.Tests/Compute/ComputeServiceTest.cs b/Engine/Source/Programs/Horde/Plugins/Build/HordeServer.Build.Tests/Compute/ComputeServiceTest.cs
--- a/Engine/Source/Programs/Horde/Plugins/Build/HordeServer.Build.Tests/Compute/ComputeServiceTest.cs
+++ b/Engine/Source/Programs/Horde/Plugins/Build/HordeServer.Build.Tests/Compute/ComputeServiceTest.cs
@@ -426,7 +426,8 @@ namespace HordeServer.Tests.Compute
 		[TestMethod]
 		public async Task Connection_UbaCache_CacheServerDefined_Async()
 		{
-			ComputeResource? cr = await AllocateWithAgentAsync(useUbaCache: true, ubaCacheEndpoint: "1.2.3.4:3333");
+			UbaComputeClusterConfig uccc = new() { CacheEndpoint = "1.2.3.4:3333", CacheHttpEndpoint = "1.2.3.4:4444", CacheHttpCrypto = "aaaaaaaaaaaaaaaaaaaaaaaaaaaaaaaa" }; 
+			ComputeResource? cr = await AllocateWithAgentAsync(useUbaCache: true, ubaConfig: uccc);
 			Assert.AreEqual("1.2.3.4:3333", cr!.Uba!.CacheEndpoint);
 			Assert.IsTrue(Regex.IsMatch(cr.Uba!.CacheSessionKey, "^[0-9a-fA-F]{32}$"));
 		}
@@ -435,17 +436,21 @@ namespace HordeServer.Tests.Compute
 		[Ignore] // Uncomment to manually test allocation with a local UBA cache server
 		public async Task Connection_UbaCache_RealServer_Async()
 		{
-			const string Endpoint = "localhost:7010";
-			ComputeResource? cr = await AllocateWithAgentAsync(useUbaCache: true, ubaCacheEndpoint: Endpoint);
-			Assert.AreEqual(Endpoint, cr!.Uba!.CacheEndpoint);
+			UbaComputeClusterConfig uccc = new() { CacheEndpoint = "localhost:7009", CacheHttpEndpoint = "localhost:7010", CacheHttpCrypto = "aaaaaaaaaaaaaaaaaaaaaaaaaaaaaaaa" };
+			ComputeResource? cr = await AllocateWithAgentAsync(useUbaCache: true, ubaConfig: uccc);
+			Assert.AreEqual(uccc.CacheEndpoint, cr!.Uba!.CacheEndpoint);
 			Assert.IsTrue(Regex.IsMatch(cr.Uba!.CacheSessionKey, "^[0-9a-fA-F]{32}$"));
 		}
 		
 		[TestMethod]
 		public async Task Connection_UbaCache_CacheServerNotDefined_Async()
 		{
-			ComputeServiceException cse = await Assert.ThrowsExceptionAsync<ComputeServiceException>(() => AllocateWithAgentAsync(useUbaCache: true, ubaCacheEndpoint: null));
-			Assert.IsTrue(cse.Message.StartsWith("No UBA cache server defined for compute cluster", StringComparison.Ordinal));
+			ComputeServiceException cse1 = await Assert.ThrowsExceptionAsync<ComputeServiceException>(() => AllocateWithAgentAsync(useUbaCache: true, ubaConfig: null));
+			Assert.IsTrue(cse1.Message.StartsWith("No UBA cache server defined for compute cluster", StringComparison.Ordinal));
+			
+			UbaComputeClusterConfig uccc = new() { CacheEndpoint = "localhost:7009", CacheHttpEndpoint = null, CacheHttpCrypto = "aaaaaaaaaaaaaaaaaaaaaaaaaaaaaaaa" };
+			ComputeServiceException cse2 = await Assert.ThrowsExceptionAsync<ComputeServiceException>(() => AllocateWithAgentAsync(useUbaCache: true, ubaConfig: uccc));
+			Assert.IsTrue(cse2.Message.StartsWith("No UBA cache server defined for compute cluster", StringComparison.Ordinal));
 		}
 		
 		#endregion Connection
@@ -473,11 +478,11 @@ namespace HordeServer.Tests.Compute
 			string[]? relayIps = null,
 			string? tunnelAddress = null,
 			bool useUbaCache = false,
-			string? ubaCacheEndpoint = null,
+			UbaComputeClusterConfig? ubaConfig = null,
 			ComputeProtocol protocol = ComputeProtocol.Latest,
 			string pool = "myPool")
 		{
-			await using ComputeService cs = await CreateComputeServiceWithAgentAsync(tunnelAddress, pool, new ComputeClusterConfig { Id = _cluster1, UbaCacheEndpoint = ubaCacheEndpoint });
+			await using ComputeService cs = await CreateComputeServiceWithAgentAsync(tunnelAddress, pool, new ComputeClusterConfig { Id = _cluster1, Uba = ubaConfig });
 			AllocateResourceParams arp = new(_cluster1, protocol, new Requirements() { Pool = pool })
 			{
 				ConnectionMode = connectionMode,
diff --git a/Engine/Source/Programs/Horde/Plugins/Compute/HordeServer.Compute/Compute/ComputeService.cs b/Engine/Source/Programs/Horde/Plugins/Compute/HordeServer.Compute/Compute/ComputeService.cs
--- a/Engine/Source/Programs/Horde/Plugins/Compute/HordeServer.Compute/Compute/ComputeService.cs
+++ b/Engine/Source/Programs/Horde/Plugins/Compute/HordeServer.Compute/Compute/ComputeService.cs
@@ -419,12 +419,16 @@ namespace HordeServer.Compute
 			Task<UbaComputeResource>? ubaCache = null;
 			if (arp.UseUbaCache ?? false)
 			{
-				if (clusterConfig.UbaCacheEndpoint == null)
+				if (clusterConfig.Uba?.CacheEndpoint == null || clusterConfig.Uba.CacheHttpEndpoint == null || clusterConfig.Uba.CacheHttpCrypto == null)
 				{
 					throw new ComputeServiceException($"No UBA cache server defined for compute cluster {arp.ClusterId}") { ShowToUser = true };
 				}
 				// Allocate UBA cache session key in the background, await the task later
-				ubaCache = AllocateUbaCacheServerAsync(_httpClient, clusterConfig.UbaCacheEndpoint, cancellationToken);
+				// Validate the config above as the method call below is not awaited and exceptions would not immediately be thrown
+				ubaCache = AllocateUbaCacheServerAsync(_httpClient,
+					clusterConfig.Uba.CacheEndpoint,
+					clusterConfig.Uba.CacheHttpEndpoint,
+					clusterConfig.Uba.CacheHttpCrypto, cancellationToken);
 			}
 
 			try
@@ -878,16 +882,19 @@ namespace HordeServer.Compute
 			return relayIps[0];
 		}
 		
-		internal static async Task<UbaComputeResource> AllocateUbaCacheServerAsync(HttpClient httpClient, string ubaCacheEndpoint, CancellationToken cancellationToken)
+		internal static async Task<UbaComputeResource> AllocateUbaCacheServerAsync(HttpClient httpClient,
+			string ubaCacheEndpoint, string ubaCacheHttpEndpoint, string ubaCacheHttpCrypto, CancellationToken cancellationToken)
 		{
 			try
 			{
 				using CancellationTokenSource timeoutCts = CancellationTokenSource.CreateLinkedTokenSource(cancellationToken);
-				timeoutCts.CancelAfter(TimeSpan.FromSeconds(3));
+				timeoutCts.CancelAfter(TimeSpan.FromSeconds(5));
+				
+				_ = ubaCacheHttpCrypto;
 				
 				// Generate a session encryption key that we activate in the UBA cache server
 				string key = RandomNumberGenerator.GetHexString(32, true);
-				HttpResponseMessage res = await httpClient.GetAsync(new Uri($"http://{ubaCacheEndpoint}/addcrypto?{key}"), timeoutCts.Token);
+				HttpResponseMessage res = await httpClient.GetAsync(new Uri($"http://{ubaCacheHttpEndpoint}/addcrypto?{key}"), timeoutCts.Token);
 				res.EnsureSuccessStatusCode();
 				
 				return new UbaComputeResource(ubaCacheEndpoint, key);
diff --git a/Engine/Source/Programs/Horde/Plugins/Compute/HordeServer.Compute/ComputeConfig.cs b/Engine/Source/Programs/Horde/Plugins/Compute/HordeServer.Compute/ComputeConfig.cs
--- a/Engine/Source/Programs/Horde/Plugins/Compute/HordeServer.Compute/ComputeConfig.cs
+++ b/Engine/Source/Programs/Horde/Plugins/Compute/HordeServer.Compute/ComputeConfig.cs
@@ -234,6 +234,35 @@ namespace HordeServer
 		[CbField("cid")]
 		public string? ComputeId { get; set; }
 	}
+	
+	/// <summary>
+	/// Configuration settings for Unreal Build Accelerator (UBA) for a compute cluster
+	/// Contains cache server connectivity and other UBA-related settings.
+	/// Ensure all cache related fields are non-null to activate use of cache.
+	/// </summary>
+	public class UbaComputeClusterConfig
+	{
+		/// <summary>
+		/// Host and port for connecting to the UBA cache server (host:port format)
+		/// As specified with `-port=[PORT]` to UBA cache process.
+		/// Serves all cache requests and must be accessible to both agents and initiators running UBA / Unreal Build Tool.
+		/// </summary>
+		public string? CacheEndpoint { get; set; }
+		
+		/// <summary>
+		/// Host and port for HTTP management of UBA cache server (host:port format)
+		/// As specified with `-http=[PORT]` to UBA cache process.
+		/// This port is used by Horde to manage the UBA cache, such as adding new cache session keys for initiators
+		/// For improved security, only allow Horde server to access this port
+		/// </summary>
+		public string? CacheHttpEndpoint { get; set; }
+		
+		/// <summary>
+		/// AES encryption key used for secure Horde-to-cache-server communication
+		/// As specified with `-httpcrypto=[KEY]` to UBA cache process.
+		/// </summary>
+		public string? CacheHttpCrypto { get; set; }
+	}
 
 	/// <summary>
 	/// Profile for executing compute requests
@@ -272,11 +301,9 @@ namespace HordeServer
 		public AclConfig Acl { get; set; } = new AclConfig();
 		
 		/// <summary>
-		/// Network endpoint of an Unreal Build Accelerator (UBA) cache server in IP:port format.
-		/// Must be accessible to both agents and initiators running UBA / Unreal Build Tool.
-		/// If null, use UBA cache functionality will be disabled for this compute cluster.
+		/// Configuration settings for Unreal Build Accelerator (UBA) for a compute cluster
 		/// </summary>
-		public string? UbaCacheEndpoint { get; set; }
+		public UbaComputeClusterConfig? Uba { get; set; }
 
 		/// <summary>
 		/// Callback post loading this config file
