Horde: Allow sessions to get job responses and post job metadata

diff --git a/Engine/Source/Programs/Horde/Plugins/Build/HordeServer.Build/Jobs/JobsController.cs b/Engine/Source/Programs/Horde/Plugins/Build/HordeServer.Build/Jobs/JobsController.cs
--- a/Engine/Source/Programs/Horde/Plugins/Build/HordeServer.Build/Jobs/JobsController.cs
+++ b/Engine/Source/Programs/Horde/Plugins/Build/HordeServer.Build/Jobs/JobsController.cs
@@ -408,9 +408,13 @@ namespace HordeServer.Jobs
 			{
 				return NotFound(job.StreamId);
 			}
+
 			if (!streamConfig.Authorize(JobAclAction.UpdateJob, User))
 			{
-				return Forbid(JobAclAction.UpdateJob, jobId);
+				if (job.StartedByUserId != null || !JobService.AuthorizeSession(job, User))
+				{
+					return Forbid(JobAclAction.UpdateJob, jobId);
+				}					
 			}
 
 			Dictionary<JobStepId, List<string>>? stepMetaData = null;
@@ -489,6 +493,7 @@ namespace HordeServer.Jobs
 			{
 				return NotFound(job.StreamId);
 			}
+
 			if (!streamConfig.Authorize(NotificationAclAction.CreateSubscription, User))
 			{
 				return Forbid(NotificationAclAction.CreateSubscription, jobId);
@@ -530,10 +535,15 @@ namespace HordeServer.Jobs
 			{
 				return NotFound(job.StreamId);
 			}
+
 			if (!streamConfig.Authorize(JobAclAction.ViewJob, User))
 			{
-				return Forbid(JobAclAction.ViewJob, job.StreamId);
+				if (job.StartedByUserId != null || !JobService.AuthorizeSession(job, User))
+				{
+					return Forbid(JobAclAction.ViewJob, job.StreamId);
+				}				
 			}
+
 			if (modifiedAfter != null && job.UpdateTimeUtc <= modifiedAfter.Value)
 			{
 				return new Dictionary<string, object>();
