Horde: Support debugging jobs as scheduler

diff --git a/Engine/Source/Programs/Horde/Plugins/Build/HordeServer.Build/Jobs/JobsController.cs b/Engine/Source/Programs/Horde/Plugins/Build/HordeServer.Build/Jobs/JobsController.cs
--- a/Engine/Source/Programs/Horde/Plugins/Build/HordeServer.Build/Jobs/JobsController.cs
+++ b/Engine/Source/Programs/Horde/Plugins/Build/HordeServer.Build/Jobs/JobsController.cs
@@ -89,6 +89,11 @@ namespace HordeServer.Jobs
 				return Forbid(JobAclAction.CreateJob, streamConfig.Id);
 			}
 
+			if (create.RunAsScheduler == true && !_buildConfig.Value.Authorize(ServerAclAction.Debug, User))
+			{
+				return Forbid(ServerAclAction.Debug);
+			}
+
 			// Get the name of the template ref
 			TemplateId templateRefId = create.TemplateId;
 
@@ -199,19 +204,22 @@ namespace HordeServer.Jobs
 			else if (create.UpdateIssues.HasValue)
 			{
 				updateIssues = create.UpdateIssues.Value;
-			}
+			}			
 
 			// Create options for the new job
 			CreateJobOptions options = new CreateJobOptions(templateRefConfig);
 			options.PreflightCommitId = create.PreflightCommitId;
-			options.PreflightDescription = shelfInfo?.Description;
-			options.StartedByUserId = User.GetUserId();
+			options.PreflightDescription = shelfInfo?.Description;			
 			options.Priority = priority;
 			options.AutoSubmit = create.AutoSubmit;
 			options.UpdateIssues = updateIssues;
 			options.ParentJobId = create.ParentJobId == null ? null : JobId.Parse(create.ParentJobId);
 			options.ParentJobStepId = create.ParentJobStepId == null ? null : JobStepId.Parse(create.ParentJobStepId);
-			options.Claims.AddRange(User.Claims.Select(x => new AclClaimConfig(x)));
+			if (!create.RunAsScheduler.HasValue || !create.RunAsScheduler.Value)
+			{
+				options.StartedByUserId = User.GetUserId();
+				options.Claims.AddRange(User.Claims.Select(x => new AclClaimConfig(x)));
+			}			
 			options.JobOptions ??= create.JobOptions;
 
 			if (create.AdditionalArguments != null)
diff --git a/Engine/Source/Programs/Shared/EpicGames.Horde/Jobs/JobMessages.cs b/Engine/Source/Programs/Shared/EpicGames.Horde/Jobs/JobMessages.cs
--- a/Engine/Source/Programs/Shared/EpicGames.Horde/Jobs/JobMessages.cs
+++ b/Engine/Source/Programs/Shared/EpicGames.Horde/Jobs/JobMessages.cs
@@ -359,6 +359,11 @@ namespace EpicGames.Horde.Jobs
 		/// </summary>
 		public string? ParentJobStepId { get; set; }
 
+		/// <summary>
+		/// Run the job as the scheduler for debugging purposes, requires debugging ACL permission
+		/// </summary>
+		public bool? RunAsScheduler { get; set; }
+
 		/// <summary>
 		/// Private constructor for serialization
 		/// </summary>
