Horde: Fix sorting agents by status

diff --git a/Engine/Source/Programs/Horde/HordeDashboard/src/components/agents/AgentView.tsx b/Engine/Source/Programs/Horde/HordeDashboard/src/components/agents/AgentView.tsx
--- a/Engine/Source/Programs/Horde/HordeDashboard/src/components/agents/AgentView.tsx
+++ b/Engine/Source/Programs/Horde/HordeDashboard/src/components/agents/AgentView.tsx
@@ -2261,10 +2261,14 @@ export const AgentViewInner: React.FC<{ agentId?: string, poolId?: string, searc
             }
             return left.name.localeCompare(right.name);
          case 'status':
-            if (localState.columnSearchState['status'][left.id] && localState.columnSearchState['status'][right.id]) {
-               return localState.columnSearchState['status'][left.id].join('').localeCompare(localState.columnSearchState['status'][right.id].join(''));
+            let [leftAgentStatus, leftAgentDate] = getAgentStatus(left);
+            let [rightAgentStatus, rightAgentDate] = getAgentStatus(right);
+
+            // sort by date if agents have date data and the same status
+            if(leftAgentDate && rightAgentDate && (leftAgentStatus == rightAgentStatus)) {
+               return leftAgentDate.getTime() - rightAgentDate.getTime();
             }
-            return left.name.localeCompare(right.name);
+            return (leftAgentStatus ?? "").localeCompare(rightAgentStatus ?? "");
          case 'software':
             return (left.version ?? "").toLocaleLowerCase().localeCompare(right.version ?? "");
          case 'taskTime':
@@ -2290,6 +2294,62 @@ export const AgentViewInner: React.FC<{ agentId?: string, poolId?: string, searc
       }
    }
 
+   function getAgentStatus(agent: AgentData): [string, Date | undefined] {
+      let agentStatus: string = '';
+      let dateInfo: Date | undefined = undefined;
+
+      // get status from agent's leases
+      if (agent.leases && agent.leases.length > 0) {
+         agent.leases.forEach(lease => {
+            let statusText = lease.name;
+
+            if (lease.state === LeaseState.Cancelled) {
+               if (lease.finishTime) {
+                  statusText += " (Canceled)"
+               } else {
+                  statusText += " (Cancelling)"
+               }
+            }
+            
+           agentStatus += statusText;
+         });
+      }
+      // if there are no leases, get status from other state
+      if (agentStatus.length === 0) {
+         agentStatus = agent.status ?? "";
+
+         if (agentStatus === "Ok") {
+            agentStatus = "Ready";
+         }
+         if (!agent.online) {
+            agentStatus = `Offline - ${agent.lastShutdownReason}`
+            dateInfo = new Date(agent.updateTime);
+         }
+         else if (agent.pendingShutdown) {
+            agentStatus = `Pending Shutdown ${agent.lastShutdownReason}`;
+         }
+         else if (agent.pendingConform || agent.pendingFullConform) {
+            if (agent.conformAttemptCount && agent.conformAttemptCount > 0) {
+               agentStatus = `${agent.pendingFullConform ? "Full " : ""}Conform failed (${agent.conformAttemptCount} attempts)`;
+               if (agent.nextConformTime) {
+                  dateInfo = new Date(agent.nextConformTime);
+               }
+            }
+            else {
+               agentStatus = "Pending Conform";
+               if (agent.pendingFullConform) {
+                  agentStatus += " (Full)";
+               }
+            }
+         }
+         else if (!agent.enabled) {
+            agentStatus = "Disabled";
+         }
+      }
+
+      return [agentStatus, dateInfo];
+   }
+
    function onContextMenu(item?: any, index?: number | undefined, ev?: Event | undefined) {
       const event = ev as MouseEvent;
       localState.setRightClickDiv(event?.clientX, event?.clientY);
@@ -2914,4 +2974,4 @@ export const AgentView: React.FC = () => {
       </Stack>
    </Stack>
 
-}
\ No newline at end of file
+}
