Horde Dashboard: Disable job step retries with expired TempStorage

diff --git a/Engine/Source/Programs/Horde/HordeDashboard/src/components/jobDetailsV2/StepRetryModal.tsx b/Engine/Source/Programs/Horde/HordeDashboard/src/components/jobDetailsV2/StepRetryModal.tsx
--- a/Engine/Source/Programs/Horde/HordeDashboard/src/components/jobDetailsV2/StepRetryModal.tsx
+++ b/Engine/Source/Programs/Horde/HordeDashboard/src/components/jobDetailsV2/StepRetryModal.tsx
@@ -1,5 +1,5 @@
 // Copyright Epic Games, Inc. All Rights Reserved.
-import { Checkbox, DefaultButton, IconButton, Modal, PrimaryButton, ScrollablePane, ScrollbarVisibility, Spinner, SpinnerSize, Stack, Text, TextField } from '@fluentui/react';
+import { Checkbox, DefaultButton, IconButton, MessageBar, Modal, PrimaryButton, ScrollablePane, ScrollbarVisibility, Spinner, SpinnerSize, Stack, Text, TextField } from '@fluentui/react';
 import React, { useState } from 'react';
 import backend from '../../backend';
 import { CreateJobRequest } from '../../backend/Api';
@@ -8,6 +8,7 @@ import { JobDetailsV2 } from './JobDetailsViewCommon';
 import { getHordeStyling } from '../../styles/Styles';
 import { errorDialogStore } from '../error/ErrorStore';
 import moment from 'moment';
+import dashboard from 'horde/backend/Dashboard';
 
 export enum StepRetryType {
    RunAgain,
@@ -142,40 +143,74 @@ export const RetryStepsModal: React.FC<{ stepIds: string[]; jobDetails: JobDetai
    const stepLookup = new Map<string, string>();
 
    const stepNames: string[] = [];
+   const disabledSteps = new Set<string>();
+   const stepOutputInfo = dashboard.getArtifactTypeInfo("step-output")
    stepIds.forEach((stepId, index) => {
       const stepName = jobDetails.getStepName(stepId, true) ?? `Unknown Step ${index}`;
       stepLookup.set(stepName, stepId);
       stepNames.push(stepName);
+
+      // Determine if temp storage for step is expired
+      const step = jobDetails.stepById(stepId)!;
+      let tempStorageExpired = false;
+      if (stepOutputInfo && stepOutputInfo.keepDays && step.startTime) {
+
+         const expireMs = stepOutputInfo.keepDays * 86400000;
+         if ((Date.now() - new Date(step.startTime).getTime()) > expireMs)
+         {
+            tempStorageExpired = true;
+            // Disable step if temp storage expired
+            disabledSteps.add(stepId);
+         }      
+      } 
+
    });
 
    const stepElements: JSX.Element[] = [];
 
    if (stepIds.length > 1) {
-
-      const allSelected = retrySteps.size === stepIds.length;
-
-      stepElements.push(<Stack horizontal tokens={{ childrenGap: 12 }}><Checkbox checked={allSelected} onChange={(ev, checked) => {
-         const newSteps = new Set<string>();
-         if (checked) {
-            stepIds.forEach(s => newSteps.add(s));
-         } else {
-         }
-         setRetrySteps(newSteps);
-      }} /><Text>{retrySteps.size === stepIds.length ? "Deselect All" : "Select All"}</Text></Stack>)
+      const allSelected = retrySteps.size === stepIds.length - disabledSteps.size;
+
+      stepElements.push(
+         <Checkbox
+            key={"selAll"}
+            label={allSelected ? "Deselect All" : "Select All"} 
+            checked={allSelected} 
+            onChange={(ev, checked) => {
+               const newSteps = new Set<string>();
+               if (checked) {
+                  stepIds.forEach(s => {
+                     if(!disabledSteps.has(s)) {
+                        newSteps.add(s)
+                     }
+                  });
+               }
+               setRetrySteps(newSteps);
+            }} 
+         />
+      )
    }
 
-
    stepElements.push(...stepNames.map(name => {
       const stepId = stepLookup.get(name)!;
-      return <Stack horizontal tokens={{ childrenGap: 12 }}><Checkbox checked={retrySteps.has(stepId)} onChange={(ev, checked) => {
-         const newSteps = new Set(retrySteps);
-         if (checked) {
-            newSteps.add(stepId);
-         } else {
-            newSteps.delete(stepId);
-         }
-         setRetrySteps(newSteps);
-      }} /><Text>{name}</Text></Stack>
+
+      return (
+         <Checkbox 
+            label={name}
+            key={stepId}
+            disabled={disabledSteps.has(stepId)} 
+            checked={retrySteps.has(stepId)} 
+            onChange={(ev, checked) => {
+               const newSteps = new Set(retrySteps);
+               if (checked) {
+                  newSteps.add(stepId);
+               } else {
+                  newSteps.delete(stepId);
+               }
+               setRetrySteps(newSteps);
+            }} 
+         />
+      );
    }))
 
    let height = Math.max(stepElements.length * 32, 128);
@@ -197,6 +232,8 @@ export const RetryStepsModal: React.FC<{ stepIds: string[]; jobDetails: JobDetai
          </Stack.Item>
       </Stack>
 
+      {disabledSteps.size > 0 && <MessageBar styles={{root: {marginLeft: 16, marginRight: 16, width: "fitContent"}}}>Some steps are unavailable for retry because required artifacts are missing or have expired.</MessageBar>}
+
       <Stack style={{marginRight: 18, marginTop: 18}}>
          <Stack styles={{ root: { position: "relative", height: height } }}>
             <ScrollablePane scrollbarVisibility={ScrollbarVisibility.auto}>
diff --git a/Engine/Source/Programs/Horde/HordeDashboard/src/styles/Styles.tsx b/Engine/Source/Programs/Horde/HordeDashboard/src/styles/Styles.tsx
--- a/Engine/Source/Programs/Horde/HordeDashboard/src/styles/Styles.tsx
+++ b/Engine/Source/Programs/Horde/HordeDashboard/src/styles/Styles.tsx
@@ -442,7 +442,6 @@ export const getHordeStyling = () => {
                paddingTop: 3
             },
             ".ms-Checkbox-text, .ms-Checkbox-label": {
-               fontSize: "12px",
                height: 14,
                lineHeight: 14,
             },
