Horde Dashboard: Add support for admins to run jobs as scheduler
#rnx

diff --git a/Engine/Source/Programs/Horde/HordeDashboard/src/backend/Api.ts b/Engine/Source/Programs/Horde/HordeDashboard/src/backend/Api.ts
--- a/Engine/Source/Programs/Horde/HordeDashboard/src/backend/Api.ts
+++ b/Engine/Source/Programs/Horde/HordeDashboard/src/backend/Api.ts
@@ -1145,6 +1145,9 @@ export type CreateJobRequest = {
 	/** Whether to update issues based on the outcome of this job */
 	updateIssues?: boolean;
 
+	/** Whether to run the job as the scheduler*/
+	runAsScheduler?: boolean;
+
 	/** Nodes for the new job */
 	groups?: CreateGroupRequest[];
 
diff --git a/Engine/Source/Programs/Horde/HordeDashboard/src/backend/Dashboard.ts b/Engine/Source/Programs/Horde/HordeDashboard/src/backend/Dashboard.ts
--- a/Engine/Source/Programs/Horde/HordeDashboard/src/backend/Dashboard.ts
+++ b/Engine/Source/Programs/Horde/HordeDashboard/src/backend/Dashboard.ts
@@ -159,6 +159,10 @@ export class Dashboard {
         return user[0].value?.length ? user[0].value : undefined;
     }
 
+    get hordeAdmin():boolean {        
+        return !!this.claims?.find(c => c.type === "groups" && c.value === "app-horde-admins");
+    }
+
     get preview(): boolean {
 
         try {
diff --git a/Engine/Source/Programs/Horde/HordeDashboard/src/components/build/NewBuildV2.tsx b/Engine/Source/Programs/Horde/HordeDashboard/src/components/build/NewBuildV2.tsx
--- a/Engine/Source/Programs/Horde/HordeDashboard/src/components/build/NewBuildV2.tsx
+++ b/Engine/Source/Programs/Horde/HordeDashboard/src/components/build/NewBuildV2.tsx
@@ -160,6 +160,7 @@ class BuildOptions {
    advTargets?: string[];
 
    updateIssues?: boolean;
+   runAsScheduler?: boolean;
 
    parameters: Record<string, string> = {};
    disabledParameters: Set<string> = new Set();
@@ -1728,6 +1729,11 @@ const AdvancedPanel: React.FC = observer(() => {
       height = window.innerHeight - 360;
    }
 
+   // Run as scheduler
+   if (dashboard.hordeAdmin) {
+      height += 32;
+   }
+
    // target options for picker
    type TargetPickerItem = {
       key: string;
@@ -1767,6 +1773,14 @@ const AdvancedPanel: React.FC = observer(() => {
                      options.setChanged();
                   }} />
             </Stack>
+            {!options.readOnly && dashboard.hordeAdmin && <Stack>
+               <Checkbox key={`key_adv_run_as_scheduler_${options.currentRenderKey}`}
+                  label="Run As Scheduler"
+                  checked={options.runAsScheduler ?? false} onChange={(ev, checked) => {
+                     options.runAsScheduler = checked ?? false;
+                     options.setChanged();
+                  }} />
+            </Stack>}
             <Stack>
                <TextField key={`key_adv_job_name_${options.currentRenderKey}`} disabled={options.readOnly} spellCheck={false} defaultValue={options.advJobName} label="Job Name" onChange={(ev, newValue) => {
                   options.advJobName = newValue;
@@ -1946,7 +1960,8 @@ const BuildModal: React.FC<{ setUseLegacyDialog: (value: boolean) => void }> = o
          parameters: parameters,
          additionalArguments: additionalArgs?.length ? additionalArgs : undefined,
          targets: options.advTargets,
-         updateIssues: options.updateIssues
+         updateIssues: options.updateIssues,
+         runAsScheduler: options.runAsScheduler
       };
 
       if (typeof (options.change) === 'number') {
