Add phase history
- including phase comparison view
- re-organize the test events handling so it can be more easily expanded
- add jobstep testdata api call
- drop expand/collapse phase item and instead use a docking pattern to avoid focus issues

Add Job Step TestData selector
- redirect if there is only one report available to the one test report

#rnx

diff --git a/Engine/Source/Programs/Horde/HordeDashboard/plugins/test-automation/api.ts b/Engine/Source/Programs/Horde/HordeDashboard/plugins/test-automation/api.ts
--- a/Engine/Source/Programs/Horde/HordeDashboard/plugins/test-automation/api.ts
+++ b/Engine/Source/Programs/Horde/HordeDashboard/plugins/test-automation/api.ts
@@ -30,67 +30,6 @@ export type TestData = GetTestDataResponse & {
     data: object;
 }
 
-/// An event entry
-export type EventEntry = {
-
-    time: string;
-
-    id: number;
-
-    level: EventLevel;
-
-    message: string;
-
-    format: string;
-
-    properties?: Record<string, string | number | Record<string, string | number | undefined> | undefined>;
-}
-
-/// The event level
-export enum EventLevel {
-
-    Unspecified = "Unspecified",
-
-    Information = "Information",
-
-    Warning = "Warning",
-
-    Error = "Error",
-
-    Critical = "Critical"
-}
-
-/// Known property types
-export enum LogPropertyType {
-    ImageCompare = "Image Compare"
-}
-
-/// Image Compare detailed properties
-export type ImageCompare = {
-    "$type": LogPropertyType.ImageCompare;
-    "$text": string;
-
-    //** The path to the unapproved image from the artifacts */
-    unapproved: string;
-
-    //** The path to the unapproved metadata associated with */
-    //** the image from the artifacts */
-    unapproved_metadata?: string;
-
-    //** The path to the approved image from the artifacts */
-    approved?: string;
-
-    //** The path to the approved metadata associated with */
-    //** the image from the artifacts */
-    approved_metadata?: string;
-
-    //** The path to the delta image from the artifacts */
-    difference?: string;
-
-    //** The path to the json report associated with the image from the artifacts */
-    difference_report?: string;
-}
-
 /// A test name request
 export type GetTestsRequest = {
 
@@ -461,6 +400,21 @@ export type TestSessionDetailsResponse = {
     tags?: string[];
 }
 
+export type JobStepTestDataItem = {
+
+    /// The testdata id
+    testDataId: string,
+
+    /// The associated test key
+    testKey: string,
+
+    /// The associated test name
+    testName: string,
+
+    /// The associated test ref id
+    testNameRef: string
+}
+
 export async function getTestDataV2(id: string): Promise<TestData> {
     return new Promise<TestData>((resolve, reject) => {
         backend.fetch.get(`/api/v3/testdata/${id}`, {
@@ -583,4 +537,16 @@ export async function getTestPhaseSessionsV2(streamIds: string[], phaseIds?: str
             reject(reason);
         });
     });
-}
\ No newline at end of file
+}
+
+export async function getJobStepTestDataV2(jobId: string, stepId: string): Promise<JobStepTestDataItem[]> {
+    return new Promise<JobStepTestDataItem[]>((resolve, reject) => {
+        backend.fetch.get(`/api/v3/testdata/`, {
+            params: {jobId: jobId, stepId: stepId}
+        }).then((value) => {
+            resolve(value.data as JobStepTestDataItem[]);
+        }).catch((reason) => {
+            reject(reason);
+        });
+    });
+}
diff --git a/Engine/Source/Programs/Horde/HordeDashboard/plugins/test-automation/phaseCompare.ts b/Engine/Source/Programs/Horde/HordeDashboard/plugins/test-automation/phaseCompare.ts
new file mode 100644
--- /dev/null
+++ b/Engine/Source/Programs/Horde/HordeDashboard/plugins/test-automation/phaseCompare.ts
@@ -0,0 +1,84 @@
+import { action, makeObservable, observable } from "mobx";
+import { TestDataHandler, TestPhaseStatus, TestSessionDetails } from "./testData";
+import { JobParameters } from "horde/components/ChangeButton";
+
+export type JobContext = {job?: JobParameters, rangeCL?: number, step?: {id: string}};
+
+export const getJobParams = (fetchedSessionDetails: TestSessionDetails, handler: TestDataHandler): JobContext => {
+    return {
+        job: {
+            id: fetchedSessionDetails.jobId,
+            streamId: fetchedSessionDetails.streamId,
+            change: fetchedSessionDetails.commitOrder,
+        },
+        rangeCL: fetchedSessionDetails.session ?
+            handler.getStatusStream(fetchedSessionDetails.streamId)?.getPreviousSession(fetchedSessionDetails.test!, fetchedSessionDetails.meta!, fetchedSessionDetails.session)?.commitOrder
+            : undefined,
+        step : {id: fetchedSessionDetails.stepId}
+    }
+}
+
+export class PhaseComparisonContext {
+    constructor() {
+        makeObservable(this);
+        this.reset();
+    }
+
+    reset() {
+        this.isEnable = false;
+        this.commitId = undefined;
+        this.metaKey = undefined
+        this.isLoadingPhase = false;
+        this.phaseStatus = undefined;
+        this.testSessionId = undefined;
+        this.jobParams = undefined;
+        this.setUpdated();
+    }
+
+    setTestSesssionId(id?: string) {
+        if (this.testSessionId === id) return;
+        this.testSessionId = id;
+        this.setUpdated();
+    }
+
+    setMetaKey(key?: string) {
+        if (this.metaKey === key) return;
+        this.metaKey = key;
+        this.setUpdated();
+    }
+
+    setPhaseStatus(phase?: TestPhaseStatus) {
+        if (this.phaseStatus === phase) return;
+        this.phaseStatus = phase;
+        if (!!phase) this.setLoadingPhase(false);
+        this.setUpdated();
+    }
+
+    setLoadingPhase(flag: boolean) {
+        this.isLoadingPhase = flag;
+    }
+
+    setEnable() {
+        this.isEnable = true;
+    }
+
+    @action
+    setUpdated() {
+        this.updated++;
+    }
+
+    subscribe() {
+        if (this.updated) { }
+    }
+
+    @observable
+    updated: number = 0;
+
+    isEnable: boolean = false;
+    isLoadingPhase: boolean = false;
+    commitId?: string;
+    metaKey?: string;
+    phaseStatus?: TestPhaseStatus;
+    testSessionId?: string;
+    jobParams?: JobContext;
+}
diff --git a/Engine/Source/Programs/Horde/HordeDashboard/plugins/test-automation/phaseHistoryGraph.ts b/Engine/Source/Programs/Horde/HordeDashboard/plugins/test-automation/phaseHistoryGraph.ts
new file mode 100644
--- /dev/null
+++ b/Engine/Source/Programs/Horde/HordeDashboard/plugins/test-automation/phaseHistoryGraph.ts
@@ -0,0 +1,316 @@
+import * as d3 from "d3";
+import { TestDataHandler, TestNameRef, MetadataRef, PhaseSessionResult, TestMetaStatus } from './testData';
+import dashboard, { StatusColor } from 'horde/backend/Dashboard';
+import { getHumanTime, getShortNiceTime } from 'horde/base/utilities/timeUtils';
+import { getHordeStyling } from 'horde/styles/Styles';
+import { TestOutcome } from "./api";
+import { getPhaseSessionStatusColor } from "./testAutomationCommon";
+
+// Handle bad "@types/d3" types, fix if addressed upstream
+const _d3 = d3 as any;
+
+type SelectionType = d3.Selection<SVGGElement, unknown, null, undefined>;
+type DivSelectionType = d3.Selection<HTMLDivElement, unknown, null, undefined>;
+
+export class PhaseHistoryGraph {
+
+    constructor(test: TestNameRef, phase: string, stream: string, handler: TestDataHandler, onClickCallback?: (session: PhaseSessionResult) => void) {
+        this.test = test;
+        this.phase = phase;
+
+        this.handler = handler;
+        this.onClickCallback = onClickCallback;
+
+        this.status = handler.getStatusStream(stream)!.tests.get(test)!;
+
+        this.margin = { top: 0, right: 32, bottom: 0, left: 32 };
+    }
+
+    initData(meta: MetadataRef) {
+        this.meta = meta;
+        this.sessions = this.status.sessions.get(meta)?.getPhaseSessions(this.phase) ?? [];
+    }
+
+    render(container: HTMLDivElement, meta: MetadataRef, selectedPhaseSessionId?: string, withoutHeader?: boolean) {
+
+        const { modeColors } = getHordeStyling();
+
+        this.clear(container);
+
+        this.initData(meta);
+
+        this.phaseSessionId = selectedPhaseSessionId;
+
+        const handler = this.handler;
+        const sessions = this.sessions;
+        const width = 1000
+
+        const scolors = dashboard.getStatusColors();
+        const colors = new Map<TestOutcome, string>([
+            [TestOutcome.Success, scolors.get(StatusColor.Success)!],
+            [TestOutcome.Failure, scolors.get(StatusColor.Failure)!],
+            [TestOutcome.Unspecified, scolors.get(StatusColor.Unspecified)!],
+            [TestOutcome.Skipped, scolors.get(StatusColor.Skipped)!]
+        ]);
+
+        const X = _d3.map(sessions, (r) => handler.commitIdDates.get(r.commitId)!.getTime() / 1000);
+        const Y = _d3.map(sessions, (r) => 0);
+
+        const xDomain = d3.extent(handler.commitIdDates.values() as any, (d:any) => d.getTime() / 1000);
+        const yDomain = new _d3.InternSet(Y as any);
+
+        const I = d3.range(X.length);
+
+        const yPadding = 1;
+        const height = withoutHeader? 20 : Math.ceil((yDomain.size + yPadding) * 13) + this.margin.top + this.margin.bottom;
+
+        const xRange = [this.margin.left, width - this.margin.right];
+        const yRange = [this.margin.top, height - this.margin.bottom];
+
+        const xScale = _d3.scaleTime(xDomain, xRange);
+        const yScale = _d3.scalePoint(yDomain, yRange).round(true).padding(yPadding);
+
+        const svg = d3.select(container)
+            .append("svg")
+            .attr("width", width)
+            .attr("height", height + 24)
+            .attr("viewBox", [0, 0, width, height + 24] as any)
+
+        // top axis
+        const xAxis = (g: SelectionType) => {
+
+            const dateMin = new Date(xDomain[0]! * 1000);
+            const dateMax = new Date(xDomain[1]! * 1000);
+
+            let ticks: number[] = [];
+            for (const date of d3.timeDays(dateMin, dateMax, 1).reverse()) {
+                ticks.push(date.getTime() / 1000);
+            }
+
+            if (ticks.length > 14) {
+                let nticks = [...ticks];
+                // remove first and last, will be readded 
+                const first = nticks.shift()!;
+                const last = nticks.pop()!;
+
+                const n = Math.floor(nticks.length / 12);
+
+                const rticks: number[] = [];
+                for (let i = 0; i < nticks.length; i = i + n) {
+                rticks.push(nticks[i]);
+                }
+
+                rticks.unshift(first);
+                rticks.push(last);
+                ticks = rticks;
+            }
+
+            g.attr("transform", `translate(0,16)`)
+                .style("font-family", "Horde Open Sans Regular")
+                .style("font-size", "9px")
+                .call(d3.axisTop(xScale)
+                    .tickValues(ticks)
+                    .tickFormat(d => withoutHeader? "" : getHumanTime(new Date((d as number) * 1000)))
+                    .tickSizeOuter(0))
+                .call(g => g.select(".domain").remove())
+                .call(g => g.selectAll(".tick line")
+                    .attr("stroke-opacity", dashboard.darktheme ? 0.35 : 0.25).clone()
+                    .attr("stroke", dashboard.darktheme ? "#6D6C6B" : "#4D4C4B")
+                    .attr("y2", height - this.margin.bottom)
+                )
+        }
+
+        svg.append("g").call(xAxis)
+
+        const g = svg.append("g")
+            .selectAll()
+            .data(_d3.group(I, i => Y[i]))
+            .join("g")
+            .attr("class", "meta")
+            .attr("id", ([y]: any) => `meta${y as any}`)
+            .attr("transform", ([y]: any) => `translate(0,${(yScale(y) as any) + 16})`);
+
+        g.append("line")
+            .attr("stroke", dashboard.darktheme ? "#6D6C6B" : "#4D4C4B")
+            .attr("stroke-width", 1)
+            .attr("stroke-linecap", 4)
+            .attr("stroke-opacity", dashboard.darktheme ? 0.35 : 0.25)
+            .attr("x1", ([, I]: any) => this.margin.left)
+            .attr("x2", ([, I]: any) => width - this.margin.right);
+
+        g.selectAll("circle")
+            .data(([, I]: any) => I)
+            .enter()
+            .append("g")
+            .attr("id", i => `circle${sessions[i as any].id}`)
+            .attr("transform", i => `translate(${xScale(X[i as any])},0)`)
+            .append("g")
+            .attr("class", "circleScale")
+            .attr("id", i => `circleScale${sessions[i as any].id}`)
+            .each((d, i, nodes) => SessionCircle(nodes[i], sessions[d as any], 5, !!this.phaseSessionId && sessions[d as any].id === this.phaseSessionId));
+
+        const tooltip = this.tooltip = d3.select(container)
+            .append("div")
+            .attr("id", "tooltip")
+            .style("display", "none")
+            .style("background-color", modeColors.background)
+            .style("border", "solid")
+            .style("border-width", "1px")
+            .style("border-radius", "3px")
+            .style("border-color", dashboard.darktheme ? "#413F3D" : "#2D3F5F")
+            .style("padding", "6px")
+            .style("position", "absolute")
+            .style("pointer-events", "none");
+
+        const closestData = (x: number, y: number): PhaseSessionResult | undefined => {
+
+            if (x < this.margin.left - 16) {
+                return undefined;
+            }
+
+            y -= 16;
+
+            const closest = sessions.reduce((best, session, i) => {
+                const absy = Math.abs(yScale(Y[0]) - y);
+                const sessionTime = handler.commitIdDates.get(session.commitId)!.getTime() / 1000;
+                const absx = Math.abs(xScale(sessionTime) - x);
+
+                const lengthSqr = absy * absy + absx * absx;
+
+                if (lengthSqr < best.value) {
+                    return { index: i, value: lengthSqr };
+                }
+                else {
+                    return best;
+                }
+            }, { index: 0, value: Number.MAX_SAFE_INTEGER });
+
+            if (closest) {
+                return sessions[closest.index];
+            }
+
+            return undefined;
+
+        }
+
+        const handleMouseMove = (event: any) => {
+
+            const mouseX = _d3.pointer(event)[0];
+            const mouseY = _d3.pointer(event)[1];
+
+            const closest = closestData(mouseX, mouseY);
+            if (!closest) {
+                handleMouseLeave(undefined);
+                return;
+            }
+
+            if (this.tooltipSessionId !== closest.id) {
+                this.tooltipSessionId = closest.id;
+                const metaIndex = 0;
+
+                svg.selectAll(".circleScale").attr("transform", `scale(1)`);
+                svg.select(`#circleScale${closest.id}`).attr("transform", `scale(2)`);
+                svg.select(`#circle${closest.id}`).raise();
+                const date = getShortNiceTime(closest.start, true, true);
+                let desc = "";
+                desc += `${closest.outcome} <br/>`;
+                desc += `on Commit ${closest.commitId} <br/>`;
+                desc += `${date} <br/>`;
+
+                const timeStamp = handler.commitIdDates.get(closest.commitId)!.getTime() / 1000;
+                const tx = xScale(timeStamp);
+                const ty = yScale(metaIndex as any)!
+
+                this.updateTooltip(true, tx, ty, desc);
+            }
+        }
+
+        const handleMouseLeave = (event: any) => {
+            svg.selectAll(".circleScale").attr("transform", `scale(1)`);
+            tooltip.style("display", "none");
+            this.tooltipSessionId = undefined;
+        }
+
+        const handleMouseClick = (event: any) => {
+
+            const mouseX = _d3.pointer(event)[0];
+            const mouseY = _d3.pointer(event)[1];
+
+            const closest = closestData(mouseX, mouseY);
+            if (closest && this.onClickCallback) {
+                this.onClickCallback(closest);
+            }
+        }
+
+        svg.on("mousemove", (event) => handleMouseMove(event))
+        svg.on("mouseleave", (event) => handleMouseLeave(event))
+        svg.on("click", (event) => handleMouseClick(event))
+
+    }
+
+    updateTooltip(show: boolean, x?: number, y?: number, html?: string) {
+        if (!this.tooltip) {
+            return;
+        }
+
+        x = x ?? 0;
+        y = y ?? 0;
+
+        this.tooltip
+            .style("display", "block")
+            .html(html ?? "")
+            .style("position", `absolute`)
+            .style("width", `max-content`)
+            .style("top", `${y}px`)
+            .style("left", `${x}px`)
+            .style("transform", "translate(10%, 52%)")
+            .style("font-family", "Horde Open Sans Semibold")
+            .style("font-size", "10px")
+            .style("line-height", "16px")
+            .style("shapeRendering", "crispEdges")
+            .style("stroke", "none")
+            .style("z-index", 10000)
+
+    }
+
+    clear(container: HTMLDivElement) {
+        d3.select(container).selectAll("*").remove();
+    }
+
+    // refs
+    test: TestNameRef;
+    meta: MetadataRef;
+    phase: string;
+    phaseSessionId?: string;
+
+    onClickCallback?: (session: PhaseSessionResult) => void;
+
+    handler: TestDataHandler;
+    margin: { top: number, right: number, bottom: number, left: number }
+
+    hasRendered = false;
+
+    status: TestMetaStatus;
+    sessions: PhaseSessionResult[] = [];
+
+    tooltip?: DivSelectionType;
+    tooltipSessionId?: string;
+}
+
+const SessionCircle = (container: SVGGElement, session: PhaseSessionResult, radius: number, selected?: boolean) => {
+    const color = getPhaseSessionStatusColor(session);
+    const g = d3.select(container);
+    g.append("circle")
+        .attr("fill", color as any)
+        .attr("r", radius);
+
+    if (selected) {
+        g.append("circle")
+            .attr("stroke", dashboard.darktheme ? "#c1c4c9" : "#3d4654 ")
+            .attr("fill", "transparent")
+            .attr("stroke-width", 2)
+            .attr("r", radius + 3);
+    }
+
+    return g;
+}
diff --git a/Engine/Source/Programs/Horde/HordeDashboard/plugins/test-automation/testAutomationCommon.tsx b/Engine/Source/Programs/Horde/HordeDashboard/plugins/test-automation/testAutomationCommon.tsx
--- a/Engine/Source/Programs/Horde/HordeDashboard/plugins/test-automation/testAutomationCommon.tsx
+++ b/Engine/Source/Programs/Horde/HordeDashboard/plugins/test-automation/testAutomationCommon.tsx
@@ -1,8 +1,57 @@
 import { mergeStyles, mergeStyleSets, Text, ComboBox, IComboBoxOption, SelectableOptionMenuItemType } from "@fluentui/react";
-import { TestPhaseStatus, TestSessionResult } from "./testData";
+import { MetadataRef, PhaseSessionResult, TestDataHandler, TestNameRef, TestPhaseStatus, TestSessionDetails, TestSessionResult } from "./testData";
 import dashboard, { StatusColor } from "horde/backend/Dashboard";
 import { useEffect, useState } from "react";
-import { TestPhaseOutcome } from "./api";
+import { TestOutcome, TestPhaseOutcome } from "./api";
+import { PhaseHistoryGraph } from "./phaseHistoryGraph";
+
+export type PhaseSessionCallback = (session: PhaseSessionResult) => void;
+
+export const styles = {
+    labelSmall: {
+        fontSize: 12,
+        fontWeight: "bold"
+    },
+    textSmall: {
+        fontSize: 12
+    },
+    textLarge: {
+        fontSize: 24
+    },
+    defaultButton: {
+        backgroundColor: "#035ca1"
+    },
+    defaultFilter: {
+        backgroundColor: "#0078d4"
+    },
+    defaultCheckbox: {
+        root: {
+            height: 15
+        },
+        text: {
+            fontSize: 11,
+            marginLeft: 0,
+            lineHeight: 15
+        },
+        checkbox: {
+            height: 15,
+            width: 15
+        }
+    },
+    icon: {
+        cursor: "default",
+        height: 15,
+        paddingTop: 1
+    },
+}
+
+const barStyles = mergeStyleSets({
+
+    stripes: {
+        backgroundImage: 'repeating-linear-gradient(-45deg, rgba(255, 255, 255, .2) 25%, transparent 25%, transparent 50%, rgba(255, 255, 255, .2) 50%, rgba(255, 255, 255, .2) 75%, transparent 75%, transparent)',
+    }
+
+});
 
 export const MultiOptionChooser: React.FC<{options: IComboBoxOption[], initialSelection: string[], updateKeys: (selectedKeys: string[]) => void, placeholder?: string, style?: any, disabled?: boolean }> = ({ options, initialSelection, updateKeys, placeholder, style, disabled }) => {
     const [selection, setSelection] = useState<string[]>(initialSelection);
@@ -43,14 +92,6 @@ export const MultiOptionChooser: React.FC<{options: IComboBoxOption[], initialSe
         }} />
 };
 
-const styles = mergeStyleSets({
-
-    stripes: {
-        backgroundImage: 'repeating-linear-gradient(-45deg, rgba(255, 255, 255, .2) 25%, transparent 25%, transparent 50%, rgba(255, 255, 255, .2) 50%, rgba(255, 255, 255, .2) 75%, transparent 75%, transparent)',
-    }
-
-});
-
 export type StatusBarStack = {
     value: number,
     title?: string,
@@ -72,7 +113,7 @@ export const StatusBar = (stack: StatusBarStack[], width: number, height: number
         <div className={mergeStyles({ backgroundColor: basecolor, width: width, height: height, verticalAlign: 'middle', display: "flex" }, style)} title={mainTitle}>
             {stack.map((item) => <span key={item.title!}
                 onClick={item.onClick}
-                className={item.stripes ? styles.stripes : undefined}
+                className={item.stripes ? barStyles.stripes : undefined}
                 style={{
                 width: `${item.value}%`, height: '100%',
                 backgroundColor: item.color,
@@ -84,10 +125,10 @@ export const StatusBar = (stack: StatusBarStack[], width: number, height: number
     );
 }
 
-export const SessionStatusBar = (session: TestSessionResult, width: number, height: number): JSX.Element | undefined => {
-    const sessionTotalCount = session.phaseSucceededCount + session.phaseFailedCount + session.phaseUnspecifiedCount;
-    const metaFailedFactor = Math.ceil(session.phaseFailedCount / (sessionTotalCount || 1) * 50) / 50;
-    const metaUnspecifiedFactor = Math.ceil(session.phaseUnspecifiedCount / (sessionTotalCount || 1) * 50) / 50;
+export const SessionStatusBar = (session: TestSessionResult | TestSessionDetails, width: number, height: number): JSX.Element | undefined => {
+    const sessionTotalCount = session.phasesSucceededCount + session.phasesFailedCount + session.phasesUnspecifiedCount;
+    const metaFailedFactor = Math.ceil(session.phasesFailedCount / (sessionTotalCount || 1) * 50) / 50;
+    const metaUnspecifiedFactor = Math.ceil(session.phasesUnspecifiedCount / (sessionTotalCount || 1) * 50) / 50;
 
     const bDisplayBar = (metaFailedFactor + metaUnspecifiedFactor) > 0 && metaFailedFactor < 1 && metaUnspecifiedFactor < 1;
     if (!bDisplayBar) return undefined
@@ -98,14 +139,14 @@ export const SessionStatusBar = (session: TestSessionResult, width: number, heig
         {
             value: metaUnspecifiedFactor * 100,
             title: "Unspecified",
-            titleValue: `${ Math.ceil(session.phaseUnspecifiedCount / (sessionTotalCount || 1) * 100)}%`,
+            titleValue: `${ Math.ceil(session.phasesUnspecifiedCount / (sessionTotalCount || 1) * 100)}%`,
             color: statusColors.get(StatusColor.Unspecified)!,
             stripes: true
         },
         {
             value: metaFailedFactor * 100,
             title: "Failure",
-            titleValue: `${ Math.ceil(session.phaseFailedCount / (sessionTotalCount || 1) * 100)}%`,
+            titleValue: `${ Math.ceil(session.phasesFailedCount / (sessionTotalCount || 1) * 100)}%`,
             color: statusColors.get(StatusColor.Failure)!,
             stripes: true
         }
@@ -146,9 +187,9 @@ export const StatusPie = (stack: StatusBarStack[], radius: number, style?: any):
 }
 
 export const SessionStatusPie = (session: TestSessionResult, radius: number): JSX.Element | undefined => {
-    const sessionTotalCount = session.phaseSucceededCount + session.phaseFailedCount + session.phaseUnspecifiedCount;
-    const metaFailedFactor = Math.ceil(session.phaseFailedCount / (sessionTotalCount || 1) * 100) / 100;
-    const metaUnspecifiedFactor = Math.ceil(session.phaseUnspecifiedCount / (sessionTotalCount || 1) * 100) / 100;
+    const sessionTotalCount = session.phasesSucceededCount + session.phasesFailedCount + session.phasesUnspecifiedCount;
+    const metaFailedFactor = Math.ceil(session.phasesFailedCount / (sessionTotalCount || 1) * 100) / 100;
+    const metaUnspecifiedFactor = Math.ceil(session.phasesUnspecifiedCount / (sessionTotalCount || 1) * 100) / 100;
 
     const bDisplayBar = (metaFailedFactor + metaUnspecifiedFactor) > 0 && (metaFailedFactor + metaUnspecifiedFactor) < 1;
     if (!bDisplayBar) return;
@@ -178,8 +219,8 @@ export const SessionStatusPie = (session: TestSessionResult, radius: number): JS
 
 export const PhasesStatusPie = (phases: TestPhaseStatus[], radius: number): JSX.Element | undefined => {
     let totalCount = phases.length || 1;
-    let phaseFailedCount = 0;
-    let phaseUnspecifiedCount = 0;
+    let phasesFailedCount = 0;
+    let phasesUnspecifiedCount = 0;
     phases.forEach(p => {
         switch (p.outcome) {
             case TestPhaseOutcome.Skipped:
@@ -187,17 +228,17 @@ export const PhasesStatusPie = (phases: TestPhaseStatus[], radius: number): JSX.
                 break;
             case TestPhaseOutcome.Failed:
             case TestPhaseOutcome.Interrupted:
-                ++phaseFailedCount;
+                ++phasesFailedCount;
                 break;
             case TestPhaseOutcome.Unknown:
             case TestPhaseOutcome.NotRun:
-                ++phaseUnspecifiedCount;
+                ++phasesUnspecifiedCount;
                 break;
         }
     });
 
-    const failedFactor = Math.ceil(phaseFailedCount / totalCount * 100) / 100;
-    const unspecifiedFactor = Math.ceil(phaseUnspecifiedCount / totalCount * 100) / 100;
+    const failedFactor = Math.ceil(phasesFailedCount / totalCount * 100) / 100;
+    const unspecifiedFactor = Math.ceil(phasesUnspecifiedCount / totalCount * 100) / 100;
 
     const statusColors = dashboard.getStatusColors();
 
@@ -225,8 +266,65 @@ export const PhasesStatusPie = (phases: TestPhaseStatus[], radius: number): JSX.
 export const SessionValues = (session: TestSessionResult, style?: any): JSX.Element => {
     const statusColors = dashboard.getStatusColors();
     return <Text style={style}>
-                {!!session.phaseUnspecifiedCount && <span style={{color: statusColors.get(StatusColor.Unspecified)}}> Unspecified: {session.phaseUnspecifiedCount}</span>}
-                {!!session.phaseFailedCount && <span style={{color: statusColors.get(StatusColor.Failure)}}> Failed: {session.phaseFailedCount}</span>}
-                {!!session.phaseSucceededCount && <span style={{color: statusColors.get(StatusColor.Success)}}> Success: {session.phaseSucceededCount}</span>}
+                {!!session.phasesUnspecifiedCount && <span style={{color: statusColors.get(StatusColor.Unspecified)}}> Unspecified: {session.phasesUnspecifiedCount}</span>}
+                {!!session.phasesFailedCount && <span style={{color: statusColors.get(StatusColor.Failure)}}> Failed: {session.phasesFailedCount}</span>}
+                {!!session.phasesSucceededCount && <span style={{color: statusColors.get(StatusColor.Success)}}> Success: {session.phasesSucceededCount}</span>}
             </Text>
 }
+
+export const getPhaseSessionStatusColor = (phase: PhaseSessionResult | TestPhaseStatus) => {
+    const statusColors = dashboard.getStatusColors();
+    let color = statusColors.get(StatusColor.Unspecified)!
+    switch(phase.outcome) {
+        case TestPhaseOutcome.Skipped:
+            color = statusColors.get(StatusColor.Skipped)!
+            break;
+        case TestPhaseOutcome.Failed:
+            color = statusColors.get(StatusColor.Failure)!
+            break;
+        case TestPhaseOutcome.Success:
+            color = statusColors.get(phase.hasWarning? StatusColor.Warnings : StatusColor.Success)!
+            break;
+    }
+
+    return color;
+}
+
+export const getTestSessionStatusColor = (test: TestSessionResult | TestSessionDetails) => {
+    const statusColors = dashboard.getStatusColors();
+    let color = statusColors.get(StatusColor.Success);
+    if (test.outcome === TestOutcome.Failure) {
+        color = statusColors.get(StatusColor.Failure);
+    } else if (test.outcome === TestOutcome.Unspecified) {
+        color = statusColors.get(StatusColor.Unspecified);
+    }
+
+    return color;
+}
+
+export const PhaseHistoryGraphWidget: React.FC<{test: TestNameRef, phaseKey: string, meta: MetadataRef, streamId: string, sessionId?: string, handler: TestDataHandler, onClick: PhaseSessionCallback, withoutHeader?: boolean}> = ({ test, phaseKey, meta, streamId, sessionId, handler, onClick, withoutHeader }) => {
+    const [container, setContainer] = useState<HTMLDivElement | null>(null);
+    const [state, setState] = useState<{ graph?: PhaseHistoryGraph }>({});
+
+    useEffect(() => {
+        const graph = new PhaseHistoryGraph(test, phaseKey, streamId, handler, onClick);
+        setState({ graph: graph });
+    }, [test, phaseKey, streamId, onClick]);
+
+    useEffect(() => {
+        // refresh on meta or sessionId update
+        try {
+            container && state.graph?.render(container, meta, sessionId, withoutHeader);
+        } catch (err) {
+            console.error(err);
+        }
+    }, [container, state, meta, sessionId])
+
+    if (!state.graph) {
+        return null;
+    }
+
+    const graph_container_id = `test_history_graph_container_${test!.id}_${meta.id}_${phaseKey}_${streamId}`;
+
+    return <div id={graph_container_id} style={{ userSelect: "none", position: "relative" }} ref={(ref: HTMLDivElement) => setContainer(ref)}/>
+}
diff --git a/Engine/Source/Programs/Horde/HordeDashboard/plugins/test-automation/testAutomationEvents.tsx b/Engine/Source/Programs/Horde/HordeDashboard/plugins/test-automation/testAutomationEvents.tsx
new file mode 100644
--- /dev/null
+++ b/Engine/Source/Programs/Horde/HordeDashboard/plugins/test-automation/testAutomationEvents.tsx
@@ -0,0 +1,261 @@
+import { Checkbox, DefaultButton, DirectionalHint, FontIcon, HoverCard, HoverCardType, Image, IPlainCardProps, mergeStyles, mergeStyleSets, Spinner, SpinnerSize, Stack, Text } from "@fluentui/react";
+import { EventEntry, EventLevel, ImageCompare, LogPropertyType } from "./testEventsModel";
+import { TestPhaseOutcome } from "./api";
+import { memo, useCallback, useEffect, useState } from "react";
+import { ArtifactFactory, DeviceRef, TestPhaseStatus } from "./testData";
+import dashboard, { StatusColor } from "horde/backend/Dashboard";
+import { getShortNiceTime, msecToElapsed } from "horde/base/utilities/timeUtils";
+import { DashboardPreference } from "horde/backend/Api";
+import { isBright } from "horde/base/utilities/colors";
+import { styles } from "./testAutomationCommon";
+
+let eventStyles: any = undefined;
+export const getEventStyling = () => {
+    if (eventStyles) return eventStyles;
+
+    const colors = dashboard.getStatusColors();
+    const errorColorPref = dashboard.getPreference(DashboardPreference.ColorError);
+    const errorHighlight = errorColorPref ? isBright(errorColorPref) ? "brightness(0.9)" : "brightness(1.2)" : "";
+
+    const eventClass = mergeStyles({
+        fontSize: "12px",
+        fontFamily: "Horde Cousine Regular",
+        paddingBottom: 3,
+        paddingTop: 3,
+        paddingRight: 5,
+        paddingLeft: 3
+    });
+
+    eventStyles = mergeStyleSets({
+        event: eventClass,
+        eventWarning: [{
+            background: dashboard.darktheme ? "#1E1817" : "#FEF8E7",
+            borderLeftStyle: 'solid',
+            borderLeftColor: colors.get(StatusColor.Warnings)!
+        }, eventClass],
+        eventError: [{
+            background: dashboard.darktheme ? "#330606" : "#FEF6F6",
+            borderLeftStyle: 'solid',
+            borderLeftColor: colors.get(StatusColor.Failure)!
+        }, eventClass],
+        failureButton: {
+            backgroundColor: errorColorPref ? errorColorPref : dashboard.darktheme ? "#9D1410" : "#EC4C47",
+            borderStyle: "hidden",
+            color: "#FFFFFF",
+            selectors: {
+                ':active,:hover': {
+                    color: "#FFFFFF",
+                    backgroundColor: errorColorPref ? errorColorPref : "#DC3C37",
+                    filter: errorHighlight
+                }
+            }
+        },
+        failureButtonDisabled: {         
+            borderStyle: "hidden",
+            color: dashboard.darktheme ? "#909398" : "#616E85 !important",
+            backgroundColor: dashboard.darktheme ? "#1F2223" : "#f3f2f1"
+        },
+    });
+
+    return eventStyles;
+}
+
+/// Main event entry point
+const Event: React.FC<{ entry: EventEntry, phase: TestPhaseStatus }> = memo(({entry, phase}) => {
+    const eventStyles = getEventStyling();
+
+    const isError = entry.level === EventLevel.Error || entry.level === EventLevel.Critical;
+    const isWarning = entry.level === EventLevel.Warning;
+    const style = isError ? eventStyles.eventError : (isWarning ? eventStyles.eventWarning : eventStyles.event);
+
+    return <Stack className={style}>
+                {BasicEvent(entry)}
+                {!!entry.properties && Object.values(entry.properties).map(property => !!property && EventTypeChooser(property, phase))}
+            </Stack>
+});
+
+const EventTypeChooser = (property: string | number | Record<string, string | number | undefined>, phase: TestPhaseStatus): JSX.Element | undefined => {
+    switch (property["$type"]) {
+        case LogPropertyType.ImageCompare:
+            return <ImageComparer imageCompare={property as ImageCompare} phase={phase} />
+    }
+
+    return;
+}
+
+const BasicEvent = (entry: EventEntry): JSX.Element => {
+    return <pre style={{ margin: 0, whiteSpace: "pre-wrap" }}>[{entry.time}] {entry.level}: {entry.message}</pre>
+}
+
+/// Event pane
+const OnDeviceCard = (device: DeviceRef): JSX.Element => {
+    const meta = device.metadata;
+    return <Stack style={{padding: 6}}>
+            { meta && <Stack horizontal tokens={{childrenGap: 4}}>
+                        <Stack tokens={{childrenGap: 4}}>
+                            {Object.keys(meta).map((key, i) => <Stack style={{fontWeight: "bold"}} key={`device-${device.name}-key-${i}`}>{key}: </Stack>)}
+                        </Stack>
+                        <Stack tokens={{childrenGap: 4}}>
+                            {Object.values(meta).map((value, i) => <Stack key={`device-${device.name}-value-${i}`}>{value}</Stack>)}
+                        </Stack>
+                </Stack>
+            }
+            { !meta && <Text>no metadata</Text> }
+        </Stack>;
+}
+
+const DeviceItem = (device: DeviceRef): JSX.Element => {
+    const deviceCardProps: IPlainCardProps = {
+        renderData: device,
+        onRenderPlainCard: OnDeviceCard,
+        directionalHint: DirectionalHint.topRightEdge,
+        gapSpace: 8,
+        calloutProps: { isBeakVisible: true }
+    };
+
+    return <HoverCard plainCardProps={deviceCardProps} type={HoverCardType.plain} cardOpenDelay={20} key={device.name}>
+                <Text style={{fontSize: 12, cursor: 'pointer' }}>{device.name}</Text>
+            </HoverCard>
+}
+
+export const PhaseEventsPane: React.FC<{phase: TestPhaseStatus, onEventsLoaded?: () => void}> = ({phase, onEventsLoaded}) => {
+    const [events, setEvents] = useState<EventEntry[] | undefined>(undefined);
+    const [filterError, setFilterError] = useState<boolean>(true);
+    const [filterWarning, setFilterWarning] = useState<boolean>(true);
+    const [filterInfo, setFilterInfo] = useState<boolean>(false);
+
+    useEffect(() => {
+        // reset events and filtering
+        if (events !== undefined) setEvents(undefined);
+        if (!filterInfo && phase.outcome !== TestPhaseOutcome.Failed) {
+            setFilterInfo(true);
+        }
+        // fetch the events
+        phase.fetchEventStream().then((value) => {
+            setEvents(value ?? []);
+        }).catch((reason) => setEvents([{level: EventLevel.Error, message: `Could not retrieve test events: ${reason}`, format: "", time: "", id: 0}]));
+    }, [phase]);
+
+    useEffect(() => {
+        if (events) onEventsLoaded?.();
+    }, [events]);
+
+    const isEntryNeedDisplay = useCallback((event: EventEntry): boolean => {
+        const eventType = event.level;
+        if (filterError && eventType === EventLevel.Error) {
+            return true;
+        }
+        if (filterWarning && eventType === EventLevel.Warning) {
+            return true;
+        }
+        if (filterInfo && eventType === EventLevel.Information) {
+            return true;
+        }
+
+        return false;
+    }, [filterError, filterWarning, filterInfo]);
+
+    const borderColor = dashboard.darktheme ? "#4D4C4B" : "#6D6C6B";
+
+    return <Stack style={{borderWidth: 1, borderColor: borderColor, borderStyle: 'solid', padding: 4, overflow: 'hidden'}} grow>
+                {!events &&
+                    <Stack horizontalAlign='center' style={{ padding: 12}} grow>
+                        <Spinner size={SpinnerSize.medium} />
+                    </Stack>
+                }
+                {!!events &&
+                    <Stack style={{overflow: 'hidden'}} tokens={{childrenGap: 3}} grow>
+                        <Stack style={{borderWidth: '0px 0px 1px 0px', borderColor: borderColor, borderStyle: 'solid', paddingBottom: 2}} horizontal tokens={{childrenGap: 8}} disableShrink>
+                            <Text><span style={styles.labelSmall}>Run on: </span><span>{getShortNiceTime(phase.start, true, true, true)}</span></Text>
+                            <Text><span style={styles.labelSmall}> For: </span><span>{msecToElapsed(phase.duration * 1000, true, true)}</span></Text>
+                            {!!phase.devices.length &&
+                                <Stack horizontal tokens={{childrenGap: 4}} verticalAlign="end">
+                                    <Text style={styles.labelSmall}> Devices: </Text>
+                                    {phase.devices.map(d => DeviceItem(d))}
+                                </Stack>
+                            }
+                            <Stack verticalAlign="center" horizontal tokens={{ childrenGap: 5 }} style={{borderWidth: '0px 1px', paddingLeft: 9, paddingRight: 9, borderStyle: 'solid', borderColor: borderColor}}>
+                                <FontIcon iconName="PageListFilter" title="Events filtering" style={styles.icon} />
+                                <Checkbox styles={styles.defaultCheckbox} label="Error" checked={filterError} onChange={(_: any, checked?: boolean) => setFilterError(!!checked)} disabled={!events.length}/>
+                                <Checkbox styles={styles.defaultCheckbox} label="Warning" checked={filterWarning} onChange={(_: any, checked?: boolean) => setFilterWarning(!!checked)} disabled={!events.length}/>
+                                <Checkbox styles={styles.defaultCheckbox} label="Info" checked={filterInfo} onChange={(_: any, checked?: boolean) => setFilterInfo(!!checked)} disabled={!events.length}/>
+                            </Stack>
+                            <Stack horizontalAlign="end">
+                                <DefaultButton disabled={true} style={{ fontSize: 11, padding: 0, height: 18, backgroundColor: styles.defaultButton.backgroundColor, color: 'white' }}
+                                    onClick={(ev) => {
+                                        ev.stopPropagation();
+                                    }}>Artifacts</DefaultButton>
+                            </Stack>
+                        </Stack>
+                        <Stack style={{ overflowX: 'hidden', overflowY: 'auto' }}>
+                            { events.map((e, i) => isEntryNeedDisplay(e) && <Event key={`phase-${phase.key}-event-${i}`} entry={e} phase={phase}/>)}
+                            { !events.length && <Text style={styles.textSmall}>No events</Text>}
+                        </Stack>
+                    </Stack>
+                }
+            </Stack>
+}
+
+/// Image comparison widget
+const missingImage = "/images/missing-image.png";
+const MissingImageLabel = (): JSX.Element => { return <span style={{ fontWeight: 'bold' }}> [missing image]</span> }
+type ImageData = { link?: string, ref?: string }
+type ImageLinks = {
+    unapproved: ImageData;
+    approved?: ImageData;
+    difference?: ImageData;
+}
+
+const getImageLinks = async (imageCompare: ImageCompare, artifacts: ArtifactFactory): Promise<ImageLinks> => {
+    let url = await artifacts.getLink(imageCompare.unapproved);
+    const imageInfo: ImageLinks = {unapproved: { link: url, ref: imageCompare.unapproved }};
+    if (imageCompare.approved) {
+        url = await artifacts.getLink(imageCompare.approved);
+        imageInfo.approved = { link: url, ref: imageCompare.approved };
+    }
+    if (imageCompare.difference) {
+        url = await artifacts.getLink(imageCompare.difference);
+        imageInfo.difference = { link: url, ref: imageCompare.difference };
+    }
+
+    return imageInfo;
+}
+
+const ImageComparer: React.FC<{imageCompare: ImageCompare, phase: TestPhaseStatus}> = memo(({imageCompare, phase}) => {
+    const [imageInfo, setImageInfo] = useState<ImageLinks | undefined>(undefined);
+
+    useEffect(() => {
+        const phaseArtifacts = phase.artifacts;
+        !!phaseArtifacts && getImageLinks(imageCompare, phaseArtifacts).then(info => setImageInfo(info));
+    }, [imageCompare, phase])
+
+    return  <Stack styles={{ root: { paddingLeft: 16, width: '100%' } }}>
+                <Stack><Text variant="medium"><span style={{ fontWeight: "bold" }}>{imageCompare['$type']}: </span>{imageCompare['$text']}</Text></Stack>
+                {imageInfo &&
+                    <Stack horizontal>
+                        {imageInfo.approved &&
+                            <Stack styles={{ root: { padding: 5 } }}>
+                                <a href={imageInfo.approved.link}>
+                                    <Image width="100%" style={{ minWidth: 200, maxWidth: 400, minHeight: 120 }} src={imageInfo.approved.link || missingImage} alt={imageInfo.approved.ref} />
+                                </a>
+                                <Stack.Item align="center">Reference{!imageInfo.approved && MissingImageLabel()}</Stack.Item>
+                            </Stack>
+                        }
+                        {imageInfo.difference &&
+                            <Stack styles={{ root: { padding: 5 } }}>
+                                <a href={imageInfo.difference.link}>
+                                    <Image width="100%" style={{ minWidth: 200, maxWidth: 400, minHeight: 120 }} src={imageInfo.difference.link || missingImage} alt={imageInfo.difference.ref} />
+                                </a>
+                                <Stack.Item align="center">Difference{!imageInfo.difference && MissingImageLabel()}</Stack.Item>
+                            </Stack>
+                        }
+                        <Stack styles={{ root: { padding: 5 } }}>
+                            <a href={imageInfo.unapproved.link}>
+                                <Image width="100%" style={{ minWidth: 200, maxWidth: 400, minHeight: 120 }} src={imageInfo.unapproved.link || missingImage} alt={imageInfo.unapproved.ref} />
+                            </a>
+                            <Stack.Item align="center">Produced{!imageInfo.unapproved && MissingImageLabel()}</Stack.Item>
+                        </Stack>
+                    </Stack>
+                }
+            </Stack>
+});
diff --git a/Engine/Source/Programs/Horde/HordeDashboard/plugins/test-automation/testAutomationPhaseCompare.tsx b/Engine/Source/Programs/Horde/HordeDashboard/plugins/test-automation/testAutomationPhaseCompare.tsx
new file mode 100644
--- /dev/null
+++ b/Engine/Source/Programs/Horde/HordeDashboard/plugins/test-automation/testAutomationPhaseCompare.tsx
@@ -0,0 +1,243 @@
+import { getHordeTheme } from "horde/styles/theme";
+import { getJobParams, PhaseComparisonContext } from "./phaseCompare";
+import { MetadataRef, PhaseSessionResult, TestDataHandler, TestNameRef } from "./testData";
+import { observer } from "mobx-react-lite";
+import { useCallback, useEffect, useMemo } from "react";
+import { ComboBox, FontIcon, IComboBoxOption, Link, Spinner, SpinnerSize, Stack, Text } from "@fluentui/react";
+import { getPhaseSessionStatusColor, PhaseHistoryGraphWidget } from "./testAutomationCommon";
+import { projectStore } from "horde/backend/ProjectStore";
+import { ChangeButton } from "horde/components/ChangeButton";
+import { PhaseEventsPane } from "./testAutomationEvents";
+
+export const comparisonContext = new PhaseComparisonContext();
+
+const onClickItem = (session: PhaseSessionResult) => {
+    comparisonContext.setTestSesssionId(session.sessionId);
+    comparisonContext.setPhaseStatus(undefined);
+}
+
+export const PhaseComparisonPane: React.FC<{ test: TestNameRef, meta?: MetadataRef, phaseKey: string, streamId?: string, focusedSessionId?: string, handler: TestDataHandler, onDismiss: () => void }> = observer(({test, meta, phaseKey, streamId, focusedSessionId, handler, onDismiss}) => {
+    const hordeTheme = getHordeTheme();
+
+    comparisonContext.subscribe();
+    handler.subscribeToPhaseQueryLoading();
+    handler.subscribeToSubPhaseQueryLoading();
+
+    const phaseSessionMetaStatus = useMemo(
+        () => new Map(handler.selectedStatusStream?.tests.get(test)?.sessions.entries().filter(([_, status]) => status.phaseHistories.get(phaseKey)?.length ?? 0 > 0))
+        , [test, phaseKey, handler.phaseQueryLoading, handler.subPhaseQueryLoading]);
+    const commonMetaKeys =  useMemo(
+        () => phaseSessionMetaStatus.size === 1 ? [] : MetadataRef.identifyCommonKeys(phaseSessionMetaStatus.keys().toArray())
+        , [test, phaseKey, handler.phaseQueryLoading, handler.subPhaseQueryLoading]);
+    const metas = useMemo(
+        () => new Map(
+            phaseSessionMetaStatus.entries().map<[string, {ref: MetadataRef, session: PhaseSessionResult}]>(
+                ([meta, status]) => [meta.getValuesExcept(commonMetaKeys).join(' / '), {ref: meta, session: status.phaseHistories.get(phaseKey)![0]}]
+            )
+        )
+        , [test, phaseKey, handler.phaseQueryLoading, handler.subPhaseQueryLoading]);
+    const options: IComboBoxOption[] = useMemo(
+        () => metas.keys().toArray().sort((a, b) => a.localeCompare(b)).map(key => ({key: key.toLowerCase(), text: key, data: metas.get(key)}))
+        , [test, phaseKey, handler.phaseQueryLoading, handler.subPhaseQueryLoading]);
+
+    const onRenderOption = useCallback((item: IComboBoxOption) => {
+        if (!item) return <div/>;
+        const color = getPhaseSessionStatusColor(item.data.session);
+        return <Stack style={{ paddingLeft: 4 }}>
+                    <Stack horizontal tokens={{ childrenGap: 8 }}>
+                        <FontIcon style={{ fontSize: 12, color: color }} iconName="Square" />
+                        <Text style={{ fontSize: 11, fontStyle: item.data!.ref === meta ? 'italic' : 'normal' }}>{item.text}</Text>
+                    </Stack>
+                </Stack>
+    }, [meta]);
+
+    const onClickSwap = useCallback(() => {
+        const compareTestDataId = comparisonContext.phaseStatus?.testSession?.testDataId;
+        if (!!compareTestDataId && !!meta) {
+            const selectedMetaHistory = handler.selectedStatusStream?.tests.get(test)?.sessions.get(meta)?.getPhaseSessions(phaseKey);
+            const focusedSession = selectedMetaHistory?.find(s => s.id === focusedSessionId);
+            if (focusedSession) {
+                comparisonContext.setMetaKey(meta.getValuesExcept(commonMetaKeys).join(' / ').toLowerCase());
+                comparisonContext.setTestSesssionId(focusedSession.sessionId);
+                comparisonContext.setPhaseStatus(undefined);
+                handler.setSearchParam('session', compareTestDataId);
+                return;
+            }
+        }
+    }, [test, meta, phaseKey, focusedSessionId])
+
+    const autoSelectSession = useCallback(() => {
+        const selectedMeta = options.find(o => o.key === comparisonContext.metaKey)!.data!.ref as MetadataRef;
+        const selectedMetaHistory = handler.selectedStatusStream?.tests.get(test)?.sessions.get(selectedMeta)?.getPhaseSessions(phaseKey);
+        if (!!selectedMetaHistory && selectedMetaHistory.length > 0) {
+            // select matching commit Id
+            if (!!comparisonContext.commitId) {
+                const commitIndex = selectedMetaHistory.findIndex(s => s.commitId === comparisonContext.commitId);
+                if (commitIndex !== -1) {
+                    comparisonContext.setTestSesssionId(selectedMetaHistory[commitIndex].sessionId);
+                    comparisonContext.setPhaseStatus(undefined);
+                    return;
+                }
+            }
+            // select latest
+            comparisonContext.setTestSesssionId(selectedMetaHistory[0].sessionId)
+            return;
+        }
+        comparisonContext.setTestSesssionId(undefined);
+    }, [test, phaseKey, options]);
+
+    useEffect(() => {
+        if (comparisonContext.testSessionId && comparisonContext.metaKey && !comparisonContext.phaseStatus && !comparisonContext.isLoadingPhase) {
+            const meta = options.find(o => o.key === comparisonContext.metaKey)?.data.ref;
+            const metaStatus = handler.getStatusStream(streamId!)?.tests.get(test)?.sessions.get(meta);
+            const testSession = metaStatus?.history.find(s => s.id === comparisonContext.testSessionId);
+            if (!!testSession?.testDataId) {
+                comparisonContext.setLoadingPhase(true);
+                handler.queryTestData(testSession.testDataId).then(
+                    test => {
+                        if (!!test) {
+                            comparisonContext.jobParams = getJobParams(test, handler);
+                            comparisonContext.commitId = test.commitId;
+                            const phase = test.phases.find(p => p.key === phaseKey);
+                            if (!!phase) {
+                                handler.queryPhase(phase, streamId!, true).then(() => comparisonContext.setPhaseStatus(phase));
+                            }
+                        }
+                    });
+            }
+        }
+    }, [comparisonContext.updated, options]);
+
+    useEffect(() => {
+        // auto selection
+        if (!!comparisonContext.metaKey) {
+            if (!options.find(o => o.key === comparisonContext.metaKey)) {
+                // invalid meta key
+                comparisonContext.reset();
+                comparisonContext.setEnable();
+            }
+        }
+        if (!!meta && !comparisonContext.metaKey) {
+            const selectedMetaKey = meta.getValuesExcept(commonMetaKeys).join(' / ');
+            if (metas.has(selectedMetaKey)) {
+                comparisonContext.setMetaKey(selectedMetaKey.toLowerCase());
+            }
+        }
+        if (!!comparisonContext.metaKey) {
+            const selectedMeta = options.find(o => o.key === comparisonContext.metaKey)!.data!.ref as MetadataRef;
+            const selectedMetaHistory = handler.selectedStatusStream?.tests.get(test)?.sessions.get(selectedMeta)?.getPhaseSessions(phaseKey);
+            if (!!selectedMetaHistory && selectedMetaHistory.length > 0) {
+                if (!!comparisonContext.commitId) {
+                    // get matching commit
+                    const commitIndex = selectedMetaHistory.findIndex(s => s.commitId === comparisonContext.commitId);
+                    if (commitIndex !== -1) {
+                        comparisonContext.setTestSesssionId(selectedMetaHistory[commitIndex].sessionId)
+                        comparisonContext.setPhaseStatus(undefined);
+                        return;
+                    }
+                }
+                if (!!focusedSessionId) {
+                    const focusedIndex = selectedMetaHistory.findIndex(s => s.id === focusedSessionId);
+                    if (focusedIndex !== -1) {
+                        const errorFingerPrint = selectedMetaHistory[focusedIndex].errorFingerprint;
+                        if (!!errorFingerPrint) {
+                            // get previous error
+                            const previousErrorIndex = selectedMetaHistory.findIndex(s => s.errorFingerprint !== errorFingerPrint);
+                            if (previousErrorIndex > focusedIndex) {
+                                comparisonContext.setTestSesssionId(selectedMetaHistory[previousErrorIndex].sessionId)
+                                comparisonContext.setPhaseStatus(undefined);
+                                return;
+                            }
+                        }
+                        // get previous available session
+                        if (focusedIndex + 1 < selectedMetaHistory.length) {
+                            comparisonContext.setTestSesssionId(selectedMetaHistory[focusedIndex + 1].sessionId)
+                            comparisonContext.setPhaseStatus(undefined);
+                            return;
+                        }
+                    }
+                }
+            }
+            autoSelectSession();
+            comparisonContext.setPhaseStatus(undefined);
+        }
+    }, [meta, test, phaseKey, focusedSessionId, options])
+
+    const selectedMeta = !!comparisonContext.metaKey ? options.find(o => o.key === comparisonContext.metaKey)?.data?.ref as MetadataRef : undefined;
+
+    return <Stack style={{overflow: 'hidden', borderWidth: '1px 0px 0px 0px', borderStyle: 'solid', borderColor: "#035ca1", paddingTop: 5}} grow>
+                {!!streamId &&
+                    <Stack disableShrink tokens={{childrenGap: 3}}>
+                        <Stack horizontal verticalAlign="center">
+                            <ComboBox styles={{ root: { width: 210, height: 22, fontSize: 12, marginRight: 9 }, input: {color: hordeTheme.semanticColors.link, fontWeight: 600} }} placeholder={"Compare with"} allowFreeform={true} autoComplete="on" disabled={!options.length}                                
+                                selectedKey={comparisonContext.metaKey}
+                                onChange={(_1, option, _2, value) => {
+                                    let selectedKey = option?.key as string;
+                                    if (value) {
+                                        // find first partial match
+                                        const lowerCaseValue = value.toLowerCase();
+                                        selectedKey = options.find(o => (o.key as string).includes(lowerCaseValue))?.key as string;
+                                    }
+                                    comparisonContext.setMetaKey(selectedKey);
+                                    comparisonContext.setPhaseStatus(undefined);
+                                    comparisonContext.jobParams = undefined;
+                                    autoSelectSession();
+                                }}
+                                options={options}
+                                onRenderOption={onRenderOption}
+                            />
+
+                            {!!comparisonContext.testSessionId && !!comparisonContext.phaseStatus?.testSession && 
+                                <Stack style={{ paddingLeft: 8 }} horizontal tokens={{childrenGap: 8}} verticalAlign="center">
+                                    <Text>{projectStore.streamById(comparisonContext.phaseStatus.testSession.streamId)?.fullname ?? comparisonContext.phaseStatus.testSession.streamId}</Text>
+                                    <ChangeButton job={comparisonContext.jobParams?.job} rangeCL={comparisonContext.jobParams?.rangeCL} />
+                                    <Link href={`/job/${comparisonContext.jobParams?.job?.id}?step=${comparisonContext.jobParams?.step?.id}`} target="_blank"><Text style={{ fontSize: 12, color: hordeTheme.semanticColors.link }}>View Job</Text></Link>
+                                </Stack>
+                            }
+
+                            <Stack grow horizontalAlign="end">
+                                <FontIcon
+                                    style={{ fontSize: 13, cursor: "pointer", padding: 5 }}
+                                    title="Close Comparison View"
+                                    iconName={"Cancel"}
+                                    onClick={onDismiss}
+                                />
+                            </Stack>
+                        </Stack>
+                        {!!comparisonContext.metaKey && !!selectedMeta &&
+                            <Stack style={{ paddingLeft: 3 }} horizontal verticalAlign="end">
+                                <Stack style={{ marginLeft: 30, marginBottom: 18, width: 10, cursor: "default" }} title="History">
+                                    <FontIcon style={{ fontSize: 15 }} iconName={"History"}/>
+                                </Stack>
+                                <PhaseHistoryGraphWidget
+                                    test={test}
+                                    streamId={streamId}
+                                    phaseKey={phaseKey}
+                                    sessionId={comparisonContext.phaseStatus?.session?.id}
+                                    meta={selectedMeta}
+                                    onClick={onClickItem}
+                                    handler={handler}
+                                />
+                                <Stack grow horizontalAlign="end" style={{ marginBottom: 16, marginRight: 16}}>
+                                    <FontIcon
+                                        style={{ fontSize: 15, cursor: "pointer" , backgroundColor: "#035ca1", borderRadius: 4, padding: 7, color: 'white' }}
+                                        title="Toggle Comparison View"
+                                        iconName={"SwitcherStartEnd"}
+                                        onClick={onClickSwap}
+                                    />
+                                </Stack>
+                            </Stack>
+                        }
+                    </Stack>
+                }
+                {!!comparisonContext.testSessionId && !!comparisonContext.phaseStatus &&
+                    <PhaseEventsPane phase={comparisonContext.phaseStatus}/>
+                }
+                {!!comparisonContext.testSessionId && !comparisonContext.phaseStatus && comparisonContext.isLoadingPhase &&
+                    <Stack horizontalAlign='center' style={{ padding: 12}} grow>
+                        <Spinner size={SpinnerSize.medium} />
+                    </Stack>
+                }
+                {!options.length &&  <Stack horizontalAlign='center' style={{ padding: 12}} grow><Text>No phase matched in session history</Text></Stack>}
+            </Stack>
+});
diff --git a/Engine/Source/Programs/Horde/HordeDashboard/plugins/test-automation/testAutomationPhaseView.tsx b/Engine/Source/Programs/Horde/HordeDashboard/plugins/test-automation/testAutomationPhaseView.tsx
--- a/Engine/Source/Programs/Horde/HordeDashboard/plugins/test-automation/testAutomationPhaseView.tsx
+++ b/Engine/Source/Programs/Horde/HordeDashboard/plugins/test-automation/testAutomationPhaseView.tsx
@@ -1,351 +1,180 @@
-import { IconButton, Stack, Text, Spinner, SpinnerSize, Link, FontIcon, IPlainCardProps, DirectionalHint, HoverCard, HoverCardType, mergeStyleSets, mergeStyles, DefaultButton, TooltipHost, Icon, ITag, TagPicker, IPickerItemProps, Label, List, Image, Checkbox } from "@fluentui/react";
+import { IconButton, Stack, Text, Spinner, SpinnerSize, Link, FontIcon, DefaultButton, TooltipHost, Icon, ITag, TagPicker, IPickerItemProps, Label, List } from "@fluentui/react";
 import { useState, useEffect, useCallback, useRef, memo, useMemo } from "react";
 import { observer } from "mobx-react-lite";
-import { TestDataHandler, TestSessionDetails, TestPhaseStatus, DeviceRef, MetadataRef } from "./testData";
-import { EventEntry, EventLevel, ImageCompare, LogPropertyType, TestPhaseOutcome } from './api';
-import dashboard, { StatusColor } from "horde/backend/Dashboard";
-import { ChangeButton, JobParameters } from "horde/components/ChangeButton";
+import { TestDataHandler, TestSessionDetails, TestPhaseStatus, MetadataRef, TestNameRef } from "./testData";
+import { TestPhaseOutcome } from './api';
+import dashboard from "horde/backend/Dashboard";
+import { ChangeButton } from "horde/components/ChangeButton";
 import { getShortNiceTime, msecToElapsed } from "horde/base/utilities/timeUtils";
 import { getHordeTheme } from "horde/styles/theme";
-import { DashboardPreference } from "horde/backend/Api";
+import { DashboardPreference, GetArtifactResponse } from "horde/backend/Api";
 import { isBright } from "horde/base/utilities/colors";
-import { PhasesStatusPie } from "./testAutomationCommon";
+import { getPhaseSessionStatusColor, PhaseHistoryGraphWidget, PhasesStatusPie, styles } from "./testAutomationCommon";
 import { projectStore } from "horde/backend/ProjectStore";
-
-const styles = {
-    labelSmall: {
-        fontSize: 12,
-        fontWeight: "bold"
-    },
-    textSmall: {
-        fontSize: 12
-    },
-    textLarge: {
-        fontSize: 24
-    },
-    defaultButton: {
-        backgroundColor: "#035ca1"
-    },
-    defaultFilter: {
-        backgroundColor: "#0078d4"
-    },
-    defaultCheckbox: {
-        root: {
-            height: 15
-        },
-        text: {
-            fontSize: 11,
-            marginLeft: 0,
-            lineHeight: 15
-        },
-        checkbox: {
-            height: 15,
-            width: 15
-        }
-    },
-    icon: {
-        cursor: "default",
-        height: 15,
-        paddingTop: 1
-    },
-}
-
-let eventStyles: any = undefined;
-const getEventStyling = () => {
-    if (eventStyles) return eventStyles;
-
-    const colors = dashboard.getStatusColors();
-    const errorColorPref = dashboard.getPreference(DashboardPreference.ColorError);
-    const errorHighlight = errorColorPref ? isBright(errorColorPref) ? "brightness(0.9)" : "brightness(1.2)" : "";
-
-    const eventClass = mergeStyles({
-        fontSize: "12px",
-        fontFamily: "Horde Cousine Regular",
-        paddingBottom: 3,
-        paddingTop: 3,
-        paddingRight: 5,
-        paddingLeft: 3
-    });
-
-    eventStyles = mergeStyleSets({
-        event: eventClass,
-        eventWarning: [{
-            background: dashboard.darktheme ? "#1E1817" : "#FEF8E7",
-            borderLeftStyle: 'solid',
-            borderLeftColor: colors.get(StatusColor.Warnings)!
-        }, eventClass],
-        eventError: [{
-            background: dashboard.darktheme ? "#330606" : "#FEF6F6",
-            borderLeftStyle: 'solid',
-            borderLeftColor: colors.get(StatusColor.Failure)!
-        }, eventClass],
-        failureButton: {
-            backgroundColor: errorColorPref ? errorColorPref : dashboard.darktheme ? "#9D1410" : "#EC4C47",
-            borderStyle: "hidden",
-            color: "#FFFFFF",
-            selectors: {
-                ':active,:hover': {
-                    color: "#FFFFFF",
-                    backgroundColor: errorColorPref ? errorColorPref : "#DC3C37",
-                    filter: errorHighlight
-                }
-            }
-        },
-        failureButtonDisabled: {         
-            borderStyle: "hidden",
-            color: dashboard.darktheme ? "#909398" : "#616E85 !important",
-            backgroundColor: dashboard.darktheme ? "#1F2223" : "#f3f2f1"
-        },
-    });
-
-    return eventStyles;
-}
-
-const missingImage = "/images/missing-image.png";
-const MissingImageLabel = (): JSX.Element => { return <span style={{ fontWeight: 'bold' }}> [missing image]</span> }
-type ImageData = { link?: string, ref?: string }
-type ImageLinks = {
-    unapproved: ImageData;
-    approved?: ImageData;
-    difference?: ImageData;
-}
-
-const EventPane: React.FC<{ entry: EventEntry, phase: TestPhaseStatus }> = memo(({entry, phase}) => {
-    const [imageInfo, setImageInfo] = useState<ImageLinks | undefined>(undefined);
-    const eventStyles = getEventStyling();
-
-    const is_error = entry.level === EventLevel.Error || entry.level === EventLevel.Critical;
-    const is_warning = entry.level === EventLevel.Warning;
-    const style = is_error ? eventStyles.eventError : (is_warning ? eventStyles.eventWarning : eventStyles.event);
-
-    const need_image_comparison = entry.properties && entry.properties["$type"] === LogPropertyType.ImageCompare;
-
-    useEffect(() => {
-        if(need_image_comparison && entry.properties) {
-            const imageCompare = entry.properties as ImageCompare;
-            const getLinks = async () => {
-                if (!phase.artifacts) return;
-                let url = await phase.artifacts.getLink(imageCompare.unapproved);
-                const imageInfo: ImageLinks = {unapproved: { link: url, ref: imageCompare.unapproved }};
-                if (imageCompare.approved) {
-                    url = await phase.artifacts.getLink(imageCompare.approved);
-                    imageInfo.approved = { link: url, ref: imageCompare.approved };
-                }
-                if (imageCompare.difference) {
-                    url = await phase.artifacts.getLink(imageCompare.difference);
-                    imageInfo.difference = { link: url, ref: imageCompare.difference };
-                }
-                setImageInfo(imageInfo);
-            }
-            getLinks();
-        }
-    }, [])
-
-    return <Stack className={style} grow>
-                <pre style={{ margin: 0, whiteSpace: "pre-wrap" }}>[{entry.time}] {entry.level}: {entry.message}</pre>
-                {need_image_comparison && imageInfo &&
-                    <Stack styles={{ root: { paddingLeft: 16, width: '100%' } }}>
-                    <Stack><Text variant="medium"><span style={{ fontWeight: "bold" }}>Image comparison: </span>{entry['$text']}</Text></Stack>
-                    <Stack horizontal>
-                        {imageInfo.approved &&
-                            <Stack styles={{ root: { padding: 5 } }}>
-                                <a href={imageInfo.approved.link}>
-                                    <Image width="100%" style={{ minWidth: 200, maxWidth: 400, minHeight: 120 }} src={imageInfo.approved.link || missingImage} alt={imageInfo.approved.ref} />
-                                </a>
-                                <Stack.Item align="center">Reference{!imageInfo.approved && MissingImageLabel()}</Stack.Item>
-                            </Stack>
-                        }
-                        {imageInfo.difference &&
-                            <Stack styles={{ root: { padding: 5 } }}>
-                                <a href={imageInfo.difference.link}>
-                                    <Image width="100%" style={{ minWidth: 200, maxWidth: 400, minHeight: 120 }} src={imageInfo.difference.link || missingImage} alt={imageInfo.difference.ref} />
-                                </a>
-                                <Stack.Item align="center">Difference{!imageInfo.difference && MissingImageLabel()}</Stack.Item>
-                            </Stack>
-                        }
-                        <Stack styles={{ root: { padding: 5 } }}>
-                            <a href={imageInfo.unapproved.link}>
-                                <Image width="100%" style={{ minWidth: 200, maxWidth: 400, minHeight: 120 }} src={imageInfo.unapproved.link || missingImage} alt={imageInfo.unapproved.ref} />
-                            </a>
-                            <Stack.Item align="center">Produced{!imageInfo.unapproved && MissingImageLabel()}</Stack.Item>
-                        </Stack>
-                    </Stack>
-                    </Stack>
-                }
+import { JobArtifactsModal } from "horde/components/artifacts/ArtifactsModal";
+import { getJobParams, JobContext } from "./phaseCompare";
+import { getEventStyling, PhaseEventsPane } from "./testAutomationEvents";
+import { comparisonContext, PhaseComparisonPane } from "./testAutomationPhaseCompare";
+
+const PhaseHistoryPane: React.FC<{ test: TestNameRef, meta?: MetadataRef, phaseKey: string, streamId?: string, sessionId?: string, handler: TestDataHandler, onClickCompare: () => void }> = memo(({test, meta, phaseKey, streamId, sessionId, handler, onClickCompare}) => {
+
+    const metaStatus = !!meta && !!streamId? handler.getStatusStream(streamId)?.tests.get(test)?.sessions.get(meta) : undefined;
+   
+    return <Stack style={{ paddingLeft: 3 }} horizontal verticalAlign="end" disableShrink>
+                <Stack style={{ marginLeft: 30, marginBottom: 18, width: 10, cursor: "default" }} title="History">
+                    <FontIcon style={{ fontSize: 15 }} iconName={"History"} />
+                </Stack>
+                {!metaStatus || !sessionId && <Text style={{ marginLeft: 10,marginBottom: 21}}>No history available for this phase</Text>}
+                {!!metaStatus && !!sessionId && <PhaseHistoryGraphWidget
+                    test={test}
+                    streamId={streamId!}
+                    sessionId={sessionId}
+                    phaseKey={phaseKey}
+                    meta={meta!}
+                    handler={handler}
+                    onClick={(session) => {
+                        const testSession = metaStatus?.history.find(s => s.id === session.sessionId);
+                        !!testSession?.testDataId && handler.setSearchParam('session', testSession.testDataId);
+                    }}
+                />}
+                <Stack grow horizontalAlign="end" style={{ marginBottom: 16, marginRight: 16}}>
+                    <FontIcon
+                        style={{ fontSize: 15, cursor: "pointer" , backgroundColor: "#035ca1", borderRadius: 4, padding: 7, color: 'white' }}
+                        title="Toggle Comparison View"
+                        iconName={"BranchCompare"}
+                        onClick={onClickCompare}
+                    />
+                </Stack>
             </Stack>
 });
 
-const OnDeviceCard = (device: DeviceRef): JSX.Element => {
-    const meta = device.metadata;
-    return <Stack style={{padding: 6}}>
-            { meta && <Stack horizontal tokens={{childrenGap: 4}}>
-                        <Stack tokens={{childrenGap: 4}}>
-                            {Object.keys(meta).map((key, i) => <Stack style={{fontWeight: "bold"}} key={`device-${device.name}-key-${i}`}>{key}: </Stack>)}
-                        </Stack>
-                        <Stack tokens={{childrenGap: 4}}>
-                            {Object.values(meta).map((value, i) => <Stack key={`device-${device.name}-value-${i}`}>{value}</Stack>)}
-                        </Stack>
-                </Stack>
-            }
-            { !meta && <Text>no metadata</Text> }
-        </Stack>;
-}
-
-const DeviceItem = (device: DeviceRef): JSX.Element => {
-    const deviceCardProps: IPlainCardProps = {
-        renderData: device,
-        onRenderPlainCard: OnDeviceCard,
-        directionalHint: DirectionalHint.topRightEdge,
-        gapSpace: 8,
-        calloutProps: { isBeakVisible: true }
-    };
-
-    return <HoverCard plainCardProps={deviceCardProps} type={HoverCardType.plain} cardOpenDelay={20} key={device.name}>
-                <Text style={{fontSize: 12, cursor: 'pointer' }}>{device.name}</Text>
-            </HoverCard>
-}
-
 const TestPhaseExpandPane: React.FC<{ phase: TestPhaseStatus, onEventsLoaded?: () => void, handler: TestDataHandler }> = memo(({ phase, onEventsLoaded, handler }) => {
-    const [events, setEvents] = useState<EventEntry[] | undefined>(undefined);
-    const [filterError, setFilterError] = useState<boolean | undefined>(true);
-    const [filterWarning, setFilterWarning] = useState<boolean | undefined>(true);
-    const [filterInfo, setFilterInfo] = useState<boolean | undefined>(phase.outcome !== TestPhaseOutcome.Failed);
+    const [sessionId, setSessionId] = useState<string | undefined>();
+    const [comparison, setComparison] = useState<boolean>(comparisonContext.isEnable);
 
     useEffect(() => {
+        // query phase data
+        setSessionId(undefined);
         handler.queryPhase(phase, handler.stream!).then(() => {
-            phase.fetchEventStream().then((value) => {
-                setEvents(value ?? []);
-            }).catch((reason) => setEvents([{level: EventLevel.Error, message: `Could not retrieve test events: ${reason}`, format: "", time: "", id: 0}]));
+            setComparison(comparisonContext.isEnable);
+            setSessionId(phase.session?.id ?? "");
+            if (!phase.session) {
+                // if no phase session associated with this phase in history, it should be a preflight.
+                // we are going to prefill the commitId in the comparison context with the preflight base commit id.
+                comparisonContext.commitId = phase.commitId;
+            }
         });
-    }, []);
+    }, [phase]);
 
-    useEffect(() => {
-        if (events) onEventsLoaded?.();
-    }, [events]);
-
-    const isEntryNeedDisplay = useCallback((event: EventEntry): boolean => {
-        const eventType = event.level;
-        if (filterError && eventType === EventLevel.Error) {
-            return true;
-        }
-        if (filterWarning && eventType === EventLevel.Warning) {
-            return true;
-        }
-        if (filterInfo && eventType === EventLevel.Information) {
-            return true;
-        }
-
-        return false;
-    }, [filterError, filterWarning, filterInfo]);
-
-    const borderColor = dashboard.darktheme ? "#4D4C4B" : "#6D6C6B";
-
-    return <Stack>
-                {!events &&
+    return <Stack style={{overflow: 'hidden'}} grow>
+                {sessionId === undefined &&
                     <Stack horizontalAlign='center' style={{ padding: 12}} tokens={{ childrenGap: 10 }}>
                         <Text style={styles.textSmall}>Loading Data</Text>
                         <Spinner size={SpinnerSize.medium} />
                     </Stack>
                 }
+                {sessionId !== undefined &&
+                    <Stack style={{margin: 3, overflow: 'hidden', height: comparison? '47%' : 'auto'}} grow>
+                        { !!phase.test && <PhaseHistoryPane
+                                                test={phase.test}
+                                                streamId={phase.testSession?.streamId ?? handler.stream}
+                                                meta={phase.testMetadata}
+                                                sessionId={phase.session?.id}
+                                                phaseKey={phase.key}
+                                                handler={handler}
+                                                onClickCompare={() => {
+                                                    if (comparison) { comparisonContext.reset() }
+                                                        else { comparisonContext.setEnable() }
+                                                    setComparison(!comparison)
+                                                }}
+                                            /> }
+                        <PhaseEventsPane phase={phase} onEventsLoaded={onEventsLoaded}/>
+                    </Stack>
+                }
+                {comparison && sessionId !== undefined && !!phase.test && 
+                    <Stack style={{margin: 3, overflow: 'hidden', height: '53%'}} grow>
+                        <PhaseComparisonPane
+                            test={phase.test}
+                            streamId={phase.testSession?.streamId ?? handler.stream}
+                            meta={phase.testMetadata}
+                            focusedSessionId={phase.session?.id}
+                            phaseKey={phase.key}
+                            handler={handler}
+                            onDismiss={() => {comparisonContext.reset(); setComparison(false)}}
+                        />
+                    </Stack>
+                }
+            </Stack>
+});
 
-                {!!events &&
-                    <Stack style={{margin: 3}} tokens={{childrenGap: 3}}>
-                        <Stack>[ History graph place holder ]</Stack>
-                        <Stack style={{borderWidth: 1, borderColor: borderColor, borderStyle: 'solid', padding: 6}} tokens={{childrenGap: 3}}>
-                            <Stack horizontal tokens={{childrenGap: 8}}>
-                                <Text><span style={styles.labelSmall}>Run on: </span><span>{getShortNiceTime(phase.start, true, true, true)}</span></Text>
-                                <Text><span style={styles.labelSmall}> For: </span><span>{msecToElapsed(phase.duration * 1000, true, true)}</span></Text>
-                                {!!phase.devices.length &&
-                                    <Stack horizontal tokens={{childrenGap: 4}} verticalAlign="end">
-                                        <Text style={styles.labelSmall}> Devices: </Text>
-                                        {phase.devices.map(d => DeviceItem(d))}
-                                    </Stack>
-                                }
-                                <Stack verticalAlign="center" horizontal tokens={{ childrenGap: 5 }} style={{borderWidth: '0px 1px', paddingLeft: 9, paddingRight: 9, borderStyle: 'solid', borderColor: borderColor}}>
-                                    <FontIcon iconName="PageListFilter" title="Events filtering" style={styles.icon} />
-                                    <Checkbox styles={styles.defaultCheckbox} label="Error" defaultChecked={filterError} onChange={(_: any, checked?: boolean) => setFilterError(checked)} disabled={!events.length}/>
-                                    <Checkbox styles={styles.defaultCheckbox} label="Warning" defaultChecked={filterWarning} onChange={(_: any, checked?: boolean) => setFilterWarning(checked)} disabled={!events.length}/>
-                                    <Checkbox styles={styles.defaultCheckbox} label="Info" defaultChecked={filterInfo} onChange={(_: any, checked?: boolean) => setFilterInfo(checked)} disabled={!events.length}/>
-                                </Stack>
-                                <Stack horizontalAlign="end">
-                                    <DefaultButton style={{ fontSize: 11, padding: 0, height: 18, backgroundColor: styles.defaultButton.backgroundColor }}
-                                        onClick={(ev) => {
-                                            ev.stopPropagation();
-                                            /** do something */
-                                        }}>Artifacts</DefaultButton>
-                                </Stack>
-                            </Stack>
-                            <Stack style={{ paddingTop: 6, paddingBottom: 6}}>
-                                { events.map((e, i) => isEntryNeedDisplay(e) && <EventPane key={`phase-${phase.key}-event-${i}`} entry={e} phase={phase}/>)}
-                                { !events.length && <Text style={styles.textSmall}>No events</Text>}
-                            </Stack>
+const TestPhaseHeader: React.FC<{ phase: TestPhaseStatus, expanded: boolean, toggleView: (phaseKey: string) => void }> = memo(({ phase, expanded, toggleView }) => {
+    const onClick = useCallback(() => {
+        toggleView(phase.key);
+    }, [phase]);
+
+    const color = getPhaseSessionStatusColor(phase);
+    return <Stack tokens={{childrenGap: 2}} disableShrink>
+            <Stack horizontal verticalAlign="center">
+                <Stack style={{ paddingTop: 1, paddingRight: 4, cursor: "pointer" }} onClick={onClick}>
+                    <FontIcon style={{ fontSize: 13 }} iconName={expanded ? "ChevronDown" : "ChevronRight"} />
+                </Stack>
+                <Text style={{fontSize: 13, fontWeight: 600}}>{phase.name}</Text>
+                {!!expanded &&
+                    <Stack horizontalAlign="end" grow>
+                        <Stack style={{ padding: 3, cursor: "pointer" }}>
+                            <FontIcon
+                                style={{ fontSize: 12 }}
+                                iconName={"ChromeMinimize"}
+                                onClick={onClick}
+                            />
                         </Stack>
                     </Stack>
                 }
             </Stack>
+            <Stack horizontal verticalAlign="center" onClick={onClick} style={{cursor: 'pointer'}}>
+                <Stack style={{ paddingLeft: 5, paddingRight: 4 }}>
+                    <FontIcon style={{ fontSize: 12, color: color }} iconName="Square" />
+                </Stack>
+                <Text style={{fontSize: 11}}>{phase.outcome}</Text>
+            </Stack>
+        </Stack>
 });
 
-const TestPhaseSection: React.FC<{ phase: TestPhaseStatus, focused: boolean, handler: TestDataHandler }> = memo(({ phase, focused, handler }) => {
-    const [expanded, setExpanded] = useState(false);
+const TestPhaseSection: React.FC<{ phase: TestPhaseStatus, handler: TestDataHandler, onCollapse: (phaseKey: string) => void }> = ({ phase, handler, onCollapse }) => {
+    const toggleView = useCallback((phaseKey: string) => {
+        comparisonContext.reset();
+        handler.removeSearchParam('phase');
+        onCollapse(phaseKey);
+    }, []);
+
+    const borderColor = dashboard.darktheme ? "#4D4C4B" : "#6D6C6B";
+    return <Stack style={{ borderTopWidth: 1, borderTopColor: borderColor, borderTopStyle: 'solid', height: '100%'}}>
+                <TestPhaseHeader phase={phase} expanded={true} toggleView={toggleView} />
+                <TestPhaseExpandPane phase={phase} handler={handler}/>
+            </Stack>
+}
 
+const TestPhaseListItem: React.FC<{ phase: TestPhaseStatus, focused: boolean, handler: TestDataHandler }> = memo(({ phase, focused, handler }) => {
     useEffect(() => {
-        if (focused && !expanded) {
-            setExpanded(true);
-        }
-        scrollToSelected();
+        if (focused) {
+            const component = phase.componentRef;
+            if (!component) return;
+            component.scrollIntoView();
+       }
     }, [focused]);
 
     const registerRef = useCallback((ref: any) => {
         phase.componentRef = ref;
     }, []);
 
-    const scrollToSelected = useCallback(() => {
-        if (expanded && focused) {
-            const component = phase.componentRef;
-            if (!component) return;
-            component.scrollIntoView();
-       }
-    }, [focused, expanded]);
-
-    const toggleView = useCallback(() => {
-        if (!expanded) {
-            handler.setSearchParam("phase", phase.key);
-        } else if(focused) {
-            handler.removeSearchParam('phase');
-        }
-        setExpanded(!expanded);
-    }, [focused, expanded]);
+    const toggleView = useCallback((phaseKey: string) => {
+        handler.setSearchParam("phase", phaseKey);
+    }, []);
     
     const borderColor = dashboard.darktheme ? "#4D4C4B" : "#6D6C6B";
-    const statusColors = dashboard.getStatusColors();
-    let color = statusColors.get(StatusColor.Unspecified);
-    switch(phase.outcome) {
-        case TestPhaseOutcome.Failed:
-            color = statusColors.get(StatusColor.Failure);
-            break;
-        case TestPhaseOutcome.Success:
-            color = statusColors.get(phase.hasWarning? StatusColor.Warnings : StatusColor.Success);
-            break;
-    }
 
     return <Stack style={{ borderTopWidth: 1, borderTopColor: borderColor, borderTopStyle: 'solid', paddingTop: 3, paddingBottom: 6 }}>
-                <div ref={registerRef}>
-                    <Stack tokens={{childrenGap: 2}}>
-                        <Stack horizontal verticalAlign="center">
-                            <Stack style={{ paddingTop: 1, paddingRight: 4, cursor: "pointer" }} onClick={toggleView}>
-                                <FontIcon style={{ fontSize: 13 }} iconName={expanded ? "ChevronDown" : "ChevronRight"} />
-                            </Stack>
-                            <Text style={{fontSize: 13, fontWeight: 600}}>{phase.name}</Text>
-                        </Stack>
-                        <Stack horizontal verticalAlign="center">
-                            <Stack style={{ paddingLeft: 5, paddingRight: 4 }}>
-                                <FontIcon style={{ fontSize: 12, color: color }} iconName="Square" />
-                            </Stack>
-                            <Text style={{fontSize: 11}}>{phase.outcome}</Text>
-                        </Stack>
-                    </Stack>
-                </div>
-                {expanded && <TestPhaseExpandPane phase={phase} onEventsLoaded={scrollToSelected} handler={handler}/>}
+                <div ref={registerRef}><TestPhaseHeader phase={phase} expanded={false} toggleView={toggleView} /></div>
             </Stack>
 });
 
@@ -361,7 +190,7 @@ const onRenderItem = (props: IPickerItemProps<ITag>) => {
     return <Stack key={`picker_item_${item.name}`} horizontal verticalAlign="center" verticalFill={true} style={{ marginRight: 2 }}>
                 <TooltipHost content={item.name} delay={0} closeDelay={500}>
                     <FontIcon
-                        style={{ fontSize: 13, padding: 4, paddingTop: 5, cursor: 'pointer', backgroundColor: styles.defaultFilter.backgroundColor, borderRadius: 3}}
+                        style={{ fontSize: 13, padding: 4, paddingTop: 5, cursor: 'pointer', backgroundColor: styles.defaultFilter.backgroundColor, borderRadius: 3, color: 'white'}}
                         iconName="Filter"
                         title="Remove"
                         onClick={props.onRemoveItem} />
@@ -460,10 +289,14 @@ class CircularList<T> {
 export const TestPhasesView: React.FC<{ testDataId: string, handler: TestDataHandler, onDismiss: () => void }> = observer(({ testDataId, handler, onDismiss }) => {
     const [sessionDetails, setSessionDetails] = useState<TestSessionDetails | undefined>(undefined);
     const [filteredPhases, setFilteredPhases] = useState<TestPhaseStatus[]>([]);
+    const [focusPhase, setFocusPhase] = useState<string>();
+    const [jobArtifacts, setJobArtifacts] = useState<GetArtifactResponse>();
+    const selectedPhase = useRef<TestPhaseStatus>();
     const failureCursor = useRef<CircularList<string>>();
     const listRef = useRef<List>();
-    const job = useRef<{params?: JobParameters, rangeCL?: number}>();
-    const selectedPhase = handler.getSearchParam('phase') as string | undefined;
+    const jobContext = useRef<JobContext>();
+
+    const phaseKey = handler.getSearchParam('phase') as string | undefined;
 
     handler.subscribeToTestDataQueryLoading();
     handler.subscribeToSearch();
@@ -471,55 +304,63 @@ export const TestPhasesView: React.FC<{ testDataId: string, handler: TestDataHan
     const updateFilter = useCallback((phases: TestPhaseStatus[]) => {
         const failures = phases.filter((p) => p.outcome !== TestPhaseOutcome.Success && p.outcome !== TestPhaseOutcome.NotRun && p.outcome !== TestPhaseOutcome.Skipped).map((p) => p.key) ?? [];
         failureCursor.current = new CircularList(failures);
-        failureCursor.current.setCursor(selectedPhase);
+        failureCursor.current.setCursor(selectedPhase.current?.key);
         setFilteredPhases(phases);
     }, [handler.searchUpdated]);
 
     const onRenderCell = useCallback(
-        (phase: TestPhaseStatus) => <TestPhaseSection key={phase.key} phase={phase} focused={selectedPhase === phase.key} handler={handler}/>
+        (phase: TestPhaseStatus) => <TestPhaseListItem key={phase.key} phase={phase} focused={focusPhase === phase.key} handler={handler}/>
         , [handler.searchUpdated]);
 
     useEffect(() => {
         handler.queryTestData(testDataId).then((fetchedSessionDetails) => {
             if (fetchedSessionDetails) {
+                handler.selectStream(fetchedSessionDetails.streamId);
                 setSessionDetails(fetchedSessionDetails);
                 updateFilter(fetchedSessionDetails.phases);
-                job.current = {
-                    params: {
-                        id: fetchedSessionDetails.jobId,
-                        streamId: fetchedSessionDetails.streamId,
-                        change: fetchedSessionDetails.commitOrder,
-                    },
-                    rangeCL: fetchedSessionDetails.session ?
-                        handler.getStatusStream(fetchedSessionDetails.streamId)?.getPreviousSession(fetchedSessionDetails.test!, fetchedSessionDetails.meta!, fetchedSessionDetails.session)?.commitOrder
-                        : undefined
-                }
+                jobContext.current = getJobParams(fetchedSessionDetails, handler);
             }
         });
         setSessionDetails(undefined);
     }, [testDataId]);
 
     useEffect(() => {
-        if (selectedPhase !== failureCursor.current?.selected) {
-            failureCursor.current?.setCursor(selectedPhase);
+        return () => {
+            // on unmount
+            comparisonContext.reset();
         }
-        if (!selectedPhase) return;
-        const selectedIndex = filteredPhases.findIndex((phase) => phase.key === selectedPhase);
-        if (selectedIndex === -1) return;
+    }, [])
 
-        if (!filteredPhases[selectedIndex].componentRef) {
+    const onPagesUpdated = useCallback(() => {
+        if (!focusPhase) return;
+        const phaseIndex = filteredPhases.findIndex((phase) => phase.key === focusPhase);
+        if (phaseIndex === -1) return;
+
+        if (!filteredPhases[phaseIndex].componentRef) {
             // force the list to load the page that contains the target item
-            listRef.current?.scrollToIndex(selectedIndex);
+            listRef.current?.scrollToIndex(phaseIndex);
+        } else {
+            // validate the scroll to index was achieved
+            setFocusPhase(undefined);
+        }
+    }, [focusPhase]);
+
+    selectedPhase.current = useMemo(() => {
+        return filteredPhases.find((phase) => phase.key === phaseKey);
+    }, [handler.searchUpdated, filteredPhases]);
+
+    useEffect(() => {
+        if (phaseKey !== failureCursor.current?.selected) {
+            failureCursor.current?.setCursor(phaseKey);
         }
-        // force the list to update in anycase so that the loaded items get the propagation of the search params update
-        listRef.current?.forceUpdate();
     }, [handler.searchUpdated, filteredPhases]);
 
     const StatusPie = useMemo(
         () => <Stack horizontal verticalAlign='center' style={{ height: '100%' }} tokens={{childrenGap: 3}}>
                 <Label>{filteredPhases.length} {filteredPhases.length > 1 ? "Phases" : "Phase"}</Label>
                 {PhasesStatusPie(filteredPhases, 14)}
-            </Stack>, [filteredPhases]);
+            </Stack>
+        , [filteredPhases]);
 
     const metaString = useMemo(() => sessionDetails?.meta?.getValues().join(' / '), [sessionDetails]);
 
@@ -533,6 +374,9 @@ export const TestPhasesView: React.FC<{ testDataId: string, handler: TestDataHan
     const failureText = failureEnabled ? failureCursor.current!.length.toString() : "No";
 
     return <Stack style={{width: '100%', height: '100%'}}>
+                {!!jobArtifacts && jobContext.current?.job?.id && jobContext.current?.step?.id &&
+                    <JobArtifactsModal jobId={jobContext.current.job.id} stepId={jobContext.current?.step?.id} contextType={jobArtifacts.type} artifacts={[jobArtifacts]} onClose={() => setJobArtifacts(undefined)}/>
+                }
                 {!sessionDetails && handler.testDataQueryLoading &&
                         <Stack horizontalAlign='center' style={{ padding: 12 }} tokens={{ childrenGap: 10 }}>
                             <Text style={styles.textLarge}>Loading Data</Text>
@@ -563,10 +407,17 @@ export const TestPhasesView: React.FC<{ testDataId: string, handler: TestDataHan
                                 <Text style={{ fontSize: 16, fontWeight: 600 }}>{`${sessionDetails.name} - ${metaString}`}</Text>
                                 <Stack style={{ paddingLeft: 8 }} horizontal tokens={{childrenGap: 8}}>
                                     <Text>{projectStore.streamById(sessionDetails.streamId)?.fullname ?? sessionDetails.streamId}</Text>
-                                    <ChangeButton job={job.current?.params} rangeCL={job.current?.rangeCL} />
-                                    <Link href={`/job/${sessionDetails.jobId}?step=${sessionDetails.stepId}`} target="_blank"><Text style={{ fontSize: 12, color: hordeTheme.semanticColors.link }}>View Job</Text></Link>
+                                    <ChangeButton job={jobContext.current?.job} rangeCL={jobContext.current?.rangeCL} />
+                                    <Link href={`/job/${jobContext.current?.job?.id}?step=${jobContext.current?.step?.id}`} target="_blank"><Text style={{ fontSize: 12, color: hordeTheme.semanticColors.link }}>View Job</Text></Link>
                                     <Text><span style={styles.labelSmall}>Run on: </span><span>{getShortNiceTime(sessionDetails.start, true, true, true)}</span></Text>
                                     <Text><span style={styles.labelSmall}> For: </span><span>{msecToElapsed(sessionDetails.duration * 1000, true, true)}</span></Text>
+                                    <Stack horizontalAlign="end">
+                                        <DefaultButton style={{ fontSize: 11, padding: 0, height: 18, backgroundColor: styles.defaultButton.backgroundColor, color: 'white' }}
+                                            onClick={(ev) => {
+                                                ev.stopPropagation();                                                
+                                                sessionDetails.artifacts?.getJobStepSavedArtifacts().then(value => setJobArtifacts(value));
+                                            }}>Artifacts</DefaultButton>
+                                    </Stack>
                                 </Stack>
                             </Stack>
 
@@ -625,13 +476,23 @@ export const TestPhasesView: React.FC<{ testDataId: string, handler: TestDataHan
                                 />
                             </Stack>
                         </Stack>
-                        <Stack style={{ paddingTop: 3, paddingLeft: 16, paddingBottom: 3, paddingRight: 3, height: '100%', overflow: 'auto'}} grow>
-                            <List
-                                items={filteredPhases}
-                                onRenderCell={onRenderCell}
-                                getKey={(phase: TestPhaseStatus) => `phase-${phase.key}`}
-                                ref={(list: List) => listRef.current = list}
+                        <Stack style={{ paddingTop: 3, paddingLeft: 16, paddingBottom: 3, paddingRight: 3, overflow: !selectedPhase.current? 'auto' : 'hidden', height: '100%'}} grow>
+                            {!selectedPhase.current &&
+                                <List
+                                    items={filteredPhases}
+                                    onRenderCell={onRenderCell}
+                                    getKey={(phase: TestPhaseStatus) => `phase-${phase.key}`}
+                                    ref={(list: List) => listRef.current = list}
+                                    onPagesUpdated={onPagesUpdated}
+                                />
+                            }
+                            {!!selectedPhase.current &&
+                                <TestPhaseSection
+                                    phase={selectedPhase.current}
+                                    handler={handler}
+                                    onCollapse={(phaseKey) => {setFocusPhase(phaseKey)}}
                                 />
+                            }
                         </Stack>
                     </Stack>
                 }     
diff --git a/Engine/Source/Programs/Horde/HordeDashboard/plugins/test-automation/testAutomationSummary.tsx b/Engine/Source/Programs/Horde/HordeDashboard/plugins/test-automation/testAutomationSummary.tsx
--- a/Engine/Source/Programs/Horde/HordeDashboard/plugins/test-automation/testAutomationSummary.tsx
+++ b/Engine/Source/Programs/Horde/HordeDashboard/plugins/test-automation/testAutomationSummary.tsx
@@ -7,7 +7,7 @@ import dashboard, { StatusColor } from "horde/backend/Dashboard";
 import { getShortNiceTime } from "horde/base/utilities/timeUtils";
 import { getHordeStyling } from "horde/styles/Styles";
 import { projectStore } from 'horde/backend/ProjectStore';
-import { SessionStatusBar, SessionValues, StatusBar, StatusBarStack, MultiOptionChooser } from "./testAutomationCommon";
+import { SessionStatusBar, SessionValues, StatusBar, StatusBarStack, MultiOptionChooser, getTestSessionStatusColor } from "./testAutomationCommon";
 import { TestHistoryGraph } from "./testHistoryGraph";
 
 type SessionCallback = (session: TestSessionResult) => void;
@@ -240,20 +240,14 @@ const TestSummaryCard: React.FC<{ test: TestNameRef, sessions: [MetadataRef, Tes
     const maskMetaString  = sessions.map(([meta,]) => meta).find(() => true)?.getCommonValues(metaKeyMask)?.join(" / ");
  
     sessions.forEach(([meta, session]) => {
-        const outcome = session.outcome;
         const metaName = meta.getValuesExcept(commonKeys).join(" / ") || (meta.getValues().findLast(() => true) ?? test.name);
     
-        const sessionTotalCount = session.phaseSucceededCount + session.phaseFailedCount + session.phaseUnspecifiedCount;
-        testUnspecified += session.phaseUnspecifiedCount;
-        testFailed += session.phaseFailedCount;
+        const sessionTotalCount = session.phasesSucceededCount + session.phasesFailedCount + session.phasesUnspecifiedCount;
+        testUnspecified += session.phasesUnspecifiedCount;
+        testFailed += session.phasesFailedCount;
         testTotal += sessionTotalCount;
     
-        let color = statusColors.get(StatusColor.Success);
-        if (outcome === TestOutcome.Failure) {
-            color = statusColors.get(StatusColor.Failure);
-        } else if (outcome === TestOutcome.Unspecified) {
-            color = statusColors.get(StatusColor.Unspecified);
-        }
+        const color = getTestSessionStatusColor(session);
     
         let dateString = "";
     
diff --git a/Engine/Source/Programs/Horde/HordeDashboard/plugins/test-automation/testAutomationView.tsx b/Engine/Source/Programs/Horde/HordeDashboard/plugins/test-automation/testAutomationView.tsx
--- a/Engine/Source/Programs/Horde/HordeDashboard/plugins/test-automation/testAutomationView.tsx
+++ b/Engine/Source/Programs/Horde/HordeDashboard/plugins/test-automation/testAutomationView.tsx
@@ -1,17 +1,19 @@
-import { Stack, FontIcon, IComboBoxOption, IContextualMenuItem, IContextualMenuProps, Text, Spinner, SpinnerSize, Label, DefaultButton, TagItemSuggestion, ITag, TagPicker, IPickerItemProps, PrimaryButton } from "@fluentui/react";
+import { Stack, FontIcon, IComboBoxOption, IContextualMenuItem, IContextualMenuProps, Text, Spinner, SpinnerSize, Label, DefaultButton, TagItemSuggestion, ITag, TagPicker, IPickerItemProps, PrimaryButton, Modal, IconButton, Link } from "@fluentui/react";
 import { Breadcrumbs } from "horde/components/Breadcrumbs";
 import { TopNav } from "horde/components/TopNav";
 import { getHordeStyling } from "horde/styles/Styles";
 import { observer } from "mobx-react-lite";
-import { useSearchParams } from "react-router-dom";
+import { useNavigate, useSearchParams } from "react-router-dom";
 import React, { useEffect, useState, useCallback } from "react";
 import { projectStore } from 'horde/backend/ProjectStore';
-import { TestDataHandler, TestDataVersionRegistrar } from "./testData";
+import { TestDataHandler, TestDataVersionRegistrar, TestSessionDetails } from "./testData";
 import { TestSummary } from "./testAutomationSummary";
-import { MultiOptionChooser } from "./testAutomationCommon";
+import { getTestSessionStatusColor, MultiOptionChooser, SessionStatusBar } from "./testAutomationCommon";
 import { TestPhasesView } from "./testAutomationPhaseView";
+import { JobStepTestDataItem } from "./api";
 
 import { TestDataV2 } from "./testDataV2Fetcher";
+import dashboard from "horde/backend/Dashboard";
 
 const handler = new TestDataHandler();
 
@@ -281,12 +283,90 @@ const TestAutomationPanel: React.FC = observer(() => {
 
 });
 
+const JobStepTestsSelectionModal: React.FC<{ jobId: string, stepId: string }> = ({ jobId, stepId }) => {
+    const navigate = useNavigate();
+    const [tests, setTests] = useState<(JobStepTestDataItem & {details?: TestSessionDetails})[]>();
+
+    const loadTestData = async (tests: JobStepTestDataItem[]) => {
+        const testWithDetails: (JobStepTestDataItem & {details?: TestSessionDetails})[] = [];
+        for (const test of tests) {
+            const sessionDetails = await handler.queryTestData(test.testDataId);
+            testWithDetails.push({...test, details: sessionDetails});
+        }
+
+        return testWithDetails;
+    }
+
+    useEffect(() => {
+        handler.queryJobStepTestData(jobId, stepId).then(tests => {
+            if (!!tests && tests.length === 1) {
+                // if only one item, redirect automatically
+                navigate(`?session=${tests[0].testDataId}`, {replace: true});
+            } else {
+                const selection = tests ?? [];
+                if (selection.length > 0) {
+                    loadTestData(selection).then(detailedList => setTests(detailedList));
+                } else {
+                    setTests(selection);
+                }
+            }
+        })
+    }, []);
+
+    const color = dashboard.darktheme ? "rgba(255, 255, 255, 0.1)" : "rgba(0, 0, 0, 0.1)";
+
+    return <Modal isOpen={true} isBlocking={true} topOffsetFixed={true} styles={{ main: { padding: 8, width: 1134, hasBeenOpened: false, top: "200px", position: "absolute" } }} onDismiss={() => navigate(-1)} >
+            <Stack grow>
+                <Stack horizontal verticalAlign="center">
+                    <Stack><Text style={{ paddingLeft: 8, fontSize: 16, fontWeight: 600 }}>Test Reports</Text></Stack>
+                    <Stack grow />
+                    <Stack style={{ paddingBottom: 4 }}>
+                        <IconButton
+                            iconProps={{ iconName: 'Cancel' }}
+                            ariaLabel="Close popup modal"
+                            onClick={() => navigate(-1)}
+                        />
+                    </Stack>
+                </Stack>
+                <Stack tokens={{ childrenGap: 9 }} style={{ padding: 24, overflow: "auto", maxHeight: 'calc(100vh - 300px)'}}>
+                    {!tests &&
+                        <Stack horizontalAlign='center' style={{ paddingTop: 24, width: "100%" }}>
+                            <Spinner size={SpinnerSize.large} />
+                        </Stack>
+                    }
+                    {!!tests && tests.length &&
+                        tests.map((test, i) =>  <Stack key={`${test.testKey}-${i}`} tokens={{childrenGap: 15}} horizontal verticalAlign="end"
+                                                styles={{root: {borderBottom: `1px solid ${color}`, padding: '8px 2px', cursor: 'pointer', selectors: {':hover': {backgroundColor: color}}}}}
+                                                onClick={() => navigate(`?session=${test.testDataId}`)}>
+                                                <Stack horizontal verticalAlign="center">
+                                                    {!!test.details &&
+                                                        <Stack style={{ paddingLeft: 0, paddingTop: 1, paddingRight: 4 }}>
+                                                            <FontIcon style={{ fontSize: 11, color: getTestSessionStatusColor(test.details) }} iconName="Square" />
+                                                        </Stack>
+                                                    }
+                                                    <Stack>
+                                                        <Text style={{ whiteSpace: 'nowrap' }} >{test.testName}{!!test.details?.meta && <span> - {test.details.meta.getValues().join(' / ')}</span>}</Text>                                                        
+                                                    </Stack>
+                                                </Stack>
+                                                <Stack>{!!test.details && SessionStatusBar(test.details, 200, 10)}</Stack>
+                                            </Stack>)
+                    }
+                    {!!tests && tests.length === 0 &&
+                        <Text>No report found</Text>
+                    }
+                </Stack>
+            </Stack>
+        </Modal>
+}
+
 export const TestAutomationView: React.FC = () => {
 
     const [searchParams, setSearchParams] = useSearchParams();
     const [initialized, setInitialized] = useState(false);
 
     const testDataId = searchParams.get('session');
+    const jobId = searchParams.get('job');
+    const stepId = searchParams.get('step');
 
     useEffect(() => {
         handler.initialize(searchParams.toString(), (search: string) => {setSearchParams(search)}).then(() => {setInitialized(true)});
@@ -306,8 +386,15 @@ export const TestAutomationView: React.FC = () => {
                 <TopNav />
                 <Breadcrumbs items={[{ text: 'Test Automation Hub' }]} />
                 <Stack grow style={{ padding: 12, backgroundColor: modeColors.background, height: 'calc(100vh - 148px)', overflow: 'auto' } } horizontalAlign='center'>
+                    {initialized && !!jobId && !!stepId &&
+                        <JobStepTestsSelectionModal jobId={jobId} stepId={stepId} />
+                    }
                     {initialized && !!testDataId &&
-                        <TestPhasesView testDataId={testDataId!} handler={handler} onDismiss={() => {handler.removeSearchParam('session'); handler.removeSearchParam('phase')}} />
+                        <TestPhasesView testDataId={testDataId!} handler={handler}
+                            onDismiss={() => {
+                                handler.removeSearchParam('session');
+                                handler.removeSearchParam('phase');
+                            }} />
                     }
                     {initialized && !testDataId && !!handler.availableStreams.length &&
                         <Stack horizontal grow>
diff --git a/Engine/Source/Programs/Horde/HordeDashboard/plugins/test-automation/testData.ts b/Engine/Source/Programs/Horde/HordeDashboard/plugins/test-automation/testData.ts
--- a/Engine/Source/Programs/Horde/HordeDashboard/plugins/test-automation/testData.ts
+++ b/Engine/Source/Programs/Horde/HordeDashboard/plugins/test-automation/testData.ts
@@ -1,8 +1,10 @@
 import { isNumber } from '@datadog/browser-core';
 import { action, makeObservable, observable, reaction, when } from 'mobx';
-import { EventEntry, TestOutcome, TestPhaseOutcome } from './api';
+import { JobStepTestDataItem, TestOutcome, TestPhaseOutcome } from './api';
 import { projectStore } from 'horde/backend/ProjectStore';
 import backend from 'horde/backend';
+import { GetArtifactResponse } from 'horde/backend/Api';
+import { EventEntry } from './testEventsModel';
 
 const defaultQueryWeeks = 2;
 const filterRefreshMS = 200;
@@ -32,6 +34,9 @@ export interface TestVersionFetcher {
     /// fetch test phase sessions and populate test status phase session
     fetchPhaseSessions(streamRef: TestStreamRef, status: TestMetaStatus, phaseKey: string, weeks: number): Promise<boolean>;
 
+    /// fetch test data information from job and step id
+    fetchJobStepTestData(jobId: string, stepId: string): Promise<JobStepTestDataItem[] | undefined>;
+
 }
 
 export class TestDataVersionRegistrar {
@@ -94,9 +99,9 @@ export class TestSessionResult {
         this.streamId = streamId;
         this.outcome = outcome;
         this.testDataId = testDataId;
-        this.phaseSucceededCount = 0;
-        this.phaseFailedCount = 0;
-        this.phaseUnspecifiedCount = 0;
+        this.phasesSucceededCount = 0;
+        this.phasesFailedCount = 0;
+        this.phasesUnspecifiedCount = 0;
     }
 
     get artifacts(): ArtifactFactory | undefined {
@@ -115,9 +120,9 @@ export class TestSessionResult {
     metadataId: MetaId;
     streamId: StreamId
     outcome: TestOutcome;
-    phaseSucceededCount: number;
-    phaseFailedCount: number;
-    phaseUnspecifiedCount: number;
+    phasesSucceededCount: number;
+    phasesFailedCount: number;
+    phasesUnspecifiedCount: number;
     jobId?: string;
     stepId?: string;
     testDataId?: string;
@@ -273,6 +278,10 @@ export abstract class TestPhaseStatus {
         return this.phaseRef?.name ?? this._name;
     }
 
+    get commitId() {
+        return this._sessionDetails.commitId;
+    }
+
     get testSession() {
         return this._sessionDetails.session;
     }
@@ -342,6 +351,34 @@ export class TestSessionDetails {
         }
         return this._artifactFactory;
     }
+    
+    get phasesSucceededCount(): number {
+        return this.session?.phasesSucceededCount ?? this._phasesSucceededCount;
+    }
+    set phasesSucceededCount(value: number) {
+        this._phasesSucceededCount = value;
+    }
+
+    get phasesUnspecifiedCount(): number {
+        return this.session?.phasesUnspecifiedCount ?? this._phasesUnspecifiedCount;
+    }
+    set phasesUnspecifiedCount(value: number) {
+        this._phasesUnspecifiedCount = value;
+    }
+
+    get phasesFailedCount(): number {
+        return this.session?.phasesFailedCount ?? this._phasesFailedCount;
+    }
+    set phasesFailedCount(value: number) {
+        this._phasesFailedCount = value;
+    }
+
+    get outcome(): TestOutcome {
+        return this.session?.outcome ?? this._outcome;
+    }
+    set outcome(value: TestOutcome) {
+        this._outcome = value;
+    }
 
     key: string;
     name: string;
@@ -360,6 +397,10 @@ export class TestSessionDetails {
     private _start: Date;
     private _duration: number;
     private _artifactFactory?: ArtifactFactory;
+    private _phasesSucceededCount: number;
+    private _phasesUnspecifiedCount: number;
+    private _phasesFailedCount: number;
+    private _outcome: TestOutcome;
 }
 
 export class TestNameRef {
@@ -1227,6 +1268,21 @@ export class TestDataHandler {
         return this.testTestData.get(testDataId);
     }
 
+    async queryJobStepTestData(jobId: string, stepId: string): Promise<JobStepTestDataItem[] | undefined> {
+        for (const fetcherVersion of TestDataVersionRegistrar.manifest) {
+            console.debug(`Fetching from version ${fetcherVersion.version} job step test data`)
+            try {
+                const response = await fetcherVersion.fetchJobStepTestData(jobId, stepId);
+                if (response) {
+                    return response
+                }
+            }  catch(reason) {
+                /** ignore failed attempt to try with next version */
+                console.debug(`Failed to fetch from version ${fetcherVersion.version}: ${reason}`);
+            }
+        }
+    }
+
     async queryPhase(phase: TestPhaseStatus, streamId: StreamId, isSubQuery: boolean = false) {
 
         const phaseKey = phase.key;
@@ -1437,28 +1493,30 @@ export class ArtifactFactory {
         this.stepId = stepId;
     }
 
-    private async getArtifactId() {
-
-        if (this.artifactId) {
-            return;
+    async getJobStepSavedArtifacts(): Promise<GetArtifactResponse | undefined> {
+        if (this.artifactV2Response) {
+            return this.artifactV2Response;
         }
 
         const v = await backend.getJobArtifactsV2(undefined, [`job:${this.jobId}/step:${this.stepId}`]);
         const av2 = v?.artifacts.find(a => a.type === "step-saved")
-        if (av2?.id) {
-            this.artifactId = av2.id;
+        if (av2) {
+            this.artifactV2Response = av2;
+            if (av2.id) {
+                this.artifactId = av2.id;
+            }
         } else {
             console.error("Unable to get step-saved artifacts for test report");
         }
 
+        return this.artifactV2Response;
     }
 
     async getLink(referencePath: string): Promise<string> {
 
-        await this.getArtifactId();
-
+        await this.getJobStepSavedArtifacts();
+        
         return `${backend.serverUrl}/api/v2/artifacts/${this.artifactId}/file?path=${encodeURIComponent(referencePath)}&inline=true`;
-
     }
 
     async fetch(referencePath: string): Promise<any> {
@@ -1484,6 +1542,7 @@ export class ArtifactFactory {
 
     private jobId: string;
     private stepId: string;
+    private artifactV2Response?: GetArtifactResponse;
     private artifactId?: string;
 
     private static cache: Map<string, any> = new Map();
diff --git a/Engine/Source/Programs/Horde/HordeDashboard/plugins/test-automation/testDataV2Fetcher.ts b/Engine/Source/Programs/Horde/HordeDashboard/plugins/test-automation/testDataV2Fetcher.ts
--- a/Engine/Source/Programs/Horde/HordeDashboard/plugins/test-automation/testDataV2Fetcher.ts
+++ b/Engine/Source/Programs/Horde/HordeDashboard/plugins/test-automation/testDataV2Fetcher.ts
@@ -1,5 +1,6 @@
-import { EventEntry, GetPhaseResponse, getStreamTestsV2, getTestDataV2, GetTestMetadataResponse, GetTestNameResponse, GetTestPhaseSessionResponse, getTestPhaseSessionsV2, getTestPhasesV2, GetTestSessionResponse, getTestSessionsV2, GetTestTagResponse, PhaseSessionDetailsResponse, TestData, TestSessionDetailsResponse } from "./api";
+import { getJobStepTestDataV2, GetPhaseResponse, getStreamTestsV2, getTestDataV2, GetTestMetadataResponse, GetTestNameResponse, GetTestPhaseSessionResponse, getTestPhaseSessionsV2, getTestPhasesV2, GetTestSessionResponse, getTestSessionsV2, GetTestTagResponse, JobStepTestDataItem, PhaseSessionDetailsResponse, TestData, TestOutcome, TestSessionDetailsResponse } from "./api";
 import { TestStreamRef, TestStatus, TestSessionResult, TestNameRef, MetadataRef, TestTagRef, PhaseSessionResult, TestPhaseRef, TestMetaStatus, TestSessionDetails, TestPhaseStatus, DeviceRef, TestVersionFetcher } from "./testData"
+import { EventEntry } from "./testEventsModel";
 
 const minuteInWeek = 10080;
 
@@ -87,6 +88,10 @@ export class TestDataV2 implements TestVersionFetcher {
 
     }
 
+    async fetchJobStepTestData(jobId: string, stepId: string): Promise<JobStepTestDataItem[] | undefined> {
+        return await getJobStepTestDataV2(jobId, stepId)
+    }
+
     async fetchPhaseSessions(streamRef: TestStreamRef, status: TestMetaStatus, phaseKey: string, weeks: number): Promise<boolean> {
 
         const streamId = streamRef.id;
@@ -165,9 +170,9 @@ export class TestDataV2 implements TestVersionFetcher {
     private createTestSession(test: GetTestSessionResponse): TestSessionResult {
         const start = new Date(test.startDateTime);
         const result = new TestSessionResult(test.id, start, test.duration, test.commitId, test.commitOrder, test.nameRef, test.metadataId, test.streamId, test.outcome, test.testDataId);
-        result.phaseSucceededCount = test.phasesSucceededCount;
-        result.phaseFailedCount = test.phasesFailedCount;
-        result.phaseUnspecifiedCount = test.phasesUndefinedCount;
+        result.phasesSucceededCount = test.phasesSucceededCount;
+        result.phasesFailedCount = test.phasesFailedCount;
+        result.phasesUnspecifiedCount = test.phasesUndefinedCount;
         result.jobId = test.jobId;
         result.stepId = test.stepId;
 
@@ -193,6 +198,15 @@ export class TestDataV2 implements TestVersionFetcher {
             result.devices.set(key, new DeviceRef(d.name, d.appInstanceLogPath, d.metadata, this.version));
         });
         result.phases.push(...testDetails.phases.map(p => this.createPhaseStatus(p, result.devices, result)));
+        result.phasesSucceededCount = testDetails.summary.phasesSucceededCount;
+        result.phasesFailedCount = testDetails.summary.phasesFailedCount;
+        result.phasesUnspecifiedCount = testDetails.summary.phasesUndefinedCount;
+        result.outcome = testDetails.summary.phasesFailedCount > 0 ?
+                            TestOutcome.Failure
+                            : testDetails.summary.phasesUndefinedCount > 0 ?
+                                TestOutcome.Unspecified
+                                : testDetails.summary.phasesTotalCount > 0 ?
+                                    TestOutcome.Success : TestOutcome.Skipped;
 
         // search for known metadata
         const metaFilter = new Map(Object.entries(testDetails.metadata).map(([key, value]) => [key, [value]]));
diff --git a/Engine/Source/Programs/Horde/HordeDashboard/plugins/test-automation/testEventsModel.ts b/Engine/Source/Programs/Horde/HordeDashboard/plugins/test-automation/testEventsModel.ts
new file mode 100644
--- /dev/null
+++ b/Engine/Source/Programs/Horde/HordeDashboard/plugins/test-automation/testEventsModel.ts
@@ -0,0 +1,60 @@
+/// An event entry
+export type EventEntry = {
+
+    time: string;
+
+    id: number;
+
+    level: EventLevel;
+
+    message: string;
+
+    format: string;
+
+    properties?: Record<string, string | number | Record<string, string | number | undefined> | undefined>;
+}
+
+/// The event level
+export enum EventLevel {
+
+    Unspecified = "Unspecified",
+
+    Information = "Information",
+
+    Warning = "Warning",
+
+    Error = "Error",
+
+    Critical = "Critical"
+}
+
+/// Known property types
+export enum LogPropertyType {
+    ImageCompare = "Image Compare"
+}
+
+/// Image Compare detailed properties
+export type ImageCompare = {
+    "$type": LogPropertyType.ImageCompare;
+    "$text": string;
+
+    //** The path to the unapproved image from the artifacts */
+    unapproved: string;
+
+    //** The path to the unapproved metadata associated with */
+    //** the image from the artifacts */
+    unapproved_metadata?: string;
+
+    //** The path to the approved image from the artifacts */
+    approved?: string;
+
+    //** The path to the approved metadata associated with */
+    //** the image from the artifacts */
+    approved_metadata?: string;
+
+    //** The path to the delta image from the artifacts */
+    difference?: string;
+
+    //** The path to the json report associated with the image from the artifacts */
+    difference_report?: string;
+}
diff --git a/Engine/Source/Programs/Horde/HordeDashboard/plugins/test-automation/testHistoryGraph.ts b/Engine/Source/Programs/Horde/HordeDashboard/plugins/test-automation/testHistoryGraph.ts
--- a/Engine/Source/Programs/Horde/HordeDashboard/plugins/test-automation/testHistoryGraph.ts
+++ b/Engine/Source/Programs/Horde/HordeDashboard/plugins/test-automation/testHistoryGraph.ts
@@ -270,7 +270,7 @@ export class TestHistoryGraph {
                 desc += `Commit ${closest.commitId} <br/>`;
                 desc += `Duration ${msecToElapsed(moment.duration(closest.duration).asMilliseconds(), true, true)} <br/>`;
                 desc += `${date} <br/>`;
-                if (closest.phaseSucceededCount + closest.phaseFailedCount + closest.phaseUnspecifiedCount > 1)
+                if (closest.phasesSucceededCount + closest.phasesFailedCount + closest.phasesUnspecifiedCount > 1)
                     desc += renderToString(SessionValues(closest, {fontSize: 9}));
 
                 const timeStamp = handler.commitIdDates.get(closest.commitId)!.getTime() / 1000;
@@ -382,9 +382,9 @@ const PieChart = (container: SVGGElement, data: ItemStack[], radius: number, inn
 
 const SessionPieChart = (container: SVGGElement, session: TestSessionResult, radius: number, innerRadius:number = 0, padAngle: number = 0) => {
 
-    const sessionTotalCount = session.phaseSucceededCount + session.phaseFailedCount + session.phaseUnspecifiedCount;
-    const metaFailedFactor = Math.ceil(session.phaseFailedCount / (sessionTotalCount || 1) * 10) / 10;
-    const metaUnspecifiedFactor = Math.ceil(session.phaseUnspecifiedCount / (sessionTotalCount || 1) * 10) / 10;
+    const sessionTotalCount = session.phasesSucceededCount + session.phasesFailedCount + session.phasesUnspecifiedCount;
+    const metaFailedFactor = Math.ceil(session.phasesFailedCount / (sessionTotalCount || 1) * 10) / 10;
+    const metaUnspecifiedFactor = Math.ceil(session.phasesUnspecifiedCount / (sessionTotalCount || 1) * 10) / 10;
 
     const statusColors = dashboard.getStatusColors();
 
@@ -429,9 +429,9 @@ const ConcentricCircle = (container: SVGGElement, data: ItemStack[], radius: num
 
 const SessionConcentricCircle = (container: SVGGElement, session: TestSessionResult, radius: number) => {
 
-    const sessionTotalCount = session.phaseSucceededCount + session.phaseFailedCount + session.phaseUnspecifiedCount;
-    const metaFailedFactor = Math.ceil(clampMin(session.phaseFailedCount / (sessionTotalCount || 1)) * 5) / 5;
-    const metaUnspecifiedFactor = Math.ceil(clampMin(session.phaseUnspecifiedCount / (sessionTotalCount || 1)) * 5) / 5;
+    const sessionTotalCount = session.phasesSucceededCount + session.phasesFailedCount + session.phasesUnspecifiedCount;
+    const metaFailedFactor = Math.ceil(clampMin(session.phasesFailedCount / (sessionTotalCount || 1)) * 5) / 5;
+    const metaUnspecifiedFactor = Math.ceil(clampMin(session.phasesUnspecifiedCount / (sessionTotalCount || 1)) * 5) / 5;
     let metaSuccessFactor = 1 - (metaFailedFactor + metaUnspecifiedFactor);
     metaSuccessFactor = metaSuccessFactor < 0? 0 : metaSuccessFactor
 
