Horde: Updating schema documentation

#rnx

diff --git a/Engine/Source/Programs/Horde/Docs/Config/Schema/Globals.md b/Engine/Source/Programs/Horde/Docs/Config/Schema/Globals.md
--- a/Engine/Source/Programs/Horde/Docs/Config/Schema/Globals.md
+++ b/Engine/Source/Programs/Horde/Docs/Config/Schema/Globals.md
@@ -281,7 +281,17 @@ Name | Description
 `responseBucketId` | `string`<br>Name of the output bucket
 `condition` | `string`<br>Filter for agents to include
 `acl` | [AclConfig](#aclconfig)<br>Access control list
-`ubaCacheEndpoint` | `string`<br>Network endpoint of an Unreal Build Accelerator (UBA) cache server in IP:port format. Must be accessible to both agents and initiators running UBA / Unreal Build Tool. If null, use UBA cache functionality will be disabled for this compute cluster.
+`uba` | [UbaComputeClusterConfig](#ubacomputeclusterconfig)<br>Configuration settings for Unreal Build Accelerator (UBA) for a compute cluster
+
+## UbaComputeClusterConfig
+
+Configuration settings for Unreal Build Accelerator (UBA) for a compute cluster Contains cache server connectivity and other UBA-related settings. Ensure all cache related fields are non-null to activate use of cache.
+
+Name | Description
+---- | -----------
+`cacheEndpoint` | `string`<br>Host and port for connecting to the UBA cache server (host:port format) As specified with `-port=[PORT]` to UBA cache process. Serves all cache requests and must be accessible to both agents and initiators running UBA / Unreal Build Tool.
+`cacheHttpEndpoint` | `string`<br>Host and port for HTTP management of UBA cache server (host:port format) As specified with `-http=[PORT]` to UBA cache process. This port is used by Horde to manage the UBA cache, such as adding new cache session keys for initiators For improved security, only allow Horde server to access this port
+`cacheHttpCrypto` | `string`<br>AES encryption key used for secure Horde-to-cache-server communication As specified with `-httpcrypto=[KEY]` to UBA cache process.
 
 ## PoolConfig
 
diff --git a/Engine/Source/Programs/Horde/Docs/Deployment/ServerSettings.md b/Engine/Source/Programs/Horde/Docs/Deployment/ServerSettings.md
--- a/Engine/Source/Programs/Horde/Docs/Deployment/ServerSettings.md
+++ b/Engine/Source/Programs/Horde/Docs/Deployment/ServerSettings.md
@@ -132,7 +132,7 @@ Name | Description
 `compute` | [ComputeServerConfig](#computeserverconfig)<br>Configuration for the compute plugin
 `ddc` | [PluginServerConfig](#pluginserverconfig)<br>Configuration for the ddc plugin
 `experimental` | [ExperimentalServerConfig](#experimentalserverconfig)<br>Configuration for the experimental plugin
-`secrets` | [PluginServerConfig](#pluginserverconfig)<br>Configuration for the secrets plugin
+`secrets` | [SecretsServerConfig](#secretsserverconfig)<br>Configuration for the secrets plugin
 `storage` | [StorageServerConfig](#storageserverconfig)<br>Configuration for the storage plugin
 `symbols` | [PluginServerConfig](#pluginserverconfig)<br>Configuration for the symbols plugin
 `tools` | [ToolsServerConfig](#toolsserverconfig)<br>Configuration for the tools plugin
@@ -377,6 +377,15 @@ Name | Description
 ---- | -----------
 `enabled` | `boolean`<br>Whether the plugin should be enabled or not
 
+## SecretsServerConfig
+
+Static configuration for the secrets plugin
+
+Name | Description
+---- | -----------
+`withAws` | `boolean`<br>Whether to enable Amazon Web Services (AWS) specific features
+`enabled` | `boolean`<br>Whether the plugin should be enabled or not
+
 ## StorageServerConfig
 
 Static settings for the storage system
