Horde: Update release notes

#rnx

diff --git a/Engine/Source/Programs/Horde/Docs/ReleaseNotes.md b/Engine/Source/Programs/Horde/Docs/ReleaseNotes.md
--- a/Engine/Source/Programs/Horde/Docs/ReleaseNotes.md
+++ b/Engine/Source/Programs/Horde/Docs/ReleaseNotes.md
@@ -1,5 +1,21 @@
 # Release Notes
 
+## 2025-06-05
+
+* Add C# implementation of UBA crypto (43257937)
+* Horde Dashboard: Disable job step retries with expired TempStorage (43233123)
+* Fix sorting agents by status (43230711)
+* Support debugging jobs as scheduler (43229680)
+* Allow sessions to get job responses and post job metadata (43228360)
+* Add HTTP config values for UBA cache in compute cluster config (43225086)
+* Add API to resolve a string to a secret property value (43224431)
+* Update Secrets plugin with a new server configuration option WithAws that controls if the secret provider AwsParameterStore is to be enabled (43222303)
+* Adding capabilities to the warning defaults, cleaning up tests, adding more warning+error strings for easier long term adoption & toggling. (43196850)
+* Add a "No Suspects" emoji indication to build health slack threads that lack suspects (43191214)
+* Use preferred colors for device status
+  Slightly increased brightness threshold (43187852)
+* Use code commit tag to get suspect CLs (43161937)
+
 ## 2025-06-02
 
 * Cloud DDC - Added cross bucket search endpoint (43145204)
