Horde Server: Add job start time to JobStepRef response

diff --git a/Engine/Source/Programs/Horde/Plugins/Build/HordeServer.Build/Streams/StreamsController.cs b/Engine/Source/Programs/Horde/Plugins/Build/HordeServer.Build/Streams/StreamsController.cs
--- a/Engine/Source/Programs/Horde/Plugins/Build/HordeServer.Build/Streams/StreamsController.cs
+++ b/Engine/Source/Programs/Horde/Plugins/Build/HordeServer.Build/Streams/StreamsController.cs
@@ -328,6 +328,7 @@ namespace HordeServer.Streams
 			response.AgentId = jobStepRef.AgentId?.ToString();
 			response.Outcome = jobStepRef.Outcome;
 			response.IssueIds = jobStepRef.IssueIds?.Select(id => id).ToList();
+			response.JobStartTime = jobStepRef.JobStartTimeUtc;
 			response.StartTime = jobStepRef.StartTimeUtc;
 			response.FinishTime = jobStepRef.FinishTimeUtc;
 			response.Metadata = jobStepRef.Metadata?.ToList();
diff --git a/Engine/Source/Programs/Shared/EpicGames.Horde/Jobs/IJobStepRef.cs b/Engine/Source/Programs/Shared/EpicGames.Horde/Jobs/IJobStepRef.cs
--- a/Engine/Source/Programs/Shared/EpicGames.Horde/Jobs/IJobStepRef.cs
+++ b/Engine/Source/Programs/Shared/EpicGames.Horde/Jobs/IJobStepRef.cs
@@ -102,6 +102,11 @@ namespace HordeServer.Jobs
 		/// </summary>
 		public float BatchInitTime { get; }
 
+		/// <summary>
+		/// Time at which the job started
+		/// </summary>
+		public DateTime JobStartTimeUtc { get; }
+
 		/// <summary>
 		/// Time at which the step started.
 		/// </summary>
diff --git a/Engine/Source/Programs/Shared/EpicGames.Horde/Jobs/JobMessages.cs b/Engine/Source/Programs/Shared/EpicGames.Horde/Jobs/JobMessages.cs
--- a/Engine/Source/Programs/Shared/EpicGames.Horde/Jobs/JobMessages.cs
+++ b/Engine/Source/Programs/Shared/EpicGames.Horde/Jobs/JobMessages.cs
@@ -1550,6 +1550,11 @@ namespace EpicGames.Horde.Jobs
 		/// </summary>
 		public List<string>? Metadata { get; set; }
 
+		/// <summary>
+		/// Time at which the job started
+		/// </summary>
+		public DateTimeOffset JobStartTime { get; set; }
+
 		/// <summary>
 		/// Time at which the step started.
 		/// </summary>
