Add PreflightCommitId property to TestData

#rnx

diff --git a/Engine/Source/Programs/Horde/Plugins/Build/HordeServer.Build/Jobs/TestData/ITestData.cs b/Engine/Source/Programs/Horde/Plugins/Build/HordeServer.Build/Jobs/TestData/ITestData.cs
--- a/Engine/Source/Programs/Horde/Plugins/Build/HordeServer.Build/Jobs/TestData/ITestData.cs
+++ b/Engine/Source/Programs/Horde/Plugins/Build/HordeServer.Build/Jobs/TestData/ITestData.cs
@@ -585,6 +585,11 @@ namespace HordeServer.Jobs.TestData
 		/// </summary>
 		CommitIdWithOrder CommitId { get; }
 
+		/// <summary>
+		/// The preflight commit id if any
+		/// </summary>
+		CommitId? PreflightCommitId { get; }
+
 		/// <summary>
 		/// Key used to identify the particular data
 		/// </summary>
diff --git a/Engine/Source/Programs/Horde/Plugins/Build/HordeServer.Build/Jobs/TestData/TestDataCollection.cs b/Engine/Source/Programs/Horde/Plugins/Build/HordeServer.Build/Jobs/TestData/TestDataCollection.cs
--- a/Engine/Source/Programs/Horde/Plugins/Build/HordeServer.Build/Jobs/TestData/TestDataCollection.cs
+++ b/Engine/Source/Programs/Horde/Plugins/Build/HordeServer.Build/Jobs/TestData/TestDataCollection.cs
@@ -360,6 +360,9 @@ namespace HordeServer.Jobs.TestData
 			[BsonElement("Change")]
 			public int CommitOrder { get; set; }
 
+			[BsonIgnore]
+			public CommitId? PreflightCommitId => null;
+
 			public string Key { get; set; }
 			public BsonDocument Data { get; set; }
 
diff --git a/Engine/Source/Programs/Horde/Plugins/Build/HordeServer.Build/Jobs/TestData/TestDataCollectionV2.cs b/Engine/Source/Programs/Horde/Plugins/Build/HordeServer.Build/Jobs/TestData/TestDataCollectionV2.cs
--- a/Engine/Source/Programs/Horde/Plugins/Build/HordeServer.Build/Jobs/TestData/TestDataCollectionV2.cs
+++ b/Engine/Source/Programs/Horde/Plugins/Build/HordeServer.Build/Jobs/TestData/TestDataCollectionV2.cs
@@ -383,6 +383,15 @@ namespace HordeServer.Jobs.TestData
 			[BsonElement("Change")]
 			public int CommitOrder { get; set; }
 
+			[BsonIgnore]
+			public CommitId? PreflightCommitId
+			{
+				get => PreflightCommitName != null ? new CommitId(PreflightCommitName) : null;
+				set => PreflightCommitName = value?.Name;
+			}
+
+			public string? PreflightCommitName { get; set; }
+
 			public string Key { get; set; }
 			public BsonDocument Data { get; set; }
 
@@ -400,6 +409,7 @@ namespace HordeServer.Jobs.TestData
 				JobId = job.Id;
 				StepId = jobStep.Id;
 				CommitId = job.CommitId;
+				PreflightCommitId = job.PreflightCommitId;
 				Key = key;
 				Data = value;
 			}
diff --git a/Engine/Source/Programs/Horde/Plugins/Build/HordeServer.Build/Jobs/TestData/TestDataController.cs b/Engine/Source/Programs/Horde/Plugins/Build/HordeServer.Build/Jobs/TestData/TestDataController.cs
--- a/Engine/Source/Programs/Horde/Plugins/Build/HordeServer.Build/Jobs/TestData/TestDataController.cs
+++ b/Engine/Source/Programs/Horde/Plugins/Build/HordeServer.Build/Jobs/TestData/TestDataController.cs
@@ -522,6 +522,7 @@ namespace HordeServer.Jobs.TestData
 				JobId = testData.JobId.ToString(),
 				StepId = testData.StepId.ToString(),
 				CommitId = testData.CommitId,
+				PreflightCommitId = testData.PreflightCommitId,
 				Key = testData.Key,
 				Data = BsonSerializer.Deserialize<Dictionary<string, object>>(testData.Data)
 			}, filter);
diff --git a/Engine/Source/Programs/Shared/EpicGames.Horde/Jobs/TestData/TestDataMessages.cs b/Engine/Source/Programs/Shared/EpicGames.Horde/Jobs/TestData/TestDataMessages.cs
--- a/Engine/Source/Programs/Shared/EpicGames.Horde/Jobs/TestData/TestDataMessages.cs
+++ b/Engine/Source/Programs/Shared/EpicGames.Horde/Jobs/TestData/TestDataMessages.cs
@@ -189,6 +189,11 @@ namespace EpicGames.Horde.Jobs.TestData
 		}
 		CommitIdWithOrder? _commitId;
 
+		/// <summary>
+		/// The preflight commit id if any
+		/// </summary>
+		public CommitId? PreflightCommitId { get; set; }
+
 		/// <summary>
 		/// Key used to identify the particular data
 		/// </summary>
