Fix display issue of critical level event and unmatching history

#rnx

diff --git a/Engine/Source/Programs/Horde/HordeDashboard/plugins/test-automation/testAutomationEvents.tsx b/Engine/Source/Programs/Horde/HordeDashboard/plugins/test-automation/testAutomationEvents.tsx
--- a/Engine/Source/Programs/Horde/HordeDashboard/plugins/test-automation/testAutomationEvents.tsx
+++ b/Engine/Source/Programs/Horde/HordeDashboard/plugins/test-automation/testAutomationEvents.tsx
@@ -141,14 +141,13 @@ export const PhaseEventsPane: React.FC<{phase: TestPhaseStatus, onEventsLoaded?:
     }, [events]);
 
     const isEntryNeedDisplay = useCallback((event: EventEntry): boolean => {
-        const eventType = event.level;
-        if (filterError && eventType === EventLevel.Error) {
+        if (filterError && (event.level === EventLevel.Error || event.level === EventLevel.Critical)) {
             return true;
         }
-        if (filterWarning && eventType === EventLevel.Warning) {
+        if (filterWarning && event.level === EventLevel.Warning) {
             return true;
         }
-        if (filterInfo && eventType === EventLevel.Information) {
+        if (filterInfo) {
             return true;
         }
 
diff --git a/Engine/Source/Programs/Horde/HordeDashboard/plugins/test-automation/testAutomationPhaseView.tsx b/Engine/Source/Programs/Horde/HordeDashboard/plugins/test-automation/testAutomationPhaseView.tsx
--- a/Engine/Source/Programs/Horde/HordeDashboard/plugins/test-automation/testAutomationPhaseView.tsx
+++ b/Engine/Source/Programs/Horde/HordeDashboard/plugins/test-automation/testAutomationPhaseView.tsx
@@ -24,7 +24,7 @@ const PhaseHistoryPane: React.FC<{ test: TestNameRef, meta?: MetadataRef, phaseK
                 <Stack style={{ marginLeft: 30, marginBottom: 18, width: 10, cursor: "default" }} title="History">
                     <FontIcon style={{ fontSize: 15 }} iconName={"History"} />
                 </Stack>
-                {!metaStatus || !sessionId && <Text style={{ marginLeft: 10,marginBottom: 21}}>No history available for this phase</Text>}
+                {(!metaStatus || !sessionId) && <Text style={{ marginLeft: 10,marginBottom: 21}}>No history available for this phase</Text>}
                 {!!metaStatus && !!sessionId && <PhaseHistoryGraphWidget
                     test={test}
                     streamId={streamId!}
