Check for preflight data and set preflight flag on ChangeButton widget

#rnx

diff --git a/Engine/Source/Programs/Horde/HordeDashboard/plugins/test-automation/api.ts b/Engine/Source/Programs/Horde/HordeDashboard/plugins/test-automation/api.ts
--- a/Engine/Source/Programs/Horde/HordeDashboard/plugins/test-automation/api.ts
+++ b/Engine/Source/Programs/Horde/HordeDashboard/plugins/test-automation/api.ts
@@ -12,6 +12,9 @@ export type GetTestDataResponse = {
 	/// The commit id related to the job
 	commitId: {name: string, order: number};
 
+    /// The prefligth commit id if any
+    preflightCommitId: string;
+
 	/// The job id of this test data
 	jobId: string;
 
diff --git a/Engine/Source/Programs/Horde/HordeDashboard/plugins/test-automation/phaseCompare.ts b/Engine/Source/Programs/Horde/HordeDashboard/plugins/test-automation/phaseCompare.ts
--- a/Engine/Source/Programs/Horde/HordeDashboard/plugins/test-automation/phaseCompare.ts
+++ b/Engine/Source/Programs/Horde/HordeDashboard/plugins/test-automation/phaseCompare.ts
@@ -10,6 +10,7 @@ export const getJobParams = (fetchedSessionDetails: TestSessionDetails, handler:
             id: fetchedSessionDetails.jobId,
             streamId: fetchedSessionDetails.streamId,
             change: fetchedSessionDetails.commitOrder,
+            preflightChange: fetchedSessionDetails.preflightCommitId ? Number.parseInt(fetchedSessionDetails.preflightCommitId) || undefined : undefined
         },
         rangeCL: fetchedSessionDetails.session ?
             handler.getStatusStream(fetchedSessionDetails.streamId)?.getPreviousSession(fetchedSessionDetails.test!, fetchedSessionDetails.meta!, fetchedSessionDetails.session)?.commitOrder
diff --git a/Engine/Source/Programs/Horde/HordeDashboard/plugins/test-automation/testData.ts b/Engine/Source/Programs/Horde/HordeDashboard/plugins/test-automation/testData.ts
--- a/Engine/Source/Programs/Horde/HordeDashboard/plugins/test-automation/testData.ts
+++ b/Engine/Source/Programs/Horde/HordeDashboard/plugins/test-automation/testData.ts
@@ -393,6 +393,7 @@ export class TestSessionDetails {
     test?: TestNameRef;
     meta?: MetadataRef;
     session?: TestSessionResult;
+    preflightCommitId?: string;
 
     private _start: Date;
     private _duration: number;
diff --git a/Engine/Source/Programs/Horde/HordeDashboard/plugins/test-automation/testDataV2Fetcher.ts b/Engine/Source/Programs/Horde/HordeDashboard/plugins/test-automation/testDataV2Fetcher.ts
--- a/Engine/Source/Programs/Horde/HordeDashboard/plugins/test-automation/testDataV2Fetcher.ts
+++ b/Engine/Source/Programs/Horde/HordeDashboard/plugins/test-automation/testDataV2Fetcher.ts
@@ -190,7 +190,7 @@ export class TestDataV2 implements TestVersionFetcher {
         return result;
     }
 
-    private createSessionDetails(testData: TestData) {
+    private createSessionDetails(testData: TestData): TestSessionDetails {
         const testDetails = testData.data as TestSessionDetailsResponse;
         const start = new Date(testDetails.summary.dateTime);
         const result = new TestSessionDetails(testData.key, testDetails.summary.testName, start, testDetails.summary.timeElapseSec, testData.commitId.name, testData.commitId.order, testData.streamId, testData.jobId, testData.stepId);
@@ -207,6 +207,7 @@ export class TestDataV2 implements TestVersionFetcher {
                                 TestOutcome.Unspecified
                                 : testDetails.summary.phasesTotalCount > 0 ?
                                     TestOutcome.Success : TestOutcome.Skipped;
+        result.preflightCommitId = testData.preflightCommitId;
 
         // search for known metadata
         const metaFilter = new Map(Object.entries(testDetails.metadata).map(([key, value]) => [key, [value]]));
