Horde: Fix compiler warnings

diff --git a/Engine/Source/Programs/Horde/Plugins/Secrets/HordeServer.Secrets.Tests/Secrets/SecretResolveTests.cs b/Engine/Source/Programs/Horde/Plugins/Secrets/HordeServer.Secrets.Tests/Secrets/SecretResolveTests.cs
--- a/Engine/Source/Programs/Horde/Plugins/Secrets/HordeServer.Secrets.Tests/Secrets/SecretResolveTests.cs
+++ b/Engine/Source/Programs/Horde/Plugins/Secrets/HordeServer.Secrets.Tests/Secrets/SecretResolveTests.cs
@@ -17,7 +17,7 @@ namespace HordeServer.Tests.Secrets
 	{
 		private static ISecretCollection s_secretCollection = default!;
 		private static GlobalConfig s_globalConfig = default!;
-		private static Mock<ILogger<SecretCollectionInternal>> s_logger = new();
+		private static readonly Mock<ILogger<SecretCollectionInternal>> s_logger = new();
 
 		private class OptionsSnapshotStub<T>(T value) : IOptionsSnapshot<T> where T : class
 		{
@@ -29,7 +29,7 @@ namespace HordeServer.Tests.Secrets
 		}
 
 		[ClassInitialize]
-		public static async Task CreateSecretCollectionAsync(TestContext testContext)
+		public static async Task CreateSecretCollectionAsync(TestContext _)
 		{
 			string json = """
 				{
