Horde: Use new encrypted /addsession endpoint for allocating UBA cache sessions

diff --git a/Engine/Source/Programs/Horde/Plugins/Compute/HordeServer.Compute/Compute/ComputeService.cs b/Engine/Source/Programs/Horde/Plugins/Compute/HordeServer.Compute/Compute/ComputeService.cs
--- a/Engine/Source/Programs/Horde/Plugins/Compute/HordeServer.Compute/Compute/ComputeService.cs
+++ b/Engine/Source/Programs/Horde/Plugins/Compute/HordeServer.Compute/Compute/ComputeService.cs
@@ -428,7 +428,7 @@ namespace HordeServer.Compute
 				ubaCache = AllocateUbaCacheServerAsync(_httpClient,
 					clusterConfig.Uba.CacheEndpoint,
 					clusterConfig.Uba.CacheHttpEndpoint,
-					clusterConfig.Uba.CacheHttpCrypto, cancellationToken);
+					clusterConfig.Uba.CacheHttpCrypto, true, 120, cancellationToken);
 			}
 
 			try
@@ -883,21 +883,31 @@ namespace HordeServer.Compute
 		}
 		
 		internal static async Task<UbaComputeResource> AllocateUbaCacheServerAsync(HttpClient httpClient,
-			string ubaCacheEndpoint, string ubaCacheHttpEndpoint, string ubaCacheHttpCrypto, CancellationToken cancellationToken)
+			string ubaCacheEndpoint, string ubaCacheHttpEndpoint, string ubaCacheHttpCrypto,
+			bool readOnly, int ttl, CancellationToken cancellationToken)
 		{
 			try
 			{
 				using CancellationTokenSource timeoutCts = CancellationTokenSource.CreateLinkedTokenSource(cancellationToken);
 				timeoutCts.CancelAfter(TimeSpan.FromSeconds(5));
 				
-				_ = ubaCacheHttpCrypto;
-				
 				// Generate a session encryption key that we activate in the UBA cache server
-				string key = RandomNumberGenerator.GetHexString(32, true);
-				HttpResponseMessage res = await httpClient.GetAsync(new Uri($"http://{ubaCacheHttpEndpoint}/addcrypto?{key}"), timeoutCts.Token);
-				res.EnsureSuccessStatusCode();
+				byte[] sessionKey = RandomNumberGenerator.GetBytes(16);
+				string sessionKeyHex = Convert.ToHexString(sessionKey);
+				byte[] cryptoKey = Convert.FromHexString(ubaCacheHttpCrypto);
+				if (cryptoKey.Length != 16)
+				{
+					throw new Exception($"Invalid crypto key length: {cryptoKey.Length}");
+				}
 				
-				return new UbaComputeResource(ubaCacheEndpoint, key);
+				using UbaCrypto crypto = new(cryptoKey);
+				string encryptedSessionKeyHex = Convert.ToHexString(crypto.Encrypt(sessionKey));
+				string readOnlyStr = readOnly ? "true" : "false";
+				Uri uri = new ($"http://{ubaCacheHttpEndpoint}/addsession?crypto={encryptedSessionKeyHex},readonly={readOnlyStr},expiration={ttl}");
+				HttpResponseMessage res = await httpClient.GetAsync(uri, timeoutCts.Token);
+				res.EnsureSuccessStatusCode();
+
+				return new UbaComputeResource(ubaCacheEndpoint, sessionKeyHex);
 			}
 			catch (TaskCanceledException ex)
 			{
