Horde: Fix device checkout expiration link and clarify instructions

diff --git a/Engine/Source/Programs/Horde/Plugins/Build/HordeServer.Build/Devices/DeviceService.cs b/Engine/Source/Programs/Horde/Plugins/Build/HordeServer.Build/Devices/DeviceService.cs
--- a/Engine/Source/Programs/Horde/Plugins/Build/HordeServer.Build/Devices/DeviceService.cs
+++ b/Engine/Source/Programs/Horde/Plugins/Build/HordeServer.Build/Devices/DeviceService.cs
@@ -192,7 +192,7 @@ namespace HordeServer.Devices
 					{
 						foreach ((UserId, IDevice) expiredDevice in expireNotifications)
 						{
-							await NotifyDeviceServiceAsync(buildConfig, $"Device {expiredDevice.Item2.PlatformId.ToString().ToUpperInvariant()} / {expiredDevice.Item2.Name} checkout will expire in 24 hours.  Please visit {_serverInfo.DashboardUrl}/devices to renew the checkout if needed.", null, null, null, expiredDevice.Item1, cancellationToken);
+							await NotifyDeviceServiceAsync(buildConfig, $"Device {expiredDevice.Item2.PlatformId.ToString().ToUpperInvariant()} / {expiredDevice.Item2.Name} checkout will expire in 24 hours.  Please visit {_serverInfo.DashboardUrl}devices to check in and back out if needed.", null, null, null, expiredDevice.Item1, cancellationToken);
 						}
 					}
 
@@ -202,7 +202,7 @@ namespace HordeServer.Devices
 					{
 						foreach ((UserId, IDevice) expiredDevice in expireCheckouts)
 						{
-							await NotifyDeviceServiceAsync(buildConfig, $"Device {expiredDevice.Item2.PlatformId.ToString().ToUpperInvariant()} / {expiredDevice.Item2.Name} checkout has expired.  The device has been returned to the shared pool and should no longer be accessed.  Please visit {_serverInfo.DashboardUrl}/devices to checkout devices as needed.", null, null, null, expiredDevice.Item1, cancellationToken);
+							await NotifyDeviceServiceAsync(buildConfig, $"Device {expiredDevice.Item2.PlatformId.ToString().ToUpperInvariant()} / {expiredDevice.Item2.Name} checkout has expired.  The device has been returned to the shared pool and should no longer be accessed.  Please visit {_serverInfo.DashboardUrl}devices to checkout devices as needed.", null, null, null, expiredDevice.Item1, cancellationToken);
 						}
 					}
 				}
