Horde: Fix no suspects emoji not displaying
#rnx

diff --git a/Engine/Source/Programs/Horde/Plugins/Build/HordeServer.Build/Notifications/Sinks/SlackNotificationSink.cs b/Engine/Source/Programs/Horde/Plugins/Build/HordeServer.Build/Notifications/Sinks/SlackNotificationSink.cs
--- a/Engine/Source/Programs/Horde/Plugins/Build/HordeServer.Build/Notifications/Sinks/SlackNotificationSink.cs
+++ b/Engine/Source/Programs/Horde/Plugins/Build/HordeServer.Build/Notifications/Sinks/SlackNotificationSink.cs
@@ -1435,7 +1435,7 @@ namespace HordeServer.Notifications.Sinks
 			{
 				reactions |= ReactionFlags.AssignedButNotAcknowledged;
 			}
-			if (issue.OwnerId == null && suspects.Count == 0)
+			if (issue.OwnerId == null && (suspects.Count == 0 || suspects.Count > 5))
 			{
 				reactions |= ReactionFlags.NoSuspects;
 			}
