Horde Dashboard: Show "time to step" in trends panel

diff --git a/Engine/Source/Programs/Horde/HordeDashboard/src/backend/Api.ts b/Engine/Source/Programs/Horde/HordeDashboard/src/backend/Api.ts
--- a/Engine/Source/Programs/Horde/HordeDashboard/src/backend/Api.ts
+++ b/Engine/Source/Programs/Horde/HordeDashboard/src/backend/Api.ts
@@ -1585,6 +1585,8 @@ export interface GetJobStepRefResponse {
 	/**Outcome of the step, once complete. */
 	outcome?: JobStepOutcome;
 
+	jobStartTime: Date | string;
+
 	/**Time at which the step started. */
 	startTime: Date | string;
 
diff --git a/Engine/Source/Programs/Horde/HordeDashboard/src/components/jobDetailsV2/JobDetailStepTrendsV2.tsx b/Engine/Source/Programs/Horde/HordeDashboard/src/components/jobDetailsV2/JobDetailStepTrendsV2.tsx
--- a/Engine/Source/Programs/Horde/HordeDashboard/src/components/jobDetailsV2/JobDetailStepTrendsV2.tsx
+++ b/Engine/Source/Programs/Horde/HordeDashboard/src/components/jobDetailsV2/JobDetailStepTrendsV2.tsx
@@ -1,8 +1,8 @@
 // Copyright Epic Games, Inc. All Rights Reserved.
 
-import { Slider, Spinner, SpinnerSize, Stack, Text } from "@fluentui/react";
+import { Slider, Spinner, SpinnerSize, Stack, Text, Toggle } from "@fluentui/react";
 import * as d3 from "d3";
-import { action, makeObservable, observable } from "mobx";
+import { action, computed, makeObservable, observable } from "mobx";
 import { observer } from "mobx-react-lite";
 import moment, { Moment } from "moment";
 import { useEffect, useState } from "react";
@@ -93,25 +93,22 @@ class StepTrendsDataView extends JobDataView {
       durations.clear();
 
       history.forEach(h => {
-
+         // Cache step timestamps
          this.times.set(h.startTime as string, new Date(h.startTime));
          if (h.finishTime) {
             this.times.set(h.finishTime as string, new Date(h.finishTime));
          }
 
-         let end: Moment;
-         if (h.finishTime) {
-            end = moment(h.finishTime);
-         } else {
-            end = moment();
-         }
-
-         const start = moment(h.startTime);
-
-         if (!this.minTime || this.minTime.getTime() > (start.unix() * 1000)) {
-            this.minTime = start.toDate();
+         const start = moment(this.useFromJobStart ? h.jobStartTime : h.startTime);
+         const end = h.finishTime ? moment( h.finishTime) : moment();
+         
+         // Update oldest start time of any given step
+         const stepStart = moment(h.startTime);
+         if (!this.minTime || this.minTime.getTime() > (stepStart.unix() * 1000)) {
+            this.minTime = stepStart.toDate();
          }
 
+         // Update newest end time of any given step
          if (!this.maxTime || this.maxTime.getTime() < (end.unix() * 1000)) {
             this.maxTime = end.toDate();
          }
@@ -123,6 +120,7 @@ class StepTrendsDataView extends JobDataView {
             this.maxMinutes = Math.max(minutes, this.maxMinutes);
          }
 
+         // Only add successful steps to the trend data
          if (h.outcome === JobStepOutcome.Success || h.outcome === JobStepOutcome.Warnings) {
             alltimes.push(minutes);
          } else {
@@ -134,10 +132,6 @@ class StepTrendsDataView extends JobDataView {
 
       history = history.sort((a, b) => this.times.get(b.startTime as string)!.getTime() - this.times.get(a.startTime as string)!.getTime())
 
-      history = history.filter(h => {
-         return (durations.get(h.jobId) ?? (24 * 60)) <= (23 * 60)
-      })
-
       const median = (arr: number[]): number | undefined => {
          if (!arr.length) return undefined;
          const s = [...arr].sort((a, b) => a - b);
@@ -226,6 +220,56 @@ class StepTrendsDataView extends JobDataView {
    @observable
    selectedAgentId?: string;
 
+   @observable
+   useFromJobStart = false;
+
+   @observable
+   timeScales = {
+      fromJobStart: 0,
+      fromStepStart: 0
+   }
+
+   get defaultTimeScale(): number {
+      return this.median ?? this.maxMinutes;
+   }
+
+   @computed
+   get timeScale(): number {
+      const scales = this.timeScales;
+
+      if(
+         (this.useFromJobStart && !scales.fromJobStart)
+         || (!this.useFromJobStart && !scales.fromStepStart)
+      ) {
+         return this.defaultTimeScale;
+      }
+
+      return this.useFromJobStart ? scales.fromJobStart : scales.fromStepStart;
+   }
+   
+   @action
+   setTimeScale(time: number) {
+      if (this.useFromJobStart) {
+         this.timeScales.fromJobStart = time;
+      }
+      else {
+         this.timeScales.fromStepStart = time;
+      }
+
+      this.renderer?.onScaleTime?.(time);
+   }
+
+   @action
+   setUseFromJobStart(use: boolean) {
+      const prev = this.useFromJobStart;
+      this.useFromJobStart = use;
+
+      if (use !== prev) {
+         this.initData(this.history);
+         this.updateReady();
+      }
+   }
+
    @action
    setSelectedAgentId(agentId?: string) {
       if (this.selectedAgentId === agentId) {
@@ -250,6 +294,7 @@ class StepTrendsDataView extends JobDataView {
    minTime?: Date;
    maxTime?: Date;
 
+   lastZoom?: d3.ZoomTransform;
 
    tooltip = new Tooltip();
 
@@ -305,7 +350,7 @@ class StepTrendsRenderer {
          .range([margin.left, width - margin.right])
 
       const y = this.scaleY = d3.scaleLinear()
-         .domain([0, dataView.median ?? dataView.maxMinutes]).nice()
+         .domain([0, dataView.timeScale]).nice()
          .range([height - margin.bottom, margin.top])
 
 
@@ -544,10 +589,15 @@ class StepTrendsRenderer {
          dataView.tooltip.update(undefined)
          dataView.tooltip.freeze(false);
 
-         const zoomLevel = event.transform.k ? event.transform.k : 1;
+         dataView.lastZoom = event.transform;
+         zoomTo(event.transform);
+      }
+
+      function zoomTo(transform: d3.ZoomTransform) {
+         const zoomLevel = transform.k ? transform.k : 1;
          const barWidth = Math.max(2, zoomLevel / 2);
 
-         x.range([margin.left, width - margin.right].map(d => event.transform.applyX(d)));
+         x.range([margin.left, width - margin.right].map(d => transform.applyX(d)));
          svg!.selectAll(".bars rect")
             .attr("x", (d: any) => x(new Date(d.startTime!).getTime() / 1000) - (barWidth / 2))
             .attr("width", barWidth);
@@ -617,6 +667,17 @@ class StepTrendsRenderer {
 
       }
 
+      const updateTooltipArrow = (ref: GetJobStepRefResponse): void => {
+         const posx = x(new Date(ref.startTime!).getTime() / 1000);
+         const posy = this.scaleY!(dataView.durations.get(ref.jobId)!);
+
+         svg!.selectAll(".cline")
+            .attr("x1", () => posx)
+            .attr("x2", () => posx)
+            .attr("y1", d => 0)
+            .attr("y2", d => posy - 8);
+      }
+
       const handleMouseMove = (event: any) => {
 
          if (dataView.tooltip.frozen) {
@@ -629,15 +690,7 @@ class StepTrendsRenderer {
          const ref = closestData(mouseX, mouseY);
 
          if (ref) {
-            const posx = x(new Date(ref.startTime!).getTime() / 1000);
-            const posy = this.scaleY!(dataView.durations.get(ref.jobId)!);
-
-            svg!.selectAll(".cline")
-               .attr("x1", () => posx)
-               .attr("x2", () => posx)
-               .attr("y1", d => 0)
-               .attr("y2", d => posy - 8)
-
+            updateTooltipArrow(ref);
          } else {
             //svg!.selectAll(".cline rect")
             //.attr("width", 0);
@@ -663,6 +716,11 @@ class StepTrendsRenderer {
          dataView.tooltip.freeze(!dataView.tooltip.frozen);
       });
 
+      // Preserve some previous state on re-render
+      if(dataView.lastZoom) zoomTo(dataView.lastZoom);
+      if(dataView.tooltip.ref && dataView.tooltip.frozen) {
+         updateTooltipArrow(dataView.tooltip.ref);
+      }
    }
 
    showToolTip(shown: boolean = true) {
@@ -725,7 +783,7 @@ const GraphTooltip: React.FC<{ dataView: StepTrendsDataView }> = observer(({ dat
 
    const sindex = dataView.history.indexOf(ref);
 
-   const start = moment(ref.startTime);
+   const start = moment(dataView.useFromJobStart ? ref.jobStartTime : ref.startTime);
    let end = moment(Date.now());
 
    if (ref.finishTime) {
@@ -736,9 +794,11 @@ const GraphTooltip: React.FC<{ dataView: StepTrendsDataView }> = observer(({ dat
 
    const duration = getElapsedString(start, end);
 
-   const displayTime = moment(ref.startTime).tz(displayTimeZone());
+   const stepDisplayTime = moment(ref.startTime).tz(displayTimeZone());
+   const jobDisplayTime = moment(ref.jobStartTime).tz(displayTimeZone());
    const format = dashboard.display24HourClock ? "HH:mm:ss z" : "LT z";
-   let displayTimeStr = displayTime.format('MMM Do') + ` at ${displayTime.format(format)}`;
+   let displayTimeStr = stepDisplayTime.format('MMM Do') + ` at ${stepDisplayTime.format(format)}`;
+   let jobDisplayTimeStr = jobDisplayTime.format('MMM Do') + ` at ${jobDisplayTime.format(format)}`;
 
    const step = dataView.details?.stepById(dataView.stepId);
 
@@ -759,17 +819,17 @@ const GraphTooltip: React.FC<{ dataView: StepTrendsDataView }> = observer(({ dat
       transform: `translate(${translateX}, ${translateY})`
    }}>
       <Stack>
-         <Link to={`/job/${ref.jobId}?step=${ref.stepId}`}><Stack horizontal>
-            <StepRefStatusIcon stepRef={ref} />
-            <Text variant={textSize}>{step?.name}</Text>
-         </Stack>
+         <Link to={`/job/${ref.jobId}?step=${ref.stepId}`}>
+            <Stack horizontal>
+               <StepRefStatusIcon stepRef={ref} />
+               <Text variant={textSize}>{step?.name}</Text>
+            </Stack>
          </Link>
          <Stack style={{ paddingLeft: 2, paddingTop: 8 }} tokens={{ childrenGap: 8 }}>
-            <Stack>
-               <ChangeButton prefix="CL" job={dataView.details!.jobData!} stepRef={ref} hideAborted={true} rangeCL={sindex < (dataView.history.length - 1) ? (dataView.history[sindex + 1].change + 1) : undefined} />
-            </Stack>
-            <Stack><Text variant={textSize}>{displayTimeStr}</Text></Stack>
-            <Stack><Text variant={textSize}>Duration: {duration}</Text></Stack>
+            <ChangeButton prefix="CL" job={dataView.details!.jobData!} stepRef={ref} hideAborted={true} rangeCL={sindex < (dataView.history.length - 1) ? (dataView.history[sindex + 1].change + 1) : undefined} />
+            {dataView.useFromJobStart && <Text variant={textSize}>Job: {jobDisplayTimeStr}</Text>}
+            <Text variant={textSize}>{dataView.useFromJobStart && "Step: "}{displayTimeStr}</Text>
+            <Text variant={textSize}>Duration: {duration}</Text>
             <Stack style={{ cursor: "pointer" }} onClick={() => dataView.setSelectedAgentId(ref.agentId)}>
                <Text variant={textSize}>{ref.agentId!}</Text>
             </Stack>
@@ -855,26 +915,31 @@ export const StepTrendsPanelV2: React.FC<{ jobDetails: JobDetailsV2; stepId: str
                   <Text variant="mediumPlus" styles={{ root: { fontFamily: "Horde Open Sans SemiBold" } }}>{"Trends"}</Text>
                </Stack>
                <Stack grow />
-               {!!dataView.minTime && <Stack>
-                  <Slider
-                     styles={{ root: { width: 240 } }}
-                     label="Scale Time"
-                     min={1}
-                     max={dataView.maxMinutes}
-                     step={dataView.maxMinutes / 20}
-                     defaultValue={dataView.median ?? dataView.maxMinutes}
-                     showValue
-                     valueFormat={(value) => {
-                        if (!value) return ""
-                        return msecToElapsed(value * 60 * 1000, true, false)
-                     }}
-                     onChange={(time) => {
-                        if (dataView.renderer?.onScaleTime) {
-                           dataView.renderer.onScaleTime(time);
-                        }
-                     }}
-                  />
-               </Stack>}
+               {!!dataView.minTime && 
+                  <Stack horizontal horizontalAlign="end" verticalAlign="start" tokens={{childrenGap: 24}}>
+                     <Slider
+                        styles={{ root: { width: 240 } }}
+                        label="Scale Time"
+                        min={1}
+                        max={dataView.maxMinutes}
+                        step={dataView.maxMinutes / 20}
+                        value={dataView.timeScale}
+                        showValue
+                        valueFormat={(value) => {
+                           if (!value) return ""
+                           return msecToElapsed(value * 60 * 1000, true, false)
+                        }}
+                        onChange={(time) => dataView.setTimeScale(time)}
+                     />
+                     <Toggle 
+                        label="From Job Start"
+                        styles={{label: {paddingTop: 0}}}
+                        onChange={(_, checked) => {
+                           dataView.setUseFromJobStart(!!checked);
+                        }}
+                     />
+                  </Stack>
+               }
             </Stack>
             {!!dataView.history.length && <StepTrendGraph dataView={dataView} />}
             {!dataView.history.length && <Spinner size={SpinnerSize.large} />}
