TestDataV2 - Organize stream in section

#rnx

diff --git a/Engine/Source/Programs/Horde/HordeDashboard/plugins/test-automation/testAutomationCommon.tsx b/Engine/Source/Programs/Horde/HordeDashboard/plugins/test-automation/testAutomationCommon.tsx
--- a/Engine/Source/Programs/Horde/HordeDashboard/plugins/test-automation/testAutomationCommon.tsx
+++ b/Engine/Source/Programs/Horde/HordeDashboard/plugins/test-automation/testAutomationCommon.tsx
@@ -1,9 +1,10 @@
-import { mergeStyles, mergeStyleSets, Text, ComboBox, IComboBoxOption, SelectableOptionMenuItemType } from "@fluentui/react";
+import { mergeStyles, mergeStyleSets, Text, ComboBox, IComboBoxOption, SelectableOptionMenuItemType, IContextualMenuProps, IContextualMenuItem } from "@fluentui/react";
 import { MetadataRef, PhaseSessionResult, TestDataHandler, TestNameRef, TestPhaseStatus, TestSessionDetails, TestSessionResult } from "./testData";
 import dashboard, { StatusColor } from "horde/backend/Dashboard";
 import { useEffect, useState } from "react";
 import { TestOutcome, TestPhaseOutcome } from "./api";
 import { PhaseHistoryGraph } from "./phaseHistoryGraph";
+import { projectStore } from "horde/backend/ProjectStore";
 
 export type PhaseSessionCallback = (session: PhaseSessionResult) => void;
 
@@ -328,3 +329,66 @@ export const PhaseHistoryGraphWidget: React.FC<{test: TestNameRef, phaseKey: str
 
     return <div id={graph_container_id} style={{ userSelect: "none", position: "relative" }} ref={(ref: HTMLDivElement) => setContainer(ref)}/>
 }
+
+export const getStreamOptions = (streams: string[], onClick: (stream: string) => void, selected?: string): IContextualMenuProps => {
+    const projectItems: Map<string, Map<string, IContextualMenuItem[]>> = new Map();
+
+    projectStore.projects.sort(
+        (a, b) => a.order - b.order
+    ).forEach(p => {
+        p.categories?.forEach(category => {
+            category.streams.forEach(a => {
+                if (!category.showOnNavMenu || !streams.includes(a)) return;
+
+                const projectName = p.name ?? 'Unknown';
+                if (!projectItems.has(projectName)) {
+                    projectItems.set(projectName, new Map());
+                }
+
+                const categoryItems = projectItems.get(projectName)!;
+                if (!categoryItems.has(category.name)) {
+                    categoryItems.set(category.name, []);
+                }
+
+                categoryItems.get(category.name)?.push({
+                    key: `streams_${projectStore.streamById(a)?.fullname}`,
+                    category: category?.name,
+                    text: projectStore.streamById(a)?.name ?? a,
+                    onClick: () => onClick(a),
+                    disabled: a === selected
+                });
+            });
+        });
+    });
+
+    const options: IContextualMenuItem[] = projectItems.entries().map(
+        ([name, categories]) => {
+            return {
+                key: `project_${name}`,
+                text: name,
+                subMenuProps: {
+                    shouldFocusOnMount: true,
+                    subMenuHoverDelay: 0,
+                    items: categories.entries().map(
+                        ([cat, items]) => {
+                            return {
+                                key: `category_${cat}`,
+                                text: cat,
+                                subMenuProps: {
+                                    shouldFocusOnMount: true,
+                                    subMenuHoverDelay: 0,
+                                    items: items
+                                }
+                            }
+                       }         
+                    ).toArray()
+                }
+            }
+    }).toArray();
+
+    return {
+            shouldFocusOnMount: true,
+            subMenuHoverDelay: 0,
+            items: options,
+        };
+}
diff --git a/Engine/Source/Programs/Horde/HordeDashboard/plugins/test-automation/testAutomationSummary.tsx b/Engine/Source/Programs/Horde/HordeDashboard/plugins/test-automation/testAutomationSummary.tsx
--- a/Engine/Source/Programs/Horde/HordeDashboard/plugins/test-automation/testAutomationSummary.tsx
+++ b/Engine/Source/Programs/Horde/HordeDashboard/plugins/test-automation/testAutomationSummary.tsx
@@ -7,7 +7,7 @@ import dashboard, { StatusColor } from "horde/backend/Dashboard";
 import { getShortNiceTime } from "horde/base/utilities/timeUtils";
 import { getHordeStyling } from "horde/styles/Styles";
 import { projectStore } from 'horde/backend/ProjectStore';
-import { SessionStatusBar, SessionValues, StatusBar, StatusBarStack, MultiOptionChooser, getTestSessionStatusColor } from "./testAutomationCommon";
+import { SessionStatusBar, SessionValues, StatusBar, StatusBarStack, MultiOptionChooser, getTestSessionStatusColor, getStreamOptions } from "./testAutomationCommon";
 import { TestHistoryGraph } from "./testHistoryGraph";
 
 type SessionCallback = (session: TestSessionResult) => void;
@@ -50,21 +50,10 @@ const TestHistoryGraphWidget: React.FC<{ test: TestNameRef, stream: string, hand
 }
 
 const StreamSelector: React.FC<{ streams: string[], onClick: (stream: string) => void, selected?: string, disabled?: boolean }> = ({ streams, onClick, selected, disabled }) => {
-    const options: IContextualMenuItem[] = [];
-    streams.forEach(a => {
-        options.push({
-            key: `streams_${a}`,
-            text: projectStore.streamById(a)?.fullname ?? a,
-            onClick: () => onClick(a),
-            disabled: a === selected
-        });
-    });
 
-    const menuProps: IContextualMenuProps = {
-        shouldFocusOnMount: true,
-        subMenuHoverDelay: 0,
-        items: options,
-    };
+    const projectItems: Map<string, Map<string, IContextualMenuItem[]>> = new Map();
+
+    const menuProps: IContextualMenuProps = getStreamOptions(streams, onClick, selected);
 
     return <DefaultButton style={{ height: 22, fontSize: 11, whiteSpace: 'nowrap', borderWidth: 0 }} menuProps={menuProps} text={selected? projectStore.streamById(selected)?.fullname ?? selected : "Compare stream"} disabled={disabled}/>
 }
diff --git a/Engine/Source/Programs/Horde/HordeDashboard/plugins/test-automation/testAutomationView.tsx b/Engine/Source/Programs/Horde/HordeDashboard/plugins/test-automation/testAutomationView.tsx
--- a/Engine/Source/Programs/Horde/HordeDashboard/plugins/test-automation/testAutomationView.tsx
+++ b/Engine/Source/Programs/Horde/HordeDashboard/plugins/test-automation/testAutomationView.tsx
@@ -1,4 +1,4 @@
-import { Stack, FontIcon, IComboBoxOption, IContextualMenuItem, IContextualMenuProps, Text, Spinner, SpinnerSize, Label, DefaultButton, TagItemSuggestion, ITag, TagPicker, IPickerItemProps, PrimaryButton, Modal, IconButton, Link } from "@fluentui/react";
+import { Stack, FontIcon, IComboBoxOption, IContextualMenuItem, IContextualMenuProps, Text, Spinner, SpinnerSize, Label, DefaultButton, ITag, TagPicker, IPickerItemProps, PrimaryButton, Modal, IconButton } from "@fluentui/react";
 import { Breadcrumbs } from "horde/components/Breadcrumbs";
 import { TopNav } from "horde/components/TopNav";
 import { getHordeStyling } from "horde/styles/Styles";
@@ -8,7 +8,7 @@ import React, { useEffect, useState, useCallback } from "react";
 import { projectStore } from 'horde/backend/ProjectStore';
 import { TestDataHandler, TestDataVersionRegistrar, TestSessionDetails } from "./testData";
 import { TestSummary } from "./testAutomationSummary";
-import { getTestSessionStatusColor, MultiOptionChooser, SessionStatusBar } from "./testAutomationCommon";
+import { getStreamOptions, getTestSessionStatusColor, MultiOptionChooser, SessionStatusBar } from "./testAutomationCommon";
 import { TestPhasesView } from "./testAutomationPhaseView";
 import { JobStepTestDataItem } from "./api";
 
@@ -26,21 +26,8 @@ const StreamChooser: React.FC = observer(() => {
 
     const stream = handler.filterState.stream;
 
-    const options: IContextualMenuItem[] = [];
-    handler.availableStreams.forEach(a => {
-        options.push({
-            key: `streams_${a}`,
-            text: projectStore.streamById(a)?.fullname ?? a,
-            onClick: () => handler.selectStream(a)
-        });
-    });
-
-    const menuProps: IContextualMenuProps = {
-        shouldFocusOnMount: true,
-        subMenuHoverDelay: 0,
-        items: options,
-    };
-
+    const menuProps: IContextualMenuProps = getStreamOptions(handler.availableStreams, (stream) => handler.selectStream(stream), stream);
+    
     return <Stack style={{ paddingTop: 12, paddingBottom: 4 }} horizontal>
             <Label style={{ paddingTop: 5, paddingRight: 5 }}>Stream</Label>
             <DefaultButton style={{ width: 270, textAlign: "left", height: 30 }} menuProps={menuProps} text={stream ? projectStore.streamById(stream)?.fullname ?? stream : "Select"} />
