TestDataV2 - Phase header tweaking
- add a phase next/back navigation
- add phase warning next/back navigation
- hightlight phase row when hover on the phase list view
- tweak area to collapse/expand phase

#rnx

diff --git a/Engine/Source/Programs/Horde/HordeDashboard/plugins/test-automation/testAutomationEvents.tsx b/Engine/Source/Programs/Horde/HordeDashboard/plugins/test-automation/testAutomationEvents.tsx
--- a/Engine/Source/Programs/Horde/HordeDashboard/plugins/test-automation/testAutomationEvents.tsx
+++ b/Engine/Source/Programs/Horde/HordeDashboard/plugins/test-automation/testAutomationEvents.tsx
@@ -16,6 +16,8 @@ export const getEventStyling = () => {
     const colors = dashboard.getStatusColors();
     const errorColorPref = dashboard.getPreference(DashboardPreference.ColorError);
     const errorHighlight = errorColorPref ? isBright(errorColorPref) ? "brightness(0.9)" : "brightness(1.2)" : "";
+    const warningColorPref = dashboard.getPreference(DashboardPreference.ColorWarning);
+    const warningHighlight = warningColorPref ? isBright(warningColorPref) ? "brightness(0.9)" : "brightness(1.2)" : "";
 
     const eventClass = mergeStyles({
         fontSize: "12px",
@@ -38,6 +40,17 @@ export const getEventStyling = () => {
             borderLeftStyle: 'solid',
             borderLeftColor: colors.get(StatusColor.Failure)!
         }, eventClass],
+        warningButton: {
+            minWidth: 16,
+            backgroundColor: warningColorPref ? warningColorPref : dashboard.darktheme ? "#9D840E" : "#F7D154",
+            borderStyle: "hidden",
+            selectors: {
+                ':active,:hover': {
+                    backgroundColor: warningColorPref ? warningColorPref : "#E7C144",
+                    filter: warningHighlight
+                }
+            }
+        },
         failureButton: {
             backgroundColor: errorColorPref ? errorColorPref : dashboard.darktheme ? "#9D1410" : "#EC4C47",
             borderStyle: "hidden",
@@ -55,6 +68,17 @@ export const getEventStyling = () => {
             color: dashboard.darktheme ? "#909398" : "#616E85 !important",
             backgroundColor: dashboard.darktheme ? "#1F2223" : "#f3f2f1"
         },
+        defaultButton: {
+            borderStyle: 'solid', borderWidth: 1,
+            color: dashboard.darktheme ? "#909398" : "#616E85 !important",
+            backgroundColor: dashboard.darktheme ? "#1F2223" : "#f3f2f1",
+            borderColor: dashboard.darktheme ? "#4D4C4B" : "#6D6C6B",
+            selectors: {
+                ':active,:hover': {
+                    filter: errorHighlight
+                }
+            }
+        },
     });
 
     return eventStyles;
diff --git a/Engine/Source/Programs/Horde/HordeDashboard/plugins/test-automation/testAutomationPhaseView.tsx b/Engine/Source/Programs/Horde/HordeDashboard/plugins/test-automation/testAutomationPhaseView.tsx
--- a/Engine/Source/Programs/Horde/HordeDashboard/plugins/test-automation/testAutomationPhaseView.tsx
+++ b/Engine/Source/Programs/Horde/HordeDashboard/plugins/test-automation/testAutomationPhaseView.tsx
@@ -107,31 +107,46 @@ const TestPhaseExpandPane: React.FC<{ phase: TestPhaseStatus, onEventsLoaded?: (
             </Stack>
 });
 
-const TestPhaseHeader: React.FC<{ phase: TestPhaseStatus, expanded: boolean, toggleView: (phaseKey: string) => void }> = memo(({ phase, expanded, toggleView }) => {
+const TestPhaseHeader: React.FC<{ phase: TestPhaseStatus, expanded: boolean, toggleView: (phaseKey: string) => void, onNext?: () => void, onBack?: () => void }> = memo(({ phase, expanded, toggleView, onNext, onBack }) => {
     const onClick = useCallback(() => {
         toggleView(phase.key);
     }, [phase]);
 
     const color = getPhaseSessionStatusColor(phase);
-    return <Stack tokens={{childrenGap: 2}} disableShrink>
+    const eventStyles = getEventStyling();
+    return <Stack tokens={{childrenGap: 2}} disableShrink style={{cursor: "pointer" }} onClick={onClick}>
             <Stack horizontal verticalAlign="center">
-                <Stack style={{ paddingTop: 1, paddingRight: 4, cursor: "pointer" }} onClick={onClick}>
+                <Stack style={{ paddingTop: 1, paddingRight: 4}}>
                     <FontIcon style={{ fontSize: 13 }} iconName={expanded ? "ChevronDown" : "ChevronRight"} />
                 </Stack>
-                <Text style={{fontSize: 13, fontWeight: 600}}>{phase.name}</Text>
-                {!!expanded &&
-                    <Stack horizontalAlign="end" grow>
-                        <Stack style={{ padding: 3, cursor: "pointer" }}>
-                            <FontIcon
-                                style={{ fontSize: 12 }}
-                                iconName={"ChromeMinimize"}
-                                onClick={onClick}
-                            />
+                <Text style={{fontSize: 13, fontWeight: 600, whiteSpace: 'nowrap', overflow: 'hidden', textOverflow: 'ellipsis', cursor: "text"}} onClick={(ev) => ev.stopPropagation()}>{phase.name}</Text>
+                <Stack horizontalAlign="end" grow style={{paddingTop: 4, paddingRight: 12}}>
+                    {!!expanded &&
+                        <Stack horizontal>
+                            <DefaultButton
+                                className={eventStyles.defaultButton} title="Navigate to next phase"
+                                style={{ height: 24, minWidth: 24, whiteSpace: 'nowrap', fontSize: 12, paddingLeft: 4, paddingRight: 4, borderRadius: '4px 0px 0px 4px' }}
+                                onClick={(ev) => { ev.stopPropagation(); onNext && onNext() }}
+                                >
+                                    <Icon style={{ fontSize: 12 }} iconName='ChevronDown' />
+                            </DefaultButton>
+                            <DefaultButton
+                                className={eventStyles.defaultButton} title="Navigate to previous phase"
+                                style={{ height: 24, minWidth: 24, whiteSpace: 'nowrap', fontSize: 12, paddingLeft: 4, paddingRight: 4, borderRadius: 0 }}
+                                onClick={(ev) => { ev.stopPropagation(); onBack && onBack()}}
+                                >
+                                    <Icon style={{ fontSize: 12 }} iconName='ChevronUp' />
+                            </DefaultButton>
+                            <DefaultButton
+                                className={eventStyles.defaultButton} title="Collapse phase details"
+                                style={{ height: 24, minWidth: 24, whiteSpace: 'nowrap', fontSize: 12, paddingLeft: 4, paddingRight: 4, borderRadius: '0px 4px 4px 0px' }}>
+                                    <Icon style={{ fontSize: 12, paddingTop: 1 }} iconName='ChromeMinimize' />
+                            </DefaultButton>
                         </Stack>
-                    </Stack>
-                }
+                    }
+                </Stack>
             </Stack>
-            <Stack horizontal verticalAlign="center" onClick={onClick} style={{cursor: 'pointer'}}>
+            <Stack horizontal verticalAlign="center">
                 <Stack style={{ paddingLeft: 5, paddingRight: 4 }}>
                     <FontIcon style={{ fontSize: 12, color: color }} iconName="Square" />
                 </Stack>
@@ -140,7 +155,7 @@ const TestPhaseHeader: React.FC<{ phase: TestPhaseStatus, expanded: boolean, tog
         </Stack>
 });
 
-const TestPhaseSection: React.FC<{ phase: TestPhaseStatus, handler: TestDataHandler, onCollapse: (phaseKey: string) => void }> = ({ phase, handler, onCollapse }) => {
+const TestPhaseSection: React.FC<{ phase: TestPhaseStatus, handler: TestDataHandler, onCollapse: (phaseKey: string) => void, onNext: () => void, onBack: () => void }> = ({ phase, handler, onCollapse, onNext, onBack }) => {
     const toggleView = useCallback((phaseKey: string) => {
         comparisonContext.reset();
         handler.removeSearchParam('phase');
@@ -149,7 +164,7 @@ const TestPhaseSection: React.FC<{ phase: TestPhaseStatus, handler: TestDataHand
 
     const borderColor = dashboard.darktheme ? "#4D4C4B" : "#6D6C6B";
     return <Stack style={{ borderTopWidth: 1, borderTopColor: borderColor, borderTopStyle: 'solid', height: '100%'}}>
-                <TestPhaseHeader phase={phase} expanded={true} toggleView={toggleView} />
+                <TestPhaseHeader phase={phase} expanded={true} toggleView={toggleView} onNext={onNext} onBack={onBack} />
                 <TestPhaseExpandPane phase={phase} handler={handler}/>
             </Stack>
 }
@@ -172,8 +187,9 @@ const TestPhaseListItem: React.FC<{ phase: TestPhaseStatus, focused: boolean, ha
     }, []);
     
     const borderColor = dashboard.darktheme ? "#4D4C4B" : "#6D6C6B";
+    const highLightColor = dashboard.darktheme ? "rgba(255, 255, 255, 0.1)" : "rgba(0, 0, 0, 0.1)";
 
-    return <Stack style={{ borderTopWidth: 1, borderTopColor: borderColor, borderTopStyle: 'solid', paddingTop: 3, paddingBottom: 6 }}>
+    return <Stack styles={{ root: { borderTopWidth: 1, borderTopColor: borderColor, borderTopStyle: 'solid', paddingTop: 3, paddingBottom: 6, selectors: {':hover' : { backgroundColor: highLightColor }} } }}>
                 <div ref={registerRef}><TestPhaseHeader phase={phase} expanded={false} toggleView={toggleView} /></div>
             </Stack>
 });
@@ -293,6 +309,8 @@ export const TestPhasesView: React.FC<{ testDataId: string, handler: TestDataHan
     const [jobArtifacts, setJobArtifacts] = useState<GetArtifactResponse>();
     const selectedPhase = useRef<TestPhaseStatus>();
     const failureCursor = useRef<CircularList<string>>();
+    const warningCursor = useRef<CircularList<string>>();
+    const phaseCursor = useRef<CircularList<string>>();
     const listRef = useRef<List>();
     const jobContext = useRef<JobContext>();
 
@@ -305,6 +323,11 @@ export const TestPhasesView: React.FC<{ testDataId: string, handler: TestDataHan
         const failures = phases.filter((p) => p.outcome !== TestPhaseOutcome.Success && p.outcome !== TestPhaseOutcome.NotRun && p.outcome !== TestPhaseOutcome.Skipped).map((p) => p.key) ?? [];
         failureCursor.current = new CircularList(failures);
         failureCursor.current.setCursor(selectedPhase.current?.key);
+        phaseCursor.current = new CircularList(phases.map((p) => p.key));
+        phaseCursor.current.setCursor(selectedPhase.current?.key);
+        const warnings = phases.filter((p) => p.outcome !== TestPhaseOutcome.Success && p.hasWarning).map((p) => p.key) ?? [];
+        warningCursor.current = new CircularList(warnings);
+        warningCursor.current.setCursor(selectedPhase.current?.key);
         setFilteredPhases(phases);
     }, [handler.searchUpdated]);
 
@@ -353,11 +376,17 @@ export const TestPhasesView: React.FC<{ testDataId: string, handler: TestDataHan
         if (phaseKey !== failureCursor.current?.selected) {
             failureCursor.current?.setCursor(phaseKey);
         }
+        if (phaseKey !== warningCursor.current?.selected) {
+            warningCursor.current?.setCursor(phaseKey);
+        }
+        if (phaseKey !== phaseCursor.current?.selected) {
+            phaseCursor.current?.setCursor(phaseKey);
+        }
     }, [handler.searchUpdated, filteredPhases]);
 
     const StatusPie = useMemo(
         () => <Stack horizontal verticalAlign='center' style={{ height: '100%' }} tokens={{childrenGap: 3}}>
-                <Label>{filteredPhases.length} {filteredPhases.length > 1 ? "Phases" : "Phase"}</Label>
+                <Label style={{whiteSpace: 'nowrap'}}>{filteredPhases.length} {filteredPhases.length > 1 ? "Phases" : "Phase"}</Label>
                 {PhasesStatusPie(filteredPhases, 14)}
             </Stack>
         , [filteredPhases]);
@@ -369,10 +398,14 @@ export const TestPhasesView: React.FC<{ testDataId: string, handler: TestDataHan
     const buttonNoneTextColor = dashboard.darktheme ? "#949898" : "#616e85";
     const errorColorPref = dashboard.getPreference(DashboardPreference.ColorError);
     const failureButtonTextColor = errorColorPref && isBright(errorColorPref) ? "#2d3f5f" : "#F9F9FB";
+    const warningColorPref = dashboard.getPreference(DashboardPreference.ColorWarning);
+    const warningButtonTextColor = warningColorPref && isBright(warningColorPref) ? "#2d3f5f" : "#F9F9FB";
 
     const failureEnabled = !!failureCursor.current?.length;
     const failureText = failureEnabled ? failureCursor.current!.length.toString() : "No";
 
+    const warningEnabled = !!warningCursor.current?.length;
+
     return <Stack style={{width: '100%', height: '100%'}}>
                 {!!jobArtifacts && jobContext.current?.job?.id && jobContext.current?.step?.id &&
                     <JobArtifactsModal jobId={jobContext.current.job.id} stepId={jobContext.current?.step?.id} contextType={jobArtifacts.type} artifacts={[jobArtifacts]} onClose={() => setJobArtifacts(undefined)}/>
@@ -404,13 +437,13 @@ export const TestPhasesView: React.FC<{ testDataId: string, handler: TestDataHan
                     <Stack grow tokens={{childrenGap: 3}} style={{height: '100%'}}>
                         <Stack horizontal verticalAlign="start" tokens={{childrenGap: 10}}>
                             <Stack style={{ paddingLeft: 8}} tokens={{childrenGap: 5}}>
-                                <Text style={{ fontSize: 16, fontWeight: 600 }}>{`${sessionDetails.name} - ${metaString}`}</Text>
+                                <Text style={{ fontSize: 16, fontWeight: 600, whiteSpace: 'nowrap', overflow: 'hidden', textOverflow: 'ellipsis' }}>{`${sessionDetails.name} - ${metaString}`}</Text>
                                 <Stack style={{ paddingLeft: 8 }} horizontal tokens={{childrenGap: 8}}>
                                     <Text>{projectStore.streamById(sessionDetails.streamId)?.fullname ?? sessionDetails.streamId}</Text>
                                     <ChangeButton job={jobContext.current?.job} rangeCL={jobContext.current?.rangeCL} />
-                                    <Link href={`/job/${jobContext.current?.job?.id}?step=${jobContext.current?.step?.id}`} target="_blank"><Text style={{ fontSize: 12, color: hordeTheme.semanticColors.link }}>View Job</Text></Link>
-                                    <Text><span style={styles.labelSmall}>Run on: </span><span>{getShortNiceTime(sessionDetails.start, true, true, true)}</span></Text>
-                                    <Text><span style={styles.labelSmall}> For: </span><span>{msecToElapsed(sessionDetails.duration * 1000, true, true)}</span></Text>
+                                    <Link href={`/job/${jobContext.current?.job?.id}?step=${jobContext.current?.step?.id}`} target="_blank"><Text style={{ fontSize: 12, whiteSpace: 'nowrap', color: hordeTheme.semanticColors.link }}>View Job</Text></Link>
+                                    <Text style={{whiteSpace: 'nowrap'}}><span style={styles.labelSmall}>Run on: </span><span>{getShortNiceTime(sessionDetails.start, true, true, true)}</span></Text>
+                                    <Text style={{whiteSpace: 'nowrap'}}><span style={styles.labelSmall}> For: </span><span>{msecToElapsed(sessionDetails.duration * 1000, true, true)}</span></Text>
                                     <Stack horizontalAlign="end">
                                         <DefaultButton style={{ fontSize: 11, padding: 0, height: 18, backgroundColor: styles.defaultButton.backgroundColor, color: 'white' }}
                                             onClick={(ev) => {
@@ -430,11 +463,9 @@ export const TestPhasesView: React.FC<{ testDataId: string, handler: TestDataHan
                                     <TooltipHost content={failureEnabled ? 'Navigate to next failure': ''}>
                                         <DefaultButton
                                             disabled={!failureEnabled} className={failureEnabled ? eventStyles.failureButton : eventStyles.failureButtonDisabled}
-                                            style={{ color: failureEnabled ? failureButtonTextColor : buttonNoneTextColor, height: 28, fontSize: 12, paddingLeft: 8, paddingRight: 8, borderRadius: failureEnabled ? '4px 0 0 4px' : '4px' }}
+                                            style={{ color: failureEnabled ? failureButtonTextColor : buttonNoneTextColor, height: 28, whiteSpace: 'nowrap', fontSize: 12, paddingLeft: 8, paddingRight: 8, borderRadius: failureEnabled ? '4px 0 0 4px' : '4px' }}
                                             text={`${failureText} ${(failureCursor.current?.length ?? 0) > 1 ? "Failures" : "Failure"}`}
-                                            onClick={() => {
-                                                if (failureCursor.current) handler.setSearchParam('phase', failureCursor.current.next());
-                                            }} >
+                                            onClick={() => failureCursor.current && handler.setSearchParam('phase', failureCursor.current.next())} >
                                             {failureEnabled &&
                                                 <Icon style={{ fontSize: 19, paddingLeft: 4, paddingRight: 0 }} iconName='ChevronDown' />
                                             }
@@ -450,13 +481,35 @@ export const TestPhasesView: React.FC<{ testDataId: string, handler: TestDataHan
                                             disabled={!failureEnabled} className={failureEnabled ? eventStyles.failureButton : eventStyles.failureButtonDisabled}
                                             style={{ color: failureButtonTextColor, height: 28, fontSize: 12, padding: 2, borderRadius: '0 4px 4px 0' }}
                                             iconProps={{ iconName: 'ChevronUp' }}
-                                            onClick={() => {
-                                                if (failureCursor.current) handler.setSearchParam('phase', failureCursor.current.back());
-                                            }} />
+                                            onClick={() => failureCursor.current && handler.setSearchParam('phase', failureCursor.current.back())} />
                                     </TooltipHost>
                                 }
                             </Stack>
 
+                            {warningEnabled &&
+                                <Stack horizontal verticalAlign='center' style={{ height: '100%' }}>
+                                    <Stack horizontal verticalAlign='center' style={{position: 'relative'}}>
+                                        <TooltipHost content='Navigate to next warning'>
+                                            <DefaultButton
+                                                className={eventStyles.warningButton}
+                                                style={{ color: warningButtonTextColor, height: 28, whiteSpace: 'nowrap', fontSize: 12, paddingLeft: 8, paddingRight: 8, borderRadius: '4px 0 0 4px' }}
+                                                text={`${warningCursor.current?.length}${(failureCursor.current?.length ?? 0) === 0 ? (warningCursor.current?.length ?? 0) > 1 ? " Warnings" : " Warning" : ""}`}
+                                                onClick={() => warningCursor.current && handler.setSearchParam('phase', warningCursor.current.next())} >
+                                                <Icon style={{ fontSize: 19, paddingLeft: 4, paddingRight: 0 }} iconName='ChevronDown' />
+                                            </DefaultButton>
+                                        </TooltipHost>
+                                        <div style={{backgroundColor: warningButtonTextColor, height: 18, width: 1, position: 'absolute', right: -0.5, zIndex: 1 }} />
+                                    </Stack>
+                                    <TooltipHost content={'Navigate to previous warning'}>
+                                        <IconButton
+                                            className={eventStyles.warningButton}
+                                            style={{ color: warningButtonTextColor, height: 28, fontSize: 12, padding: 2, borderRadius: '0 4px 4px 0' }}
+                                            iconProps={{ iconName: 'ChevronUp' }}
+                                            onClick={() => warningCursor.current && handler.setSearchParam('phase', warningCursor.current.back())} />
+                                    </TooltipHost>
+                                </Stack>
+                            }
+
                             <Stack verticalAlign='center' style={{ height: '100%' }}>
                                 <TooltipHost content='Filter phases'>
                                     <PhaseFilter phases={sessionDetails.phases}
@@ -491,6 +544,8 @@ export const TestPhasesView: React.FC<{ testDataId: string, handler: TestDataHan
                                     phase={selectedPhase.current}
                                     handler={handler}
                                     onCollapse={(phaseKey) => {setFocusPhase(phaseKey)}}
+                                    onNext={() => phaseCursor.current && handler.setSearchParam('phase', phaseCursor.current.next())}
+                                    onBack={() => phaseCursor.current && handler.setSearchParam('phase', phaseCursor.current.back())}
                                 />
                             }
                         </Stack>
