Horde: Default to read/write access for UBA cache for now until per user permissions are added

diff --git a/Engine/Source/Programs/Horde/Plugins/Compute/HordeServer.Compute/Compute/ComputeService.cs b/Engine/Source/Programs/Horde/Plugins/Compute/HordeServer.Compute/Compute/ComputeService.cs
--- a/Engine/Source/Programs/Horde/Plugins/Compute/HordeServer.Compute/Compute/ComputeService.cs
+++ b/Engine/Source/Programs/Horde/Plugins/Compute/HordeServer.Compute/Compute/ComputeService.cs
@@ -428,7 +428,7 @@ namespace HordeServer.Compute
 				ubaCache = AllocateUbaCacheServerAsync(_httpClient,
 					clusterConfig.Uba.CacheEndpoint,
 					clusterConfig.Uba.CacheHttpEndpoint,
-					clusterConfig.Uba.CacheHttpCrypto, true, 120, cancellationToken);
+					clusterConfig.Uba.CacheHttpCrypto, false, 120, cancellationToken);
 			}
 
 			try
