- Fix missing skipped test display in test details pie chart
- Add phase outcome filtering
- Allow filtering options to applied across test sessions.

#rnx

diff --git a/Engine/Source/Programs/Horde/HordeDashboard/plugins/test-automation/testAutomationCommon.tsx b/Engine/Source/Programs/Horde/HordeDashboard/plugins/test-automation/testAutomationCommon.tsx
--- a/Engine/Source/Programs/Horde/HordeDashboard/plugins/test-automation/testAutomationCommon.tsx
+++ b/Engine/Source/Programs/Horde/HordeDashboard/plugins/test-automation/testAutomationCommon.tsx
@@ -1,4 +1,4 @@
-import { mergeStyles, mergeStyleSets, Text, ComboBox, IComboBoxOption, SelectableOptionMenuItemType, IContextualMenuProps, IContextualMenuItem } from "@fluentui/react";
+import { mergeStyles, mergeStyleSets, Text, ComboBox, IComboBoxOption, SelectableOptionMenuItemType, IContextualMenuProps, IContextualMenuItem, IconButton, IContextualMenuListProps, Stack, Checkbox, FontIcon } from "@fluentui/react";
 import { MetadataRef, PhaseSessionResult, TestDataHandler, TestNameRef, TestPhaseStatus, TestSessionDetails, TestSessionResult } from "./testData";
 import dashboard, { StatusColor } from "horde/backend/Dashboard";
 import { useEffect, useState } from "react";
@@ -219,13 +219,14 @@ export const SessionStatusPie = (session: TestSessionResult, radius: number): JS
 }
 
 export const PhasesStatusPie = (phases: TestPhaseStatus[], radius: number): JSX.Element | undefined => {
-    let totalCount = phases.length || 1;
+    const totalCount = phases.length || 1;
     let phasesFailedCount = 0;
     let phasesUnspecifiedCount = 0;
+    let phasesSkippedCount = 0
     phases.forEach(p => {
         switch (p.outcome) {
             case TestPhaseOutcome.Skipped:
-                --totalCount;
+                ++phasesSkippedCount;
                 break;
             case TestPhaseOutcome.Failed:
             case TestPhaseOutcome.Interrupted:
@@ -240,6 +241,7 @@ export const PhasesStatusPie = (phases: TestPhaseStatus[], radius: number): JSX.
 
     const failedFactor = Math.ceil(phasesFailedCount / totalCount * 100) / 100;
     const unspecifiedFactor = Math.ceil(phasesUnspecifiedCount / totalCount * 100) / 100;
+    const skippedFactor = Math.ceil(phasesSkippedCount / totalCount * 100) / 100;
 
     const statusColors = dashboard.getStatusColors();
 
@@ -249,13 +251,18 @@ export const PhasesStatusPie = (phases: TestPhaseStatus[], radius: number): JSX.
             title: "Unspecified",
             color: statusColors.get(StatusColor.Unspecified)!
         },
+        {
+            value: skippedFactor * 100,
+            title: "Skipped",
+            color: statusColors.get(StatusColor.Skipped)!
+        },
         {
             value: failedFactor * 100,
             title: "Failure",
             color: statusColors.get(StatusColor.Failure)!
         },
         {
-            value: (1 - (failedFactor + unspecifiedFactor)) * 100,
+            value: (1 - (failedFactor + unspecifiedFactor + skippedFactor)) * 100,
             title: "Passed",
             color: statusColors.get(StatusColor.Success)!
         }
@@ -392,3 +399,154 @@ export const getStreamOptions = (streams: string[], onClick: (stream: string) =>
             items: options,
         };
 }
+
+export type ICheckListOption = IContextualMenuItem & {
+    checked: boolean,
+    color?: string,
+}
+
+export const getPhaseOutcomeOptions = (): ICheckListOption[] => {
+    const statusColors = dashboard.getStatusColors();
+    return [
+        {
+            key: TestOutcome.Unspecified,
+            text: 'Unspecified',
+            checked: true,
+            color: statusColors.get(StatusColor.Unspecified)
+        },
+        {
+            key: TestOutcome.Skipped,
+            text: 'Skipped',
+            checked: true,
+            color: statusColors.get(StatusColor.Skipped)
+        },
+        {
+            key: TestOutcome.Failure,
+            text: 'Failed',
+            checked: true,
+            color: statusColors.get(StatusColor.Failure)
+        },
+        {
+            key: TestOutcome.Success,
+            text: 'Passed',
+            checked: true,
+            color: statusColors.get(StatusColor.Success)
+        },
+    ]
+}
+
+class PhaseFiltersContext {
+    constructor() {
+        this._outcomeOptions = getPhaseOutcomeOptions();
+        this._names = new Map();
+    }
+
+    get OutcomeOptions() {
+        return this._outcomeOptions
+    }
+
+    filterTestPhases(testId: string, phases: TestPhaseStatus[]) {
+        const outcomes = this.getFilteredOutcomes();
+        const names = this.getFilteredNames(testId);
+        return phases.filter(phase => phase.isMatch(names, outcomes));
+    }
+
+    getFilteredOutcomes() {
+        const outcome: Set<TestPhaseOutcome> = new Set();
+        this._outcomeOptions.filter(option => option.checked).forEach(option => {
+            switch(option.key) {
+                case TestOutcome.Unspecified:
+                    outcome.add(TestPhaseOutcome.NotRun);
+                    outcome.add(TestPhaseOutcome.Unknown);
+                    break;
+                case TestOutcome.Skipped:
+                    outcome.add(TestPhaseOutcome.Skipped);
+                    break;
+                case TestOutcome.Failure:
+                    outcome.add(TestPhaseOutcome.Failed);
+                    outcome.add(TestPhaseOutcome.Interrupted);
+                    break;
+                case TestOutcome.Success:
+                    outcome.add(TestPhaseOutcome.Success);
+                    break;
+            }
+        });
+        return outcome;
+    }
+
+    getFilteredNames(testId: string) {
+        return this._names.get(testId) ?? [];
+    }
+
+    setFilteredNames(testId: string, names: string[]) {
+        this._names.set(testId, names.map(n => n.toLowerCase()));
+    }
+
+    private _outcomeOptions: ICheckListOption[];
+    private _names: Map<string, string[]>;
+}
+export const phaseFiltersContext = new PhaseFiltersContext();
+
+export const CheckListOption: React.FC<{options: ICheckListOption[], onChange: (items: ICheckListOption[]) => void}> = ({options, onChange}) => {
+    const [items, setItems] = useState<ICheckListOption[]>([]);
+    const [update, setUpdate] = useState<number>(0);
+
+    useEffect(() => {
+        const initItems: ICheckListOption[] = [{
+            key: 'SelectAll',
+            text: 'Select All',
+            checked: !options.some(i => !i.checked)
+        }];
+        
+        initItems.push(...options);
+
+        setItems(initItems);
+        setUpdate(1);
+    }, [options]);
+
+    if (!items.length && !update) return;
+
+    const fitlerOptions: IContextualMenuProps = {
+        items: items,
+        onRenderMenuList: (props: IContextualMenuListProps) =>
+                <Stack tokens={{childrenGap: 6}} style={{margin: '3px 2px'}}>
+                    {props.items.map(item => <Checkbox
+                                                key={item.key}
+                                                label={item.text}
+                                                checked={item.checked}
+                                                indeterminate={item.key === 'SelectAll' && !item.checked && props.items.slice(1).some(i => i.checked)}
+                                                onChange={
+                                                    (ev, checked) => {
+                                                        item.checked = checked;
+                                                        if (item.key === 'SelectAll') {
+                                                            props.items.forEach(i => i.checked = checked);
+                                                        } else {
+                                                            props.items[0].checked = !props.items.slice(1).some(i => !i.checked);
+                                                        }
+                                                        setUpdate(update + 1);
+                                                        onChange(items.slice(1));
+                                                    }
+                                                }
+                                                onRenderLabel={(props) => {
+                                                    return <Stack horizontal verticalAlign="center">
+                                                                <Stack style={{ paddingRight: 5 }}>
+                                                                    <FontIcon style={{ fontSize: 12, color: item.color }} iconName={!!item.color? "Square" : "MultiSelect"} />
+                                                                </Stack>
+                                                                <Text>{props?.label}</Text>
+                                                            </Stack>
+                                                }}
+                                            />)}
+                </Stack>
+    }
+
+    return  <IconButton
+                iconProps={{iconName: "PageListFilter"}}
+                menuProps={fitlerOptions}
+                style={{
+                    fontSize:16, height: 28,
+                    border: '1px solid', borderColor: dashboard.darktheme ? "#4D4C4B" : "#6D6C6B", borderRadius: 4,
+                    backgroundColor: dashboard.darktheme ? "rgba(255, 255, 255, 0.1)" : "rgba(0, 0, 0, 0.3)",
+                    color: 'white'
+                }}
+            />
+}
diff --git a/Engine/Source/Programs/Horde/HordeDashboard/plugins/test-automation/testAutomationPhaseView.tsx b/Engine/Source/Programs/Horde/HordeDashboard/plugins/test-automation/testAutomationPhaseView.tsx
--- a/Engine/Source/Programs/Horde/HordeDashboard/plugins/test-automation/testAutomationPhaseView.tsx
+++ b/Engine/Source/Programs/Horde/HordeDashboard/plugins/test-automation/testAutomationPhaseView.tsx
@@ -9,7 +9,7 @@ import { getShortNiceTime, msecToElapsed } from "horde/base/utilities/timeUtils"
 import { getHordeTheme } from "horde/styles/theme";
 import { DashboardPreference, GetArtifactResponse } from "horde/backend/Api";
 import { isBright } from "horde/base/utilities/colors";
-import { getPhaseSessionStatusColor, PhaseHistoryGraphWidget, PhasesStatusPie, styles } from "./testAutomationCommon";
+import { CheckListOption, getPhaseSessionStatusColor, phaseFiltersContext, PhaseHistoryGraphWidget, PhasesStatusPie, styles } from "./testAutomationCommon";
 import { projectStore } from "horde/backend/ProjectStore";
 import { JobArtifactsModal } from "horde/components/artifacts/ArtifactsModal";
 import { getJobParams, JobContext } from "./phaseCompare";
@@ -216,7 +216,7 @@ const onRenderItem = (props: IPickerItemProps<ITag>) => {
 
 const getTextFromTag = (item: ITag) => item.name;
 
-const PhaseFilter: React.FC<{phases: TestPhaseStatus[], onChange: (items: string[]) => void, onFocus?: () => void, onBlur?: () => void}> = ({ phases, onChange, onFocus, onBlur }) => {
+const PhaseFilter: React.FC<{phases: TestPhaseStatus[], selected?: string[], onChange: (items: string[]) => void, onFocus?: () => void, onBlur?: () => void}> = ({ phases, selected, onChange, onFocus, onBlur }) => {
     const [filters, setFilters] = useState<ITag[]>([]);
     const phaseTags = useRef<ITag[]>();
 
@@ -224,6 +224,7 @@ const PhaseFilter: React.FC<{phases: TestPhaseStatus[], onChange: (items: string
         const tokens: ITag[] = (new Set(phases.map(p => (p.name.split('.'))).flat())).keys().map(token => ({key: token.toLowerCase(), name: token})).toArray().sort((a, b) => a.key.length - b.key.length);
         const tags: ITag[] = phases.map(p => ({key: `@${p.name.toLowerCase()}`, name: `@${p.name}`}));
         phaseTags.current = [...tokens, ...tags];
+        selected && setFilters(selected.map(i => ({key: i, name: i})));
     }, [phases]);
     
     const filterSelectedTests = useCallback((filterText: string): ITag[] => {
@@ -255,7 +256,7 @@ const PhaseFilter: React.FC<{phases: TestPhaseStatus[], onChange: (items: string
                     setFilters(items);
                     onChange(items.map(t => t.name));
                 }}
-                inputProps={{ placeholder: "Filter", onFocus: onFocus, onBlur: onBlur }}
+                inputProps={{ placeholder: "Filter names", onFocus: onFocus, onBlur: onBlur }}
             />
 }
 
@@ -335,25 +336,26 @@ export const TestPhasesView: React.FC<{ testDataId: string, handler: TestDataHan
         (phase: TestPhaseStatus) => <TestPhaseListItem key={phase.key} phase={phase} focused={focusPhase === phase.key} handler={handler}/>
         , [handler.searchUpdated]);
 
+    useEffect(() => {
+        return () => {
+            // on unmount
+            comparisonContext.reset();
+        }
+    }, [])
+
     useEffect(() => {
         handler.queryTestData(testDataId).then((fetchedSessionDetails) => {
             if (fetchedSessionDetails) {
                 handler.selectStream(fetchedSessionDetails.streamId);
                 setSessionDetails(fetchedSessionDetails);
-                updateFilter(fetchedSessionDetails.phases);
+                const testId = fetchedSessionDetails.test?.id ?? testDataId;
+                updateFilter(phaseFiltersContext.filterTestPhases(testId, fetchedSessionDetails.phases));
                 jobContext.current = getJobParams(fetchedSessionDetails, handler);
             }
         });
         setSessionDetails(undefined);
     }, [testDataId]);
 
-    useEffect(() => {
-        return () => {
-            // on unmount
-            comparisonContext.reset();
-        }
-    }, [])
-
     const onPagesUpdated = useCallback(() => {
         if (!focusPhase) return;
         const phaseIndex = filteredPhases.findIndex((phase) => phase.key === focusPhase);
@@ -406,6 +408,8 @@ export const TestPhasesView: React.FC<{ testDataId: string, handler: TestDataHan
 
     const warningEnabled = !!warningCursor.current?.length;
 
+    const testId = sessionDetails?.test?.id ?? testDataId;
+
     return <Stack style={{width: '100%', height: '100%'}}>
                 {!!jobArtifacts && jobContext.current?.job?.id && jobContext.current?.step?.id &&
                     <JobArtifactsModal jobId={jobContext.current.job.id} stepId={jobContext.current?.step?.id} contextType={jobArtifacts.type} artifacts={[jobArtifacts]} onClose={() => setJobArtifacts(undefined)}/>
@@ -511,12 +515,24 @@ export const TestPhasesView: React.FC<{ testDataId: string, handler: TestDataHan
                             }
 
                             <Stack verticalAlign='center' style={{ height: '100%' }}>
-                                <TooltipHost content='Filter phases'>
-                                    <PhaseFilter phases={sessionDetails.phases}
+                                <TooltipHost content='Filter by outcome'>
+                                    <CheckListOption
+                                        options={phaseFiltersContext.OutcomeOptions}
+                                        onChange={() => {
+                                            updateFilter(phaseFiltersContext.filterTestPhases(testId, sessionDetails.phases));
+                                        }}
+                                    />
+                                </TooltipHost>
+                            </Stack>
+
+                            <Stack verticalAlign='center' style={{ height: '100%' }}>
+                                <TooltipHost content='Filter phases by name'>
+                                    <PhaseFilter
+                                        phases={sessionDetails.phases}
+                                        selected={phaseFiltersContext.getFilteredNames(testId)}
                                         onChange={(items => {
-                                            const lowerCaseItems = items.map(i => i.toLowerCase());
-                                            const phases = sessionDetails.phases.filter(phase => phase.isMatch(lowerCaseItems))
-                                            updateFilter(phases);
+                                            phaseFiltersContext.setFilteredNames(testId, items);
+                                            updateFilter(phaseFiltersContext.filterTestPhases(testId, sessionDetails.phases));
                                         })}/>
                                 </TooltipHost>
                             </Stack>
diff --git a/Engine/Source/Programs/Horde/HordeDashboard/plugins/test-automation/testData.ts b/Engine/Source/Programs/Horde/HordeDashboard/plugins/test-automation/testData.ts
--- a/Engine/Source/Programs/Horde/HordeDashboard/plugins/test-automation/testData.ts
+++ b/Engine/Source/Programs/Horde/HordeDashboard/plugins/test-automation/testData.ts
@@ -267,11 +267,12 @@ export abstract class TestPhaseStatus {
         this.session = phaseSessions?.find(s => s.sessionId === this.testSession!.id);
     }
 
-    isMatch(items: string[]) {
-        if (items.length === 0) return true;
+    isMatch(names: string[], outcomes?: Set<TestPhaseOutcome>) {
+        if (!!outcomes && !outcomes.has(this.outcome)) return false;
+        if (names.length === 0) return true;
 
         const lowerName = this.name.toLowerCase();
-        return items.some(i => i.charAt(0) === '@'? lowerName === i.substring(1) : lowerName.includes(i));
+        return names.some(i => i.charAt(0) === '@'? lowerName === i.substring(1) : lowerName.includes(i));
     }
 
     get name() {
