TestDataV2 - Filter out stream without results

#rnx

diff --git a/Engine/Source/Programs/Horde/HordeDashboard/plugins/test-automation/testAutomationCommon.tsx b/Engine/Source/Programs/Horde/HordeDashboard/plugins/test-automation/testAutomationCommon.tsx
--- a/Engine/Source/Programs/Horde/HordeDashboard/plugins/test-automation/testAutomationCommon.tsx
+++ b/Engine/Source/Programs/Horde/HordeDashboard/plugins/test-automation/testAutomationCommon.tsx
@@ -359,7 +359,6 @@ export const getStreamOptions = (streams: string[], onClick: (stream: string) =>
 
                 categoryItems.get(category.name)?.push({
                     key: `streams_${projectStore.streamById(a)?.fullname}`,
-                    category: category?.name,
                     text: projectStore.streamById(a)?.name ?? a,
                     onClick: () => onClick(a),
                     disabled: a === selected
@@ -370,25 +369,29 @@ export const getStreamOptions = (streams: string[], onClick: (stream: string) =>
 
     const options: IContextualMenuItem[] = projectItems.entries().map(
         ([name, categories]) => {
+            // if more than one catergory and more than 8 streams overall, do a full break down of the catergories
+            const hasCategoriesAndManyItems = categories.size > 1 && categories.values().map(c => c.length).reduce((acc, l) => acc + l, 0) > 8;
             return {
                 key: `project_${name}`,
                 text: name,
                 subMenuProps: {
                     shouldFocusOnMount: true,
                     subMenuHoverDelay: 0,
-                    items: categories.entries().map(
-                        ([cat, items]) => {
-                            return {
-                                key: `category_${cat}`,
-                                text: cat,
-                                subMenuProps: {
-                                    shouldFocusOnMount: true,
-                                    subMenuHoverDelay: 0,
-                                    items: items
-                                }
-                            }
-                       }         
-                    ).toArray()
+                    items: (
+                        hasCategoriesAndManyItems ?
+                            categories.entries().map(
+                                ([cat, items]) => {
+                                    return {
+                                        key: `category_${cat}`,
+                                        text: cat,
+                                        subMenuProps: {
+                                            shouldFocusOnMount: true,
+                                            subMenuHoverDelay: 0,
+                                            items: items
+                                        }
+                                    }
+                            }).toArray()
+                            : categories.values().toArray().flat())
                 }
             }
     }).toArray();
diff --git a/Engine/Source/Programs/Horde/HordeDashboard/plugins/test-automation/testData.ts b/Engine/Source/Programs/Horde/HordeDashboard/plugins/test-automation/testData.ts
--- a/Engine/Source/Programs/Horde/HordeDashboard/plugins/test-automation/testData.ts
+++ b/Engine/Source/Programs/Horde/HordeDashboard/plugins/test-automation/testData.ts
@@ -1350,6 +1350,12 @@ export class TestDataHandler {
             await fetcherVersion.fetchStreamsInfo(testStreams);
         }
 
+        // filter out empty streams
+        for(const stream of this.testStreams.entries().filter(([, ref]) => ref.tests.length === 0).map(([key,]) => key))
+        {
+            this.testStreams.delete(stream);
+        }
+
         if (!this.loaded) {
             this.loaded = true;
             this.query(this.filterState.stream);
