Horde: Add shortened template alias to api/v1/jobs/streams/{streamId} endpoint

diff --git a/Engine/Source/Programs/Horde/Plugins/Build/HordeServer.Build/Jobs/JobsController.cs b/Engine/Source/Programs/Horde/Plugins/Build/HordeServer.Build/Jobs/JobsController.cs
--- a/Engine/Source/Programs/Horde/Plugins/Build/HordeServer.Build/Jobs/JobsController.cs
+++ b/Engine/Source/Programs/Horde/Plugins/Build/HordeServer.Build/Jobs/JobsController.cs
@@ -1238,6 +1238,7 @@ namespace HordeServer.Jobs
 		/// </summary>
 		/// <param name="streamId">The stream to search for</param>
 		/// <param name="templates">List of templates to find</param>
+		/// <param name="templateIds">Alternative to templates query param, used to shorten URLs with many templates</param>
 		/// <param name="preflightStartedByUserId">User id for which to include preflight jobs</param>
 		/// <param name="maxCreateTime">Maximum creation time</param>
 		/// <param name="modifiedAfter">If specified, only jobs updated after the given time will be returned</param>
@@ -1252,7 +1253,8 @@ namespace HordeServer.Jobs
 		[ProducesResponseType(typeof(List<GetJobResponse>), 200)]
 		public async Task<ActionResult<List<object>>> FindJobsByStreamWithTemplatesAsync(
 			string streamId,
-			[FromQuery(Name = "template")] string[] templates,
+			[FromQuery(Name = "template")] string[]? templates = null,
+			[FromQuery(Name = "t")] string[]? templateIds = null,
 			[FromQuery] string? preflightStartedByUserId = null,
 			[FromQuery] DateTimeOffset? maxCreateTime = null,
 			[FromQuery] DateTimeOffset? modifiedAfter = null,
@@ -1263,7 +1265,27 @@ namespace HordeServer.Jobs
 			CancellationToken cancellationToken = default)
 		{
 			StreamId streamIdValue = new StreamId(streamId);
-			TemplateId[] templateRefIds = templates.Select(x => new TemplateId(x)).ToArray();
+
+			IEnumerable<TemplateId>? allTemplates = null;
+			if (templates != null && templates.Length > 0 && templateIds != null && templateIds.Length > 0)
+			{
+				allTemplates = templates.Select(x => new TemplateId(x)).Union(templateIds.Select(x => new TemplateId(x)));
+			}
+			else if (templates != null && templates.Length > 0)
+			{
+				allTemplates = templates.Select(x => new TemplateId(x));
+			}
+			else if (templateIds != null && templateIds.Length > 0)
+			{
+				allTemplates = templateIds.Select(x => new TemplateId(x));
+			}
+			else
+			{
+				return BadRequest("Templates must be included in query, use template or t parameters.");
+			}
+
+			TemplateId[] templateRefIds = allTemplates.ToArray();
+
 			UserId? preflightStartedByUserIdValue = preflightStartedByUserId != null ? UserId.Parse(preflightStartedByUserId) : null;
 			count = Math.Min(1000, count);
 
