Clean up code + cleanup url after closing artifact modal

#rnx

diff --git a/Engine/Source/Programs/Horde/HordeDashboard/plugins/test-automation/testAutomationCommon.tsx b/Engine/Source/Programs/Horde/HordeDashboard/plugins/test-automation/testAutomationCommon.tsx
--- a/Engine/Source/Programs/Horde/HordeDashboard/plugins/test-automation/testAutomationCommon.tsx
+++ b/Engine/Source/Programs/Horde/HordeDashboard/plugins/test-automation/testAutomationCommon.tsx
@@ -444,7 +444,7 @@ class PhaseFiltersContext {
         this._names = new Map();
     }
 
-    get OutcomeOptions() {
+    get outcomeOptions() {
         return this._outcomeOptions
     }
 
diff --git a/Engine/Source/Programs/Horde/HordeDashboard/plugins/test-automation/testAutomationPhaseView.tsx b/Engine/Source/Programs/Horde/HordeDashboard/plugins/test-automation/testAutomationPhaseView.tsx
--- a/Engine/Source/Programs/Horde/HordeDashboard/plugins/test-automation/testAutomationPhaseView.tsx
+++ b/Engine/Source/Programs/Horde/HordeDashboard/plugins/test-automation/testAutomationPhaseView.tsx
@@ -412,7 +412,12 @@ export const TestPhasesView: React.FC<{ testDataId: string, handler: TestDataHan
 
     return <Stack style={{width: '100%', height: '100%'}}>
                 {!!jobArtifacts && jobContext.current?.job?.id && jobContext.current?.step?.id &&
-                    <JobArtifactsModal jobId={jobContext.current.job.id} stepId={jobContext.current?.step?.id} contextType={jobArtifacts.type} artifacts={[jobArtifacts]} onClose={() => setJobArtifacts(undefined)}/>
+                    <JobArtifactsModal
+                        jobId={jobContext.current.job.id} stepId={jobContext.current?.step?.id} contextType={jobArtifacts.type} artifacts={[jobArtifacts]}
+                        onClose={() => {
+                            setJobArtifacts(undefined);
+                            handler.removeSearchParam('artifactPath');
+                        }}/>
                 }
                 {!sessionDetails && handler.testDataQueryLoading &&
                         <Stack horizontalAlign='center' style={{ padding: 12 }} tokens={{ childrenGap: 10 }}>
@@ -517,7 +522,7 @@ export const TestPhasesView: React.FC<{ testDataId: string, handler: TestDataHan
                             <Stack verticalAlign='center' style={{ height: '100%' }}>
                                 <TooltipHost content='Filter by outcome'>
                                     <CheckListOption
-                                        options={phaseFiltersContext.OutcomeOptions}
+                                        options={phaseFiltersContext.outcomeOptions}
                                         onChange={() => {
                                             updateFilter(phaseFiltersContext.filterTestPhases(testId, sessionDetails.phases));
                                         }}
