TestDataV2 - Fix phase navigation buttons justification

#rnx

diff --git a/Engine/Source/Programs/Horde/HordeDashboard/plugins/test-automation/testAutomationPhaseView.tsx b/Engine/Source/Programs/Horde/HordeDashboard/plugins/test-automation/testAutomationPhaseView.tsx
--- a/Engine/Source/Programs/Horde/HordeDashboard/plugins/test-automation/testAutomationPhaseView.tsx
+++ b/Engine/Source/Programs/Horde/HordeDashboard/plugins/test-automation/testAutomationPhaseView.tsx
@@ -114,13 +114,22 @@ const TestPhaseHeader: React.FC<{ phase: TestPhaseStatus, expanded: boolean, tog
 
     const color = getPhaseSessionStatusColor(phase);
     const eventStyles = getEventStyling();
-    return <Stack tokens={{childrenGap: 2}} disableShrink style={{cursor: "pointer" }} onClick={onClick}>
-            <Stack horizontal verticalAlign="center">
-                <Stack style={{ paddingTop: 1, paddingRight: 4}}>
-                    <FontIcon style={{ fontSize: 13 }} iconName={expanded ? "ChevronDown" : "ChevronRight"} />
+    return <Stack tokens={{childrenGap: 2}} disableShrink style={{cursor: "pointer" }} onClick={onClick} horizontal verticalAlign="center">
+                <Stack>
+                    <Stack horizontal verticalAlign="center">
+                        <Stack style={{ paddingTop: 1, paddingRight: 4}}>
+                            <FontIcon style={{ fontSize: 13 }} iconName={expanded ? "ChevronDown" : "ChevronRight"} />
+                        </Stack>
+                        <Text style={{fontSize: 13, fontWeight: 600, whiteSpace: 'nowrap', overflow: 'hidden', textOverflow: 'ellipsis', cursor: "text"}} onClick={(ev) => ev.stopPropagation()}>{phase.name}</Text>
+                    </Stack>
+                    <Stack horizontal verticalAlign="center">
+                        <Stack style={{ paddingLeft: 5, paddingRight: 4 }}>
+                            <FontIcon style={{ fontSize: 12, color: color }} iconName="Square" />
+                        </Stack>
+                        <Text style={{fontSize: 11}}>{phase.outcome}</Text>
+                    </Stack>
                 </Stack>
-                <Text style={{fontSize: 13, fontWeight: 600, whiteSpace: 'nowrap', overflow: 'hidden', textOverflow: 'ellipsis', cursor: "text"}} onClick={(ev) => ev.stopPropagation()}>{phase.name}</Text>
-                <Stack horizontalAlign="end" grow style={{paddingTop: 4, paddingRight: 12}}>
+                <Stack horizontalAlign="end" grow style={{paddingRight: 12}}>
                     {!!expanded &&
                         <Stack horizontal>
                             <DefaultButton
@@ -146,13 +155,6 @@ const TestPhaseHeader: React.FC<{ phase: TestPhaseStatus, expanded: boolean, tog
                     }
                 </Stack>
             </Stack>
-            <Stack horizontal verticalAlign="center">
-                <Stack style={{ paddingLeft: 5, paddingRight: 4 }}>
-                    <FontIcon style={{ fontSize: 12, color: color }} iconName="Square" />
-                </Stack>
-                <Text style={{fontSize: 11}}>{phase.outcome}</Text>
-            </Stack>
-        </Stack>
 });
 
 const TestPhaseSection: React.FC<{ phase: TestPhaseStatus, handler: TestDataHandler, onCollapse: (phaseKey: string) => void, onNext: () => void, onBack: () => void }> = ({ phase, handler, onCollapse, onNext, onBack }) => {
@@ -163,7 +165,7 @@ const TestPhaseSection: React.FC<{ phase: TestPhaseStatus, handler: TestDataHand
     }, []);
 
     const borderColor = dashboard.darktheme ? "#4D4C4B" : "#6D6C6B";
-    return <Stack style={{ borderTopWidth: 1, borderTopColor: borderColor, borderTopStyle: 'solid', height: '100%'}}>
+    return <Stack style={{ borderTopWidth: 1, borderTopColor: borderColor, borderTopStyle: 'solid', height: '100%', paddingTop: 3}}>
                 <TestPhaseHeader phase={phase} expanded={true} toggleView={toggleView} onNext={onNext} onBack={onBack} />
                 <TestPhaseExpandPane phase={phase} handler={handler}/>
             </Stack>
