Horde: Use Win32 API for calculating memory telemetry on Windows

Previous use of WMI included page file which was incorrect. This is also faster.

diff --git a/Engine/Source/Programs/Horde/HordeAgent/Services/TelemetryService.cs b/Engine/Source/Programs/Horde/HordeAgent/Services/TelemetryService.cs
--- a/Engine/Source/Programs/Horde/HordeAgent/Services/TelemetryService.cs
+++ b/Engine/Source/Programs/Horde/HordeAgent/Services/TelemetryService.cs
@@ -141,16 +141,12 @@ public sealed class DefaultSystemMetrics : ISystemMetrics
 public sealed class WindowsSystemMetrics : ISystemMetrics, IDisposable
 {
 	private const string ProcessorInfo = "Processor Information"; // Prefer this over "Processor" as it's more modern
-	private const string Memory = "Memory";
 	private const string Total = "_Total";
 
 	private readonly PerformanceCounter _procIdleTime = new(ProcessorInfo, "% Idle Time", Total);
 	private readonly PerformanceCounter _procUserTime = new(ProcessorInfo, "% User Time", Total);
 	private readonly PerformanceCounter _procPrivilegedTime = new(ProcessorInfo, "% Privileged Time", Total);
 
-	private readonly uint _totalPhysicalMemory = GetPhysicalMemory();
-	private readonly PerformanceCounter _memAvailableBytes = new(Memory, "Available Bytes");
-
 	private readonly DirectoryReference _workingDir;
 	private readonly Stopwatch _driveInfoUpdateTimer = Stopwatch.StartNew();
 	private DriveInfo? _driveInfo;
@@ -169,7 +165,6 @@ public sealed class WindowsSystemMetrics : ISystemMetrics, IDisposable
 	/// <inheritdoc/>
 	public void Dispose()
 	{
-		_memAvailableBytes.Dispose();
 		_procIdleTime.Dispose();
 		_procUserTime.Dispose();
 		_procPrivilegedTime.Dispose();
@@ -189,14 +184,14 @@ public sealed class WindowsSystemMetrics : ISystemMetrics, IDisposable
 	/// <inheritdoc />
 	public MemoryMetrics GetMemory()
 	{
-		uint available = (uint)(_memAvailableBytes.NextValue() / 1024);
-		uint used = _totalPhysicalMemory - available;
+		(ulong totalBytes, ulong availableBytes) = WindowsCapabilities.GetPhysicalMemory();
+		ulong usedBytes = totalBytes - availableBytes;
 		return new MemoryMetrics
 		{
-			Total = _totalPhysicalMemory,
-			Available = available,
-			Used = used,
-			UsedPercentage = used / (float)_totalPhysicalMemory,
+			Total = (uint)(totalBytes / 1024),
+			Available = (uint)(availableBytes / 1024),
+			Used = (uint)(usedBytes / 1024),
+			UsedPercentage = usedBytes / (float)totalBytes,
 		};
 	}
 
@@ -227,11 +222,6 @@ public sealed class WindowsSystemMetrics : ISystemMetrics, IDisposable
 		}
 		return new DriveInfo(rootDir);
 	}
-
-	private static uint GetPhysicalMemory()
-	{
-		return (uint)(WindowsCapabilities.GetPhysicalMemory() / 1024);
-	}
 }
 
 #pragma warning restore CA1416
diff --git a/Engine/Source/Programs/Horde/HordeAgent/Services/WindowsCapabilities.cs b/Engine/Source/Programs/Horde/HordeAgent/Services/WindowsCapabilities.cs
--- a/Engine/Source/Programs/Horde/HordeAgent/Services/WindowsCapabilities.cs
+++ b/Engine/Source/Programs/Horde/HordeAgent/Services/WindowsCapabilities.cs
@@ -71,24 +71,24 @@ public class WindowsCapabilities(IOptionsMonitor<AgentSettings> settings, ILogge
 	private static void SetResource(RpcAgentCapabilities caps, string key, int value) => caps.Resources.Add(key, value);
 	
 	/// <summary>
-	/// Gets the total amount of physical memory installed on the system.
+	/// Gets details of current system's physical memory
 	/// </summary>
-	/// <returns>The total physical memory in bytes. Returns 0 if memory status could not be retrieved</returns>
-	public static int GetPhysicalMemory()
+	/// <returns>The total and available physical memory in bytes. Returns 0 if memory status could not be retrieved</returns>
+	public static (ulong total, ulong available) GetPhysicalMemory()
 	{
 		try
 		{
 			MEMORYSTATUSEX memStatus = new();
 			if (GlobalMemoryStatusEx(memStatus))
 			{
-				return (int)memStatus.ullTotalPhys;
+				return (memStatus.ullTotalPhys, memStatus.ullAvailPhys);
 			}
 		}
 		catch
 		{
 			// Ignore any error
 		}
-		return 0;
+		return (0, 0);
 	}
 	
 	private static int GetNumLogicalCores()
