Horde: Perforce passwords and tickets in the BuildConfig can be resolved to the value of a ISecret during post load replacing the need for them to be plain text values

diff --git a/Engine/Source/Programs/Horde/HordeServer.Shared/Plugins/IPluginConfig.cs b/Engine/Source/Programs/Horde/HordeServer.Shared/Plugins/IPluginConfig.cs
--- a/Engine/Source/Programs/Horde/HordeServer.Shared/Plugins/IPluginConfig.cs
+++ b/Engine/Source/Programs/Horde/HordeServer.Shared/Plugins/IPluginConfig.cs
@@ -19,7 +19,7 @@ namespace HordeServer.Plugins
 	/// <summary>
 	/// Options passed to <see cref="IPluginConfig.PostLoadAsync(PluginConfigOptions)"/>
 	/// </summary>
-	public record class PluginConfigOptions(ConfigVersion Version, IEnumerable<IPluginConfig> Plugins, AclConfig ParentAcl);
+	public record class PluginConfigOptions(ConfigVersion Version, IEnumerable<IPluginConfig> Plugins, AclConfig ParentAcl, IServiceProvider? ServiceProvider = null, CancellationToken CancellationToken = default);
 
 	/// <summary>
 	/// Empty implementation of <see cref="IPluginConfig"/>
diff --git a/Engine/Source/Programs/Horde/HordeServer.Tests.Shared/ServerTestSetup.cs b/Engine/Source/Programs/Horde/HordeServer.Tests.Shared/ServerTestSetup.cs
--- a/Engine/Source/Programs/Horde/HordeServer.Tests.Shared/ServerTestSetup.cs
+++ b/Engine/Source/Programs/Horde/HordeServer.Tests.Shared/ServerTestSetup.cs
@@ -74,15 +74,17 @@ namespace HordeServer.Tests
 
 		protected async Task SetConfigAsync(GlobalConfig globalConfig)
 		{
-			await globalConfig.PostLoadAsync(ServerSettings, _pluginCollection.LoadedPlugins, DefaultAclModifiers);
+			await ConfigService.WaitForInitialConfigAsync();
+			await globalConfig.PostLoadAsync(ServerSettings, _pluginCollection.LoadedPlugins, DefaultAclModifiers, ServiceProvider);
 			ConfigService.OverrideConfig(globalConfig);
 		}
 
 		protected async Task UpdateConfigAsync(Action<GlobalConfig> action)
 		{
+			await ConfigService.WaitForInitialConfigAsync();
 			GlobalConfig globalConfig = GlobalConfig.CurrentValue;
 			action(globalConfig);
-			await globalConfig.PostLoadAsync(ServerSettings, _pluginCollection.LoadedPlugins, DefaultAclModifiers);
+			await globalConfig.PostLoadAsync(ServerSettings, _pluginCollection.LoadedPlugins, DefaultAclModifiers, ServiceProvider);
 			ConfigService.OverrideConfig(globalConfig);
 		}
 
diff --git a/Engine/Source/Programs/Horde/HordeServer/Configuration/ConfigService.cs b/Engine/Source/Programs/Horde/HordeServer/Configuration/ConfigService.cs
--- a/Engine/Source/Programs/Horde/HordeServer/Configuration/ConfigService.cs
+++ b/Engine/Source/Programs/Horde/HordeServer/Configuration/ConfigService.cs
@@ -132,6 +132,8 @@ namespace HordeServer.Configuration
 		Task<ConfigState> _stateTask;
 		readonly ChangeToken _changeToken = new ChangeToken();
 
+		readonly IServiceProvider _serviceProvider;
+
 		/// <inheritdoc/>
 		string IOptionsChangeTokenSource<GlobalConfig>.Name => String.Empty;
 
@@ -143,7 +145,7 @@ namespace HordeServer.Configuration
 		/// <summary>
 		/// Constructor
 		/// </summary>
-		public ConfigService(IRedisService redisService, IOptions<ServerSettings> serverSettings, IEnumerable<IConfigSource> sources, IPluginCollection pluginCollection, IEnumerable<IDefaultAclModifier> aclModifiers, IClock clock, IHealthMonitor<ConfigService> health, ILogger<ConfigService> logger)
+		public ConfigService(IRedisService redisService, IOptions<ServerSettings> serverSettings, IEnumerable<IConfigSource> sources, IPluginCollection pluginCollection, IEnumerable<IDefaultAclModifier> aclModifiers, IClock clock, IHealthMonitor<ConfigService> health, ILogger<ConfigService> logger, IServiceProvider serviceProvider)
 		{
 			_redisService = redisService;
 			_serverSettings = serverSettings.Value;
@@ -153,6 +155,7 @@ namespace HordeServer.Configuration
 			_health = health;
 			_health.SetName("GlobalConfig");
 			_logger = logger;
+			_serviceProvider = serviceProvider;
 
 			_jsonOptions = new JsonSerializerOptions();
 			JsonUtils.ConfigureJsonSerializer(_jsonOptions);
@@ -282,7 +285,7 @@ namespace HordeServer.Configuration
 			try
 			{
 				GlobalConfig globalConfig = await ReadGlobalConfigAsync(context, cancellationToken);
-				await globalConfig.PostLoadAsync(_serverSettings, _pluginCollection.LoadedPlugins, _defaultAclModifiers);
+				await globalConfig.PostLoadAsync(_serverSettings, _pluginCollection.LoadedPlugins, _defaultAclModifiers, _serviceProvider, cancellationToken);
 
 				foreach (OverrideConfigFile file in overrideFiles.Values)
 				{
@@ -608,7 +611,7 @@ namespace HordeServer.Configuration
 				}
 
 				// Execute a PostLoad before returning so we can validate that everything is valid
-				await globalConfig.PostLoadAsync(_serverSettings, _pluginCollection.LoadedPlugins, _defaultAclModifiers);
+				await globalConfig.PostLoadAsync(_serverSettings, _pluginCollection.LoadedPlugins, _defaultAclModifiers, _serviceProvider, cancellationToken);
 
 				return snapshot;
 			}
@@ -857,7 +860,7 @@ namespace HordeServer.Configuration
 			globalConfig.Revision = IoHash.Compute(data.Span).ToString();
 
 			// Run the postload callbacks on all the config objects
-			await globalConfig.PostLoadAsync(_serverSettings, _pluginCollection.LoadedPlugins, _defaultAclModifiers);
+			await globalConfig.PostLoadAsync(_serverSettings, _pluginCollection.LoadedPlugins, _defaultAclModifiers, _serviceProvider, CancellationToken.None);
 			return globalConfig;
 		}
 	}
diff --git a/Engine/Source/Programs/Horde/HordeServer/Server/GlobalConfig.cs b/Engine/Source/Programs/Horde/HordeServer/Server/GlobalConfig.cs
--- a/Engine/Source/Programs/Horde/HordeServer/Server/GlobalConfig.cs
+++ b/Engine/Source/Programs/Horde/HordeServer/Server/GlobalConfig.cs
@@ -6,6 +6,7 @@ using System.Diagnostics.CodeAnalysis;
 using System.Linq;
 using System.Security.Claims;
 using System.Text.Json.Serialization;
+using System.Threading;
 using System.Threading.Tasks;
 using EpicGames.Core;
 using EpicGames.Horde.Acls;
@@ -105,7 +106,7 @@ namespace HordeServer.Server
 		/// <summary>
 		/// Called after the config file has been read
 		/// </summary>
-		public async Task PostLoadAsync(ServerSettings serverSettings, IReadOnlyList<ILoadedPlugin> loadedPlugins, IEnumerable<IDefaultAclModifier> defaultAclModifiers)
+		public async Task PostLoadAsync(ServerSettings serverSettings, IReadOnlyList<ILoadedPlugin> loadedPlugins, IEnumerable<IDefaultAclModifier> defaultAclModifiers, IServiceProvider? serviceProvider = null, CancellationToken cancellationToken = default)
 		{
 			ServerSettings = serverSettings;
 			AclAction[] aclActions = AclConfig.GetActions(
@@ -155,7 +156,7 @@ namespace HordeServer.Server
 				pluginConfigs = sortedPluginConfigs;
 			}
 
-			PluginConfigOptions pluginConfigOptions = new PluginConfigOptions(VersionEnum, Plugins.Values, Acl);
+			PluginConfigOptions pluginConfigOptions = new PluginConfigOptions(VersionEnum, Plugins.Values, Acl, serviceProvider, cancellationToken);
 			foreach (IPluginConfig pluginConfig in pluginConfigs)
 			{
 				await pluginConfig.PostLoadAsync(pluginConfigOptions);
diff --git a/Engine/Source/Programs/Horde/Plugins/Build/HordeServer.Build.Tests/Perforce/PerforceSecretsTests.cs b/Engine/Source/Programs/Horde/Plugins/Build/HordeServer.Build.Tests/Perforce/PerforceSecretsTests.cs
new file mode 100644
--- /dev/null
+++ b/Engine/Source/Programs/Horde/Plugins/Build/HordeServer.Build.Tests/Perforce/PerforceSecretsTests.cs
@@ -0,0 +1,99 @@
+// Copyright Epic Games, Inc. All Rights Reserved.
+
+using EpicGames.Horde.Secrets;
+using HordeServer.Secrets;
+using HordeServer.Server;
+
+namespace HordeServer.Tests.Perforce;
+
+[TestClass]
+public class PerforceSecretsTests : BuildTestSetup
+{
+	public PerforceSecretsTests()
+	{
+		AddPlugin<SecretsPlugin>();
+	}
+
+	[TestInitialize]
+	public async Task SetupAsync()
+	{
+		GlobalConfig globalConfig = new();
+		globalConfig.Plugins.AddBuildConfig(new BuildConfig { ResolvePerforceCredentialsWithSecrets = true });
+		globalConfig.Plugins.AddSecretsConfig(new SecretsConfig());
+		await SetConfigAsync(globalConfig);
+	}
+
+	private async Task AddCredentialsAsync(PerforceCredentials credentials, SecretConfig? secret)
+	{
+		await UpdateConfigAsync(globalConfig =>
+		{
+			if (secret != null)
+			{
+				globalConfig.Plugins.GetSecretsConfig().Secrets.Add(secret);
+			}
+			globalConfig.Plugins.GetBuildConfig().PerforceClusters.Add(new PerforceCluster { Credentials = [credentials], Name = "Test" });
+		});
+	}
+
+	[TestMethod]
+	public async Task ResolvePerforcePasswordWithSecretAsync()
+	{
+		PerforceCredentials credentials = new() { Password = "horde:secret:p4secret.Password" };
+		SecretConfig secret = new()
+		{
+			Id = new SecretId("p4secret"),
+			Data = new Dictionary<string, string> { { "Password", "Hello, World!" } }
+		};
+		await AddCredentialsAsync(credentials, secret);
+		Assert.AreEqual(credentials.Password, "Hello, World!");
+		Assert.AreEqual(credentials.Ticket, String.Empty);
+	}
+
+	[TestMethod]
+	public async Task ResolvePerforceTicketWithSecretAsync()
+	{
+		PerforceCredentials credentials = new() { Ticket = "horde:secret:p4secret.Ticket" };
+		SecretConfig secret = new()
+		{
+			Id = new SecretId("p4secret"),
+			Data = new Dictionary<string, string> { { "Ticket", "Hello, World!" } }
+		};
+		await AddCredentialsAsync(credentials, secret);
+		Assert.AreEqual(credentials.Ticket, "Hello, World!");
+		Assert.AreEqual(credentials.Password, String.Empty);
+	}
+
+	[TestMethod]
+	public async Task ResolvePerforceCredentialsWithoutSecretAsync()
+	{
+		PerforceCredentials credentials = new() { Password = "password", Ticket = "ticket" };
+		await AddCredentialsAsync(credentials, null);
+		Assert.AreEqual(credentials.Password, "password");
+		Assert.AreEqual(credentials.Ticket, "ticket");
+	}
+
+	[TestMethod]
+	[DataRow("horde:secret:notfound.property", DisplayName = "Resolve Perforce Secret, secret not found")]
+	[DataRow("horde:secret:p4secret.notfound", DisplayName = "Resolve Perforce Secret, property not found")]
+	public async Task ResolvePerforceCredentialsWithSecretNotFoundAsync(string value)
+	{
+		PerforceCredentials credentials = new() { Password = value, Ticket = value };
+		SecretConfig secret = new() { Id = new SecretId("p4secret") };
+		await AddCredentialsAsync(credentials, secret);
+		Assert.AreEqual(credentials.Password, value);
+		Assert.AreEqual(credentials.Ticket, value);
+	}
+
+	[TestMethod]
+	public async Task ResolvePerforceCredentialsWithSecretMisconfigurationAsync()
+	{
+		PerforceCredentials credentials = new() { Password = "horde:secret:p4secret.property" };
+		SecretConfig secret = new()
+		{
+			Id = new SecretId("p4secret"),
+			Sources = [ new ExternalSecretConfig() { Provider = "SecretProviderKeyNotFound" } ]
+		};
+		await AddCredentialsAsync(credentials, secret);
+		Assert.AreEqual(credentials.Password, "horde:secret:p4secret.property");
+	}
+}
diff --git a/Engine/Source/Programs/Horde/Plugins/Build/HordeServer.Build/BuildConfig.cs b/Engine/Source/Programs/Horde/Plugins/Build/HordeServer.Build/BuildConfig.cs
--- a/Engine/Source/Programs/Horde/Plugins/Build/HordeServer.Build/BuildConfig.cs
+++ b/Engine/Source/Programs/Horde/Plugins/Build/HordeServer.Build/BuildConfig.cs
@@ -12,6 +12,7 @@ using EpicGames.Horde.Common;
 using EpicGames.Horde.Jobs;
 using EpicGames.Horde.Jobs.Templates;
 using EpicGames.Horde.Projects;
+using EpicGames.Horde.Secrets;
 using EpicGames.Horde.Streams;
 using EpicGames.Perforce;
 using HordeServer.Acls;
@@ -25,6 +26,7 @@ using HordeServer.Plugins;
 using HordeServer.Projects;
 using HordeServer.Streams;
 using HordeServer.Utilities;
+using Microsoft.Extensions.DependencyInjection;
 
 #pragma warning disable CA2227 // Change x to be read-only by removing the property setter
 
@@ -46,6 +48,11 @@ namespace HordeServer
 		/// </summary>
 		public List<PerforceCluster> PerforceClusters { get; set; } = new List<PerforceCluster>();
 
+		/// <summary>
+		/// Whether to enable resolving Perforce passwords and tickets against secrets
+		/// </summary>
+		public bool ResolvePerforceCredentialsWithSecrets { get; set; }
+
 		/// <summary>
 		/// Device configuration
 		/// </summary>
@@ -92,11 +99,13 @@ namespace HordeServer
 		private readonly Dictionary<StreamId, StreamConfig> _streamLookup = new Dictionary<StreamId, StreamConfig>();
 
 		/// <inheritdoc/>
-		public Task PostLoadAsync(PluginConfigOptions configOptions)
+		public async Task PostLoadAsync(PluginConfigOptions configOptions)
 		{
 			Acl = configOptions.ParentAcl;
 			Streams = Projects.SelectMany(x => x.Streams).ToList();
 
+			await UpdatePerforceCredentialsAsync(configOptions.Plugins, configOptions.ServiceProvider, configOptions.CancellationToken);
+
 			_projectLookup.Clear();
 			_streamLookup.Clear();
 			foreach (ProjectConfig project in Projects)
@@ -116,8 +125,48 @@ namespace HordeServer
 			}
 
 			UpdateWorkspacesForPools(configOptions.Plugins);
+		}
+
+		private static async Task<string?> ResolveSecretAsync(string? value, ISecretCollection secretCollection, SecretsConfig secretsConfig, CancellationToken cancellationToken)
+		{
+			if (!String.IsNullOrEmpty(value))
+			{
+				ISecretProperty? newValue = await secretCollection.ResolveAsync(value, secretsConfig, cancellationToken);
+				if (newValue != null)
+				{
+					return newValue.Value;
+				}
+			}
+			return value;
+		}
+
+		private async Task UpdatePerforceCredentialsAsync(IEnumerable<IPluginConfig> plugins, IServiceProvider? serviceProvider, CancellationToken cancellationToken)
+		{
+			if (!ResolvePerforceCredentialsWithSecrets)
+			{
+				return;
+			}
 
-			return Task.CompletedTask;
+			SecretsConfig? secretsConfig = plugins.OfType<SecretsConfig>().FirstOrDefault();
+			if (secretsConfig == null)
+			{
+				return;
+			}
+
+			ISecretCollection? secretCollection = serviceProvider?.GetService<ISecretCollection>();
+			if (secretCollection == null)
+			{
+				return;
+			}
+
+			foreach (PerforceCluster perforceCluster in PerforceClusters)
+			{
+				foreach (PerforceCredentials credentials in perforceCluster.Credentials)
+				{
+					credentials.Password = await ResolveSecretAsync(credentials.Password, secretCollection, secretsConfig, cancellationToken);
+					credentials.Ticket = await ResolveSecretAsync(credentials.Ticket, secretCollection, secretsConfig, cancellationToken);
+				}
+			}
 		}
 
 		void UpdateWorkspacesForPools(IEnumerable<IPluginConfig> plugins)
diff --git a/Engine/Source/Programs/Horde/Plugins/Build/HordeServer.Build/BuildPlugin.cs b/Engine/Source/Programs/Horde/Plugins/Build/HordeServer.Build/BuildPlugin.cs
--- a/Engine/Source/Programs/Horde/Plugins/Build/HordeServer.Build/BuildPlugin.cs
+++ b/Engine/Source/Programs/Horde/Plugins/Build/HordeServer.Build/BuildPlugin.cs
@@ -45,7 +45,7 @@ namespace HordeServer
 	/// <summary>
 	/// Entry point for the build plugin
 	/// </summary>
-	[Plugin("Build", GlobalConfigType = typeof(BuildConfig), ServerConfigType = typeof(BuildServerConfig))]
+	[Plugin("Build", GlobalConfigType = typeof(BuildConfig), ServerConfigType = typeof(BuildServerConfig), DependsOn = ["Secrets"])]
 	public class BuildPlugin : IPluginStartup
 	{
 		readonly IServerInfo _serverInfo;
diff --git a/Engine/Source/Programs/Horde/Plugins/Build/HordeServer.Build/HordeServer.Build.csproj b/Engine/Source/Programs/Horde/Plugins/Build/HordeServer.Build/HordeServer.Build.csproj
--- a/Engine/Source/Programs/Horde/Plugins/Build/HordeServer.Build/HordeServer.Build.csproj
+++ b/Engine/Source/Programs/Horde/Plugins/Build/HordeServer.Build/HordeServer.Build.csproj
@@ -39,6 +39,7 @@
     <ProjectReference Include="..\..\..\..\Shared\EpicGames.Slack\EpicGames.Slack.csproj" />
     <ProjectReference Include="..\..\..\HordeServer.Shared\HordeServer.Shared.csproj" />
     <ProjectReference Include="..\..\Compute\HordeServer.Compute\HordeServer.Compute.csproj" />
+    <ProjectReference Include="..\..\Secrets\HordeServer.Secrets\HordeServer.Secrets.csproj" />
     <ProjectReference Include="..\..\Tools\HordeServer.Tools\HordeServer.Tools.csproj" />
   </ItemGroup>
 
