Horde: Use ServiceScope to get the secrets collection

#rnx

diff --git a/Engine/Source/Programs/Horde/Plugins/Build/HordeServer.Build/BuildConfig.cs b/Engine/Source/Programs/Horde/Plugins/Build/HordeServer.Build/BuildConfig.cs
--- a/Engine/Source/Programs/Horde/Plugins/Build/HordeServer.Build/BuildConfig.cs
+++ b/Engine/Source/Programs/Horde/Plugins/Build/HordeServer.Build/BuildConfig.cs
@@ -153,7 +153,8 @@ namespace HordeServer
 				return;
 			}
 
-			ISecretCollection? secretCollection = serviceProvider?.GetService<ISecretCollection>();
+			using IServiceScope? scope = serviceProvider?.CreateScope();
+			ISecretCollection? secretCollection = scope?.ServiceProvider.GetService<ISecretCollection>();
 			if (secretCollection == null)
 			{
 				return;
