TestDataV2 - Cache stream options

#rnx

diff --git a/Engine/Source/Programs/Horde/HordeDashboard/plugins/test-automation/testAutomationSummary.tsx b/Engine/Source/Programs/Horde/HordeDashboard/plugins/test-automation/testAutomationSummary.tsx
--- a/Engine/Source/Programs/Horde/HordeDashboard/plugins/test-automation/testAutomationSummary.tsx
+++ b/Engine/Source/Programs/Horde/HordeDashboard/plugins/test-automation/testAutomationSummary.tsx
@@ -49,14 +49,15 @@ const TestHistoryGraphWidget: React.FC<{ test: TestNameRef, stream: string, hand
             </Stack>
 }
 
-const StreamSelector: React.FC<{ streams: string[], onClick: (stream: string) => void, selected?: string, disabled?: boolean }> = ({ streams, onClick, selected, disabled }) => {
-
-    const projectItems: Map<string, Map<string, IContextualMenuItem[]>> = new Map();
-
+const StreamSelector: React.FC<{ streams: string[], onClick: (stream: string) => void, selected?: string, disabled?: boolean }> = memo(({ streams, onClick, selected, disabled }) => {
     const menuProps: IContextualMenuProps = getStreamOptions(streams, onClick, selected);
-
-    return <DefaultButton style={{ height: 22, fontSize: 11, whiteSpace: 'nowrap', borderWidth: 0 }} menuProps={menuProps} text={selected? projectStore.streamById(selected)?.fullname ?? selected : "Compare stream"} disabled={disabled}/>
-}
+    return <DefaultButton
+                style={{ height: 22, fontSize: 11, whiteSpace: 'nowrap', borderWidth: 0 }}
+                menuProps={menuProps}
+                text={selected? projectStore.streamById(selected)?.fullname ?? selected : "Compare stream"}
+                disabled={disabled}
+            />
+});
 
 export const TestHistoryView: React.FC<{ test: TestNameRef, handler: TestDataHandler, onClick: SessionCallback  }> = observer(({ test, handler, onClick }) => {
 
@@ -66,6 +67,13 @@ export const TestHistoryView: React.FC<{ test: TestNameRef, handler: TestDataHan
 
     handler.subscribeToSubQueryLoading();
 
+    const onClickStream = useCallback((targetStream: string) => {
+        setStreams({...streams, toCompare: undefined});
+        handler.query(targetStream, true)
+            .catch((reason) => console.error(reason))
+            .finally(() => setStreams({...streams, toCompare: targetStream}));
+    }, [streams]);
+
     const availableStreams = handler.availableStreams.filter((s) => s !== streams.current);
     let anyFailed = false;
 
@@ -151,10 +159,7 @@ export const TestHistoryView: React.FC<{ test: TestNameRef, handler: TestDataHan
                             updateKeys={(keys) => setSelectedMeta(new Set(keys))} />
                     </Stack>
                     <Stack>
-                        <StreamSelector streams={availableStreams} selected={streams.toCompare} disabled={availableStreams.length === 0} onClick={(targetStream: string) => {
-                            setStreams({...streams, toCompare: undefined});
-                            handler.query(targetStream, true).catch((reason) => console.error(reason)).finally(() => setStreams({...streams, toCompare: targetStream}))
-                        }} />
+                        <StreamSelector streams={availableStreams} selected={streams.toCompare} disabled={availableStreams.length === 0} onClick={onClickStream} />
                     </Stack>
                 </Stack>
                 <Stack style={{ paddingLeft: 12 }} tokens={{ childrenGap: 4 }}>
diff --git a/Engine/Source/Programs/Horde/HordeDashboard/plugins/test-automation/testAutomationView.tsx b/Engine/Source/Programs/Horde/HordeDashboard/plugins/test-automation/testAutomationView.tsx
--- a/Engine/Source/Programs/Horde/HordeDashboard/plugins/test-automation/testAutomationView.tsx
+++ b/Engine/Source/Programs/Horde/HordeDashboard/plugins/test-automation/testAutomationView.tsx
@@ -4,7 +4,7 @@ import { TopNav } from "horde/components/TopNav";
 import { getHordeStyling } from "horde/styles/Styles";
 import { observer } from "mobx-react-lite";
 import { useNavigate, useSearchParams } from "react-router-dom";
-import React, { useEffect, useState, useCallback } from "react";
+import React, { useEffect, useState, useCallback, useMemo } from "react";
 import { projectStore } from 'horde/backend/ProjectStore';
 import { TestDataHandler, TestDataVersionRegistrar, TestSessionDetails } from "./testData";
 import { TestSummary } from "./testAutomationSummary";
@@ -26,7 +26,7 @@ const StreamChooser: React.FC = observer(() => {
 
     const stream = handler.filterState.stream;
 
-    const menuProps: IContextualMenuProps = getStreamOptions(handler.availableStreams, (stream) => handler.selectStream(stream), stream);
+    const menuProps: IContextualMenuProps = useMemo(() => getStreamOptions(handler.availableStreams, (stream) => handler.selectStream(stream), stream), [stream]);
     
     return <Stack style={{ paddingTop: 12, paddingBottom: 4 }} horizontal>
             <Label style={{ paddingTop: 5, paddingRight: 5 }}>Stream</Label>
