Horde Dashboard: Agent telemetry fixes and improvements for Agent and Log views

diff --git a/Engine/Source/Programs/Horde/HordeDashboard/src/components/agents/AgentTelemetrySparkline.tsx b/Engine/Source/Programs/Horde/HordeDashboard/src/components/agents/AgentTelemetrySparkline.tsx
--- a/Engine/Source/Programs/Horde/HordeDashboard/src/components/agents/AgentTelemetrySparkline.tsx
+++ b/Engine/Source/Programs/Horde/HordeDashboard/src/components/agents/AgentTelemetrySparkline.tsx
@@ -4,9 +4,11 @@ import { observer } from "mobx-react-lite";
 import { Sparklines, SparklinesLine, SparklinesReferenceLine } from "react-sparklines";
 import backend from "../../backend";
 import { GetAgentResponse, GetAgentTelemetrySampleResponse } from "../../backend/Api";
-import dashboard from "../../backend/Dashboard";
+import dashboard, { StatusColor } from "../../backend/Dashboard";
 import { PollBase } from "../../backend/PollBase";
 
+//  http://localhost:5173/log/6849950e9935f922c0cc2339?batch=e043&agentId=10-99-112-190
+
 const sparkWidth = 1380 / 2;
 const sparkHeight = 64;
 
@@ -185,22 +187,6 @@ export const AgentTelemetrySparkline: React.FC<{ handler: AgentTelemetryHandler
       return;
    }
 
-   const ramData = data.map(d => {
-      return d.usedRam / d.totalRam;
-   })
-
-   const cpuData = data.map(d => {
-      return (d.userCpu + d.systemCpu) / 100;
-   })
-
-   const diskData = data.map(d => {
-      if (typeof (d.totalDisk) === "number" && d.totalDisk) {
-         return 1.0 - (d.freeDisk / d.totalDisk)
-      }
-
-      return 0;
-   })
-
    let cpuText = "";
    let ramText = "";
    let diskText = "";
@@ -241,45 +227,70 @@ export const AgentTelemetrySparkline: React.FC<{ handler: AgentTelemetryHandler
       }
    })
 
+   const ramData = data.map(d => {
+      return 1.0 - (d.freeRam / d.totalRam);
+   })
+
+   const cpuData = data.map(d => {
+      return (d.userCpu + d.systemCpu) / 100;
+   })
+
+   const diskData = data.map(d => {
+      if (typeof (d.totalDisk) === "number" && d.totalDisk) {
+         return 1.0 - (d.freeDisk / d.totalDisk)
+      }
+
+      return 0;
+   })
+
+   const statusColors = dashboard.getStatusColors();
+   const opacity = 0.75;
+
    return <Stack>
       <Stack horizontal horizontalAlign="center" tokens={{ childrenGap: 24 }} style={{ paddingBottom: 8 }}>
          <Stack style={{ width: sparkWidth }}>
             <Stack horizontal>
-               <Label>CPU</Label>
+               <Label>Used CPU</Label>
                <Stack grow></Stack>
                {!!cpuText && <Label>{cpuText}</Label>}
             </Stack>
             <div id="telemetry_cpu">
-               <Sparklines width={sparkWidth} height={sparkHeight} data={cpuData} max={1} style={{ backgroundColor: dashboard.darktheme ? "#060709" : "#F3F2F1", padding: 8, border: "solid 1px #181A1B" }}>
-                  <SparklinesLine color={dashboard.darktheme ? "lightblue" : "#1E90FF"} />
-                  <SparklinesReferenceLine type="avg" />
+               <Sparklines width={sparkWidth} height={sparkHeight} data={cpuData} min={0} max={1} style={{ backgroundColor: dashboard.darktheme ? "#060709" : "#F3F2F1", padding: 8, border: "solid 1px #181A1B" }}>
+                  <SparklinesLine color={dashboard.darktheme ? "lightblue" : "#1E90FF"} style={{ strokeWidth: 2 }} />
+                  <SparklinesReferenceLine type="custom" value={0} style={{ stroke: statusColors.get(StatusColor.Failure) }} />
+                  <SparklinesReferenceLine type="custom" value={22} style={{ stroke: statusColors.get(StatusColor.Warnings), opacity: opacity }} />
+                  <SparklinesReferenceLine type="custom" value={44} style={{ stroke: statusColors.get(StatusColor.Success), opacity: opacity }} />
                </Sparklines>
             </div>
          </Stack>
          <Stack style={{ width: sparkWidth }}>
             <Stack horizontal>
-               <Label>RAM</Label>
+               <Label>Used RAM</Label>
                <Stack grow></Stack>
                {!!ramText && <Label>{ramText}</Label>}
             </Stack>
             <div id="telemetry_ram">
-               <Sparklines width={sparkWidth} height={sparkHeight} data={ramData} max={1} style={{ backgroundColor: dashboard.darktheme ? "#060709" : "#F3F2F1", padding: 8, border: "solid 1px #181A1B" }}>
-                  <SparklinesLine color={dashboard.darktheme ? "lightblue" : "#1E90FF"} />
-                  <SparklinesReferenceLine type="avg" />
+               <Sparklines width={sparkWidth} height={sparkHeight} data={ramData} min={0} max={1} style={{ backgroundColor: dashboard.darktheme ? "#060709" : "#F3F2F1", padding: 8, border: "solid 1px #181A1B" }}>
+                  <SparklinesLine color={dashboard.darktheme ? "lightblue" : "#1E90FF"} style={{ strokeWidth: 2 }} />
+                  <SparklinesReferenceLine type="custom" value={0} style={{ stroke: statusColors.get(StatusColor.Failure) }} />
+                  <SparklinesReferenceLine type="custom" value={22} style={{ stroke: statusColors.get(StatusColor.Warnings), opacity: opacity }} />
+                  <SparklinesReferenceLine type="custom" value={44} style={{ stroke: statusColors.get(StatusColor.Success), opacity: opacity }} />
                </Sparklines>
             </div>
          </Stack>
       </Stack>
       <Stack style={{ width: sparkWidth, paddingLeft: 4, paddingBottom: 4 }}>
          <Stack horizontal>
-            <Label>Free Disk</Label>
+            <Label>Used Disk</Label>
             <Stack grow></Stack>
             {!!diskText && <Label>{diskText}</Label>}
          </Stack>
-         <div id="telemetry_disk" style={{width: 690}}>
-            <Sparklines width={sparkWidth} height={sparkHeight} data={diskData} max={1} style={{ backgroundColor: dashboard.darktheme ? "#060709" : "#F3F2F1", padding: 8, border: "solid 1px #181A1B" }}>
-               <SparklinesLine color={dashboard.darktheme ? "lightblue" : "#1E90FF"} />
-               <SparklinesReferenceLine type="avg" />
+         <div id="telemetry_disk" style={{ width: 690 }}>
+            <Sparklines width={sparkWidth} height={sparkHeight} data={diskData} min={0} max={1} style={{ backgroundColor: dashboard.darktheme ? "#060709" : "#F3F2F1", padding: 8, border: "solid 1px #181A1B" }}>
+               <SparklinesLine color={dashboard.darktheme ? "lightblue" : "#1E90FF"} style={{ strokeWidth: 2 }} />
+               <SparklinesReferenceLine type="custom" value={0} style={{ stroke: statusColors.get(StatusColor.Failure) }} />
+               <SparklinesReferenceLine type="custom" value={22} style={{ stroke: statusColors.get(StatusColor.Warnings), opacity: opacity }} />
+               <SparklinesReferenceLine type="custom" value={44} style={{ stroke: statusColors.get(StatusColor.Success), opacity: opacity }} />
             </Sparklines>
          </div>
       </Stack>
diff --git a/Engine/Source/Programs/Horde/HordeDashboard/src/components/agents/HistoryModal.tsx b/Engine/Source/Programs/Horde/HordeDashboard/src/components/agents/HistoryModal.tsx
--- a/Engine/Source/Programs/Horde/HordeDashboard/src/components/agents/HistoryModal.tsx
+++ b/Engine/Source/Programs/Horde/HordeDashboard/src/components/agents/HistoryModal.tsx
@@ -14,6 +14,7 @@ import { getHordeStyling } from "../../styles/Styles";
 import { getHordeTheme } from "../../styles/theme";
 import { BatchStatusIcon, LeaseStatusIcon, StepStatusIcon } from "./../StatusIcon";
 import { historyModalStore } from "./HistoryModalStore";
+import { formatBytes } from "horde/base/utilities/stringUtills";
 
 const historyStyles = mergeStyleSets({
    dialog: {
@@ -36,7 +37,7 @@ const historyStyles = mergeStyleSets({
 
 const state = historyModalStore;
 
-const TelemetryTooltip: React.FC<{ id: string }> = observer(({ id }) => {
+const TelemetryTooltip: React.FC<{ id: "cpu" | "ram" | "disk", cpuKey?: string }> = observer(({ id, cpuKey }) => {
 
    const { modeColors } = getHordeStyling();
 
@@ -117,12 +118,36 @@ const TelemetryTooltip: React.FC<{ id: string }> = observer(({ id }) => {
       elements.push(dataElement("Log:", `Lease`, `/log/${lease.details.logId}`));
    }
 
-
-
    if (tooltip.time) {
       elements.push(dataElement("Time:", getNiceTime(tooltip.time)));
    }
 
+   if (tooltip.sample) {
+
+      const d = tooltip.sample;
+      let title = "Unknown";
+      let value = "";
+
+      switch (id) {
+         case "cpu":
+            title = "CPU";
+            if (cpuKey == "user") value = `${Math.ceil(d.userCpu)}%`;
+            if (cpuKey == "system") value = `${Math.ceil(d.systemCpu)}%`;
+            if (cpuKey == "combined") value = `${Math.ceil(d.userCpu + d.systemCpu)}%`;
+            break;
+         case "disk":
+            title = "Disk"
+            value = formatBytes((d.totalDisk - d.freeDisk) * 1024 * 1024);
+            break;
+         case "ram":
+            title = "Ram"
+            value = `${formatBytes(d.freeRam * 1024 * 1024)} Free`;
+            break;
+      }
+
+      elements.push(dataElement(`${title}:`, value));
+   }
+
 
    return <div style={{
       position: "absolute",
@@ -250,10 +275,6 @@ const TelemetryPanel: React.FC<{}> = observer(({ }) => {
    const minTime = getShortNiceTime(min, true);
    const maxTime = getShortNiceTime(max, true);
 
-   const ramData = data.map(d => {
-      return d.usedRam / d.totalRam;
-   })
-
    const diskData = data.map(d => {
       if (typeof (d.totalDisk) === "number" && d.totalDisk) {
          return 1.0 - (d.freeDisk / d.totalDisk)
@@ -271,7 +292,6 @@ const TelemetryPanel: React.FC<{}> = observer(({ }) => {
    if (cpuKey == "system") cpuLabel = "CPU - System";
    if (cpuKey == "combined") cpuLabel = "CPU - Combined";
 
-
    const devices = state.selectedAgent?.capabilities?.devices;
    if (devices?.length) {
       const device = devices[0];
@@ -307,21 +327,29 @@ const TelemetryPanel: React.FC<{}> = observer(({ }) => {
       })
    }
 
+   const ramData = data.map(d => {
+      return 1.0 - (d.freeRam / d.totalRam);
+   })
+
+
    function relativeCoords(event: any, id: string) {
 
-      if (state.currentLeaseTooltip?.frozen) {
+      if (state.currentLeaseTooltip?.frozen || !data?.length) {
          return;
       }
 
       let bounds = event.target.getBoundingClientRect();
       let x = event.clientX - bounds.left;
-      let y = event.clientY - bounds.top;
-
-      // interpolate time
+      let y = event.clientY - bounds.top;      
 
-      const percent = ((x - 16) / sparkWidth);
+      const percent = ((x) / 1130);
+      const sampleIdx = Math.ceil(data.length * percent);
+      if (sampleIdx < 0 || sampleIdx >= data.length) {
+         return;
+      }
 
-      const time = new Date(min.getTime() + ((max.getTime() - min.getTime()) * percent));
+      const sample = data[sampleIdx];
+      const time = sample.time;
 
       let lease: GetAgentLeaseResponse | undefined;
       leases.forEach(s => {
@@ -343,7 +371,7 @@ const TelemetryPanel: React.FC<{}> = observer(({ }) => {
          }
       }
 
-      state.setLeaseTooltip({ lease: lease, x: x, y: y, id: id, time: time });
+      state.setLeaseTooltip({ lease: lease, x: x, y: y, id: id, time: time, sample: sample });
    };
 
    return <Stack tokens={{ childrenGap: 24 }} style={{ paddingBottom: 32, paddingLeft: 8 }}>
@@ -391,10 +419,10 @@ const TelemetryPanel: React.FC<{}> = observer(({ }) => {
             }
          }}>
             <div style={{ position: "relative" }}>
-               <TelemetryTooltip id={"cpu"} />
+               <TelemetryTooltip id={"cpu"} cpuKey={cpuKey} />
             </div>
 
-            <Sparklines width={sparkWidth} height={128} data={cpuData} max={100} style={{ backgroundColor: dashboard.darktheme ? "#060709" : "#F3F2F1", padding: 8, border: "solid 1px #181A1B" }}>
+            <Sparklines width={sparkWidth} height={128} data={cpuData} min={0} max={100} style={{ backgroundColor: dashboard.darktheme ? "#060709" : "#F3F2F1", padding: 8, border: "solid 1px #181A1B" }}>
                <SparklinesLine color={dashboard.darktheme ? "lightblue" : "#1E90FF"} />
                <SparklinesReferenceLine type="custom" value={0} style={{ stroke: statusColors.get(StatusColor.Failure) }} />
                <SparklinesReferenceLine type="custom" value={32} style={{ stroke: statusColors.get(StatusColor.Warnings) }} />
@@ -419,7 +447,7 @@ const TelemetryPanel: React.FC<{}> = observer(({ }) => {
             <div style={{ position: "relative" }}>
                <TelemetryTooltip id="ram" />
             </div>
-            <Sparklines width={sparkWidth} height={128} data={ramData} max={1} style={{ backgroundColor: dashboard.darktheme ? "#060709" : "#F3F2F1", padding: 8, border: "solid 1px #181A1B" }}>
+            <Sparklines width={sparkWidth} height={128} data={ramData} min={0} max={1} style={{ backgroundColor: dashboard.darktheme ? "#060709" : "#F3F2F1", padding: 8, border: "solid 1px #181A1B" }}>
                <SparklinesLine color={dashboard.darktheme ? "lightblue" : "#1E90FF"} />
                <SparklinesReferenceLine type="custom" value={0} style={{ stroke: statusColors.get(StatusColor.Failure) }} />
                <SparklinesReferenceLine type="custom" value={32} style={{ stroke: statusColors.get(StatusColor.Warnings) }} />
@@ -445,7 +473,7 @@ const TelemetryPanel: React.FC<{}> = observer(({ }) => {
             <div style={{ position: "relative" }}>
                <TelemetryTooltip id="disk" />
             </div>
-            <Sparklines width={sparkWidth} height={128} data={diskData} max={1} style={{ backgroundColor: dashboard.darktheme ? "#060709" : "#F3F2F1", padding: 8, border: "solid 1px #181A1B" }}>
+            <Sparklines width={sparkWidth} height={128} data={diskData} min={0} max={1} style={{ backgroundColor: dashboard.darktheme ? "#060709" : "#F3F2F1", padding: 8, border: "solid 1px #181A1B" }}>
                <SparklinesLine color={dashboard.darktheme ? "lightblue" : "#1E90FF"} />
                <SparklinesReferenceLine type="custom" value={0} style={{ stroke: statusColors.get(StatusColor.Failure) }} />
                <SparklinesReferenceLine type="custom" value={32} style={{ stroke: statusColors.get(StatusColor.Warnings) }} />
diff --git a/Engine/Source/Programs/Horde/HordeDashboard/src/components/agents/HistoryModalStore.ts b/Engine/Source/Programs/Horde/HordeDashboard/src/components/agents/HistoryModalStore.ts
--- a/Engine/Source/Programs/Horde/HordeDashboard/src/components/agents/HistoryModalStore.ts
+++ b/Engine/Source/Programs/Horde/HordeDashboard/src/components/agents/HistoryModalStore.ts
@@ -9,6 +9,7 @@ type LeaseTooltip = {
    x?: number;
    y?: number;
    frozen?: boolean;
+   sample?: GetAgentTelemetrySampleResponse;
 }
 
 type InfoPanelItem = {
@@ -266,43 +267,7 @@ export class HistoryModalStore {
    @action
    setTelemetryData(telemetryData: GetAgentTelemetrySampleResponse[], sessionData: GetAgentLeaseResponse[]) {
 
-      const linear: GetAgentTelemetrySampleResponse[] = [];
-
-      if (telemetryData.length) {
-
-         let ctime: number = 0;
-
-         telemetryData.forEach(t => {
-
-            if (!ctime) {
-               ctime = t.time.getTime();
-               linear.push(t);
-            } else {
-               const ntime = t.time.getTime();
-               // dependant on agent sample rate!               
-               while ((ntime - ctime) > 60000) {
-
-                  ctime += 30000;
-                  linear.push({
-                     time: new Date(ctime),
-                     userCpu: 0,
-                     idleCpu: 0,
-                     systemCpu: 0,
-                     freeRam: 0,
-                     usedRam: 0,
-                     totalRam: 1,
-                     freeDisk: 0,
-                     totalDisk: 0
-                  });
-               }
-
-               linear.push(t);
-               ctime = ntime;
-            }
-         });
-      }
-
-      this.currentData = [linear, sessionData];
+      this.currentData = [telemetryData, sessionData];
 
    }
 
