Horde Dashboard: A few tweaks to make the agent modal telemetry more consistent

diff --git a/Engine/Source/Programs/Horde/HordeDashboard/src/components/agents/HistoryModal.tsx b/Engine/Source/Programs/Horde/HordeDashboard/src/components/agents/HistoryModal.tsx
--- a/Engine/Source/Programs/Horde/HordeDashboard/src/components/agents/HistoryModal.tsx
+++ b/Engine/Source/Programs/Horde/HordeDashboard/src/components/agents/HistoryModal.tsx
@@ -137,11 +137,11 @@ const TelemetryTooltip: React.FC<{ id: "cpu" | "ram" | "disk", cpuKey?: string }
             break;
          case "disk":
             title = "Disk"
-            value = formatBytes((d.totalDisk - d.freeDisk) * 1024 * 1024);
+            value = formatBytes((d.totalDisk - d.freeDisk) * 1024 * 1024) + " Used";
             break;
          case "ram":
             title = "Ram"
-            value = `${formatBytes(d.freeRam * 1024 * 1024)} Free`;
+            value = `${formatBytes(d.usedRam * 1024 * 1024)} Used`;
             break;
       }
 
@@ -340,7 +340,7 @@ const TelemetryPanel: React.FC<{}> = observer(({ }) => {
 
       let bounds = event.target.getBoundingClientRect();
       let x = event.clientX - bounds.left;
-      let y = event.clientY - bounds.top;      
+      let y = event.clientY - bounds.top;
 
       const percent = ((x) / 1130);
       const sampleIdx = Math.ceil(data.length * percent);
