Horde Server: Fix v2 device reservation exception with unit test

diff --git a/Engine/Source/Programs/Horde/Plugins/Build/HordeServer.Build.Tests/Devices/DeviceServiceTests.cs b/Engine/Source/Programs/Horde/Plugins/Build/HordeServer.Build.Tests/Devices/DeviceServiceTests.cs
--- a/Engine/Source/Programs/Horde/Plugins/Build/HordeServer.Build.Tests/Devices/DeviceServiceTests.cs
+++ b/Engine/Source/Programs/Horde/Plugins/Build/HordeServer.Build.Tests/Devices/DeviceServiceTests.cs
@@ -790,5 +790,24 @@ namespace HordeServer.Tests.Devices
 			Assert.AreEqual(issueReport.PoolReports[0].Metrics[0].MaxConcurrentProblemsPercentage, 37);
 			Assert.AreEqual(issueReport.PoolReports[0].Metrics[0].MaxConcurrentProblems, 3);
 		}
+
+		[TestMethod]
+		public async Task TestReservationV2Async()
+		{
+			await SetupDevicesAsync();
+			
+			CreateDeviceReservationRequest request = new CreateDeviceReservationRequest { PoolId = "TestDevicePool1", Devices = new List<DeviceReservationRequest> { new DeviceReservationRequest { PlatformId = "testdeviceplatform1" } } };
+
+			// create a reservation
+			CreateDeviceReservationResponse reservation = (await DeviceController!.CreateDeviceReservationAsync(request)).Value!;
+			Assert.AreEqual(1, reservation.Devices.Count);
+			Assert.AreEqual("testdeviceplatform1", reservation.Devices[0].PlatformId);
+
+			// delete the reservation, would be nice to also test reservation expiration
+			ActionResult deleted = (await DeviceController!.DeleteReservationAsync(reservation.Id));
+
+			Assert.AreEqual((deleted as OkResult)!.StatusCode, 200);
+
+		}
 	}
 }
diff --git a/Engine/Source/Programs/Horde/Plugins/Build/HordeServer.Build/Devices/DevicesController.cs b/Engine/Source/Programs/Horde/Plugins/Build/HordeServer.Build/Devices/DevicesController.cs
--- a/Engine/Source/Programs/Horde/Plugins/Build/HordeServer.Build/Devices/DevicesController.cs
+++ b/Engine/Source/Programs/Horde/Plugins/Build/HordeServer.Build/Devices/DevicesController.cs
@@ -556,6 +556,7 @@ namespace HordeServer.Devices
 			CreateDeviceReservationResponse response = new CreateDeviceReservationResponse();
 
 			response.Id = reservation.Id.ToString();
+			response.Devices = new List<GetDeviceResponse>();
 
 			foreach (IDevice device in devices)
 			{
@@ -565,7 +566,7 @@ namespace HordeServer.Devices
 					StepId = u.StepId,
 					ReservationStartUtc = u.ReservationStartUtc,
 					ReservationFinishUtc = u.ReservationFinishUtc
-				}).ToList();
+				}).ToList();				
 				response.Devices.Add(new GetDeviceResponse(device.Id.ToString(), device.PlatformId.ToString(), device.PoolId.ToString(), device.Name, device.Enabled, device.Address, device.ModelId?.ToString(), device.ModifiedByUser, device.Notes, device.ProblemTimeUtc, device.MaintenanceTimeUtc, utilization));
 			}
 
