Horde Dashboard: job search url parameters

diff --git a/Engine/Source/Programs/Horde/HordeDashboard/package.json b/Engine/Source/Programs/Horde/HordeDashboard/package.json
--- a/Engine/Source/Programs/Horde/HordeDashboard/package.json
+++ b/Engine/Source/Programs/Horde/HordeDashboard/package.json
@@ -16,6 +16,7 @@
     "@fluentui/react": "8.106.2",
     "@fluentui/react-charting": "5.2.0",
     "d3": "^7.6.1",
+    "lz-string": "^1.5.0",
     "markdown-to-jsx": "7.4.7",
     "mobx": "^6.12.2",
     "mobx-react-lite": "^4.0.7",
diff --git a/Engine/Source/Programs/Horde/HordeDashboard/src/components/JobSearchSimple.tsx b/Engine/Source/Programs/Horde/HordeDashboard/src/components/JobSearchSimple.tsx
--- a/Engine/Source/Programs/Horde/HordeDashboard/src/components/JobSearchSimple.tsx
+++ b/Engine/Source/Programs/Horde/HordeDashboard/src/components/JobSearchSimple.tsx
@@ -1,8 +1,8 @@
 import { Checkbox, ContextualMenuItemType, DatePicker, DefaultButton, DetailsList, DetailsRow, IColumn, IconButton, IContextualMenuItem, IContextualMenuItemProps, IContextualMenuProps, IDetailsGroupRenderProps, IDetailsListProps, IGroup, Label, Modal, PrimaryButton, SelectionMode, Spinner, SpinnerSize, Stack, Text, TextField } from "@fluentui/react";
 import { observer } from "mobx-react-lite";
 import moment from "moment";
-import { useState } from "react";
-import { Link } from "react-router-dom";
+import { useEffect, useMemo, useState } from "react";
+import { Link, useSearchParams } from "react-router-dom";
 import backend from "../backend";
 import { GetJobResponse, GetJobsTabResponse, GetJobStepRefResponse, JobData, JobQuery, TabType, GetTemplateRefResponse, GetJobsTabLabelColumnResponse, GetLabelStateResponse, LabelOutcome, LabelState, GetJobsTabColumnResponse } from "../backend/Api";
 import dashboard from "../backend/Dashboard";
@@ -17,7 +17,24 @@
 import { StreamChooser } from "./projects/StreamChooser";
 import { getLabelColor } from "horde/styles/colors";
 import { CalloutController, JobLabelCallout } from "./JobLabelCallout";
+import LZString from "lz-string"
 
+type SearchState = {
+   templates: Set<string>, 
+   statuses: Set<string>, 
+   name?: string, 
+   userId?: string,
+   global?: boolean,
+   streamIdOverride?: string,
+   preflights: boolean,
+   preflightOnly?: boolean,
+   minDate?: Date,
+   maxDate?: Date,
+   minCL?: string,
+   maxCL?: string,
+   containsStep?: string,
+   containsParameter?: string 
+}
 
 const jobToStep = new Map<string, GetJobStepRefResponse>();
 let jobStepName = "";
@@ -32,11 +49,12 @@
    const statuses = ["Waiting", "Running", "Complete", "Cancelled"]
    const allStatuses = new Set<string>(statuses);
 
-   const [searchState, setSearchState] = useState<{ items: JobItem[], groups?: IGroup[], templates: Set<string>, statuses: Set<string>, name?: string, userId?: string, global?: boolean, streamIdOverride?: string, preflights: boolean, preflightOnly?: boolean, minDate?: Date, maxDate?: Date, minCL?: string, maxCL?: string, querying: boolean, containsStep?: string, containsParameter?: string }>({ items: [], querying: false, preflights: false, minCL: changeListIn, maxCL: changeListIn, templates: new Set(jobTab?.templates ?? []), statuses: allStatuses });
+   const [searchState, setSearchState] = useState<SearchState>({ preflights: false, minCL: changeListIn, maxCL: changeListIn, templates: new Set(jobTab?.templates ?? []), statuses: allStatuses });
+   const [searchResults, setSearchResults] = useState<{items: JobItem[], groups?: IGroup[], querying: boolean}>({items: [], querying: false});
+   const [searchParams, setSearchParams] = useSearchParams();
    const [templateData, setTemplateData] = useState<{ streamId: string; templates: GetTemplateRefResponse[] } | undefined>(undefined);
    const [controller] = useState(new CalloutController());
-   const [autoload, setAutoload] = useState<boolean | undefined>(undefined);
+   const [isFirstRender, setIsFirstRender] = useState(true);
 
-
    const { hordeClasses, detailClasses } = getHordeStyling();
 
@@ -44,6 +62,264 @@
 
    const stream = projectStore.streamById(streamId);
 
+   // Make available before queryJobs, recomputing the values when needed/able
+   const ctemplates = useMemo(
+      () => (templateData?.templates ?? []).filter(t => searchState.templates.has(t.id)),
+      [templateData, searchState.templates]
+   );
+
+   const cstatuses = useMemo(
+      () => statuses.filter(s => searchState.statuses.has(s)),
+      [searchState.statuses]
+   );
+
+   // Update the search params in the url for sharing and bookmarking queires
+   const updateSearch = (state: SearchState) => {
+      console.log(state);
+
+      // Preserve current tab
+      const tab = searchParams.get("tab")!;
+      setSearchParams("");
+
+      const search = new URLSearchParams();
+      if(tab) search.append("tab", tab);
+      
+      // Compress current state to prevent 414s
+      // JSON.stringify doesn't handle Sets
+      const stringyState = {
+         ...state,
+         templates: Array.from(state.templates),
+         statuses: Array.from(state.statuses)
+      };
+      const compressedStateString = LZString.compressToEncodedURIComponent(JSON.stringify(stringyState));
+
+      search.append("search", compressedStateString);
+
+      setSearchParams(search);
+   }
+
+   const queryJobs = async (state: SearchState) => {      
+      updateSearch(state);
+      setSearchResults(prev => ({ ...prev, querying: true }));
+
+      let jobs: JobData[] = [];
+
+      try {
+
+         let name = state.name?.trim();
+         if (!name) {
+            name = undefined;
+         } else {
+            name = "$" + name;
+         }
+
+         let minCL: number | undefined = parseInt(state.minCL ?? "");
+         if (isNaN(minCL)) {
+            minCL = undefined;
+         }
+         let maxCL: number | undefined = parseInt(state.maxCL ?? "");
+         if (isNaN(maxCL)) {
+            maxCL = undefined;
+         }
+
+         if (!minCL && maxCL) {
+            minCL = 0;
+         }
+
+         if (minCL && !maxCL) {
+            maxCL = 999999999;
+         }
+
+         let minCreateTime = state.minDate;
+         let maxCreateTime = state.maxDate;
+         
+         if (maxCreateTime) {
+
+            const maxTime = moment(maxCreateTime);
+            maxTime.add(23, "hours");
+            maxTime.add(59, "minutes");
+            maxTime.add(59, "seconds");
+
+            maxCreateTime = maxTime.toDate();
+
+         }
+
+         const maxCount = 2048;
+
+         let searchStreamId: string | undefined = streamId;
+
+         if (state.streamIdOverride) {
+            searchStreamId = state.streamIdOverride
+         } else if (state.global && !state.containsStep) {
+            searchStreamId = undefined;
+         }
+
+         const stepName = state.containsStep?.trim();
+         const stepJobIds: string[] = []
+
+         jobToStep.clear();
+         jobStepName = stepName ?? "";
+         if (stepName) {
+
+            let stepTemplates: string[] = [];
+
+            if (ctemplates?.length) {
+               stepTemplates.push(...ctemplates.map(t => t.id));
+            } else {
+               stepTemplates = templates.filter(t => !!t.id).map(t => t.id);
+            }
+
+            const numTemplates = stepTemplates.length;
+            while (stepTemplates.length) {
+
+               const batch = stepTemplates.slice(0, 5);
+               stepTemplates = stepTemplates.slice(5);
+
+               // @todo: get endpoint to handle multiple templates
+               const results = await Promise.all(batch.map(template => backend.getJobStepHistory(searchStreamId!, stepName, numTemplates === 1 ? 256 : 64, template)));
+               results.forEach(r => {
+                  r.forEach(j => {
+                     stepJobIds.push(j.jobId)
+                     jobToStep.set(j.jobId, j);
+                  });
+               });
+
+               if (stepJobIds.length >= 256) {
+                  break;
+               }
+            }
+         }
+
+         const jobIds = stepJobIds.slice(0, 256);
+
+         if (stepName && !jobIds.length) {
+            jobs = [];
+         } else {
+            
+            const query: JobQuery = {
+               id: jobIds.length ? jobIds : undefined,
+               streamId: searchStreamId?.toLowerCase(),
+               name: name,
+               filter: "id,labels,defaultLabel,name,change,createTime,streamId,preflightChange,startedByUserInfo,arguments,abortedByUserInfo,state",
+               minChange: minCL,
+               maxChange: maxCL,
+               t: (!state.global && ctemplates.length && !state.containsStep) ? ctemplates.map(t => t.id) : undefined,
+               minCreateTime: minCreateTime?.toUTCString(),
+               maxCreateTime: maxCreateTime?.toUTCString(),
+               includePreflight: state.preflights || !!state.preflightOnly,
+               preflightOnly: !!state.preflightOnly,
+               startedByUserId: state.userId,
+               count: maxCount
+            }
+
+            jobs = await backend.getJobs(query);
+
+            if (cstatuses.length && cstatuses.length != allStatuses.size) {
+               jobs = jobs.filter(j => {
+                  return (cstatuses.includes('Cancelled') && j.abortedByUserInfo?.id || cstatuses.includes(j.state) && !j.abortedByUserInfo?.id)
+               })
+            }
+
+            const p = state.containsParameter?.trim().toLowerCase();
+            if (p) {
+               jobs = jobs.filter(j => {
+                  return !!j.arguments.find(a => {
+                     return a.toLowerCase().indexOf(p) !== -1
+                  })
+               });
+            }
+         }
+
+      } catch (reason) {
+         console.error(reason);
+      } finally {
+
+         /*
+         if (column.name === "Stream") {
+             const stream = projectStore.streamById(item.job.streamId);
+             if (!stream) {
+                 return null;
+             }
+             return <Stack horizontal verticalFill={true} verticalAlign="center" tokens={{ childrenGap: 0, padding: 0 }} style={{ overflow: "hidden" }} ><Text>{stream.fullname}</Text></Stack>;
+         }
+         */
+
+         const streamIds = new Set<string>();
+         const idToFullname = new Map<string, string>();
+
+         jobs.forEach(j => {
+            const stream = projectStore.streamById(j.streamId);
+            if (!stream || !stream.fullname) {
+               streamIds.add(j.streamId);
+               idToFullname.set(j.streamId, j.streamId);
+               return;
+            };
+            streamIds.add(j.streamId);
+            idToFullname.set(j.streamId, stream.fullname);
+
+         });
+
+         const sorted = Array.from(streamIds).sort((a, b) => {
+            const aname = idToFullname.get(a)!;
+            const bname = idToFullname.get(b)!;
+
+            if (aname < bname) return -1;
+            if (aname > bname) return 1;
+            return 0;
+         })
+
+         // emit groups
+         const groups: IGroup[] = [];
+         const items: JobItem[] = [];
+         sorted.forEach(id => {
+            const streamJobs = jobs.filter(j => j.streamId === id);
+            groups.push({ key: `jobs_${id}`, name: idToFullname.get(id)!, startIndex: items.length, count: streamJobs.length, isCollapsed: false });
+            items.push(...streamJobs.map(j => { return { job: j } }));
+         })
+
+         setSearchResults({items, groups, querying: false });
+      }
+
+   }
+   
+   // Clear search parameters from the URL
+   const clearSearchParams = () => {
+      const tab = searchParams.get("tab")!;
+      setSearchParams("");
+
+      const search = new URLSearchParams();
+      if(tab) search.append("tab", tab);
+
+      setSearchParams(search);
+   }
+
+   // Sets state from search params and autoloads jobs
+   useEffect(() => {
+      if(!searchParams.has("search") || !isFirstRender) return;
+      
+      setIsFirstRender(false);
+
+      // Decompress URL data
+      const stringyState = JSON.parse(LZString.decompressFromEncodedURIComponent(searchParams.get("search")!));
+      
+      const state: SearchState = {
+         ...stringyState,
+         minDate: stringyState.minDate ? new Date(stringyState.minDate) : undefined,
+         maxDate: stringyState.maxDate ? new Date(stringyState.maxDate) : undefined,
+         templates: new Set(stringyState.templates),
+         statuses: new Set(stringyState.statuses)
+      };
+
+      if(state.global){
+         streamChooserId++;
+         state.templates = new Set();
+         state.streamIdOverride = undefined;
+      }
+
+      queryJobs(state)
+      setSearchState(state);
+   }, [])
+
    if (!stream) {
       console.error("unable to get stream");
       return <div>unable to get stream</div>;
@@ -124,13 +400,21 @@
                return <Stack horizontal verticalFill verticalAlign="center" style={{ padding: "4px 12px 4px 12px", cursor: "pointer" }} onClick={(ev) => {
                   ev?.preventDefault();
                   ev?.stopPropagation();
-                  if (searchState.templates.has(t.id)) {
-                     searchState.templates.delete(t.id);
-                  } else {
-                     searchState.templates.add(t.id);
-                  }
-                  setSearchState({ ...searchState })
 
+                  setSearchState(prev => {
+                     const newTemplates = new Set(prev.templates);
+
+                     if (newTemplates.has(t.id)) {
+                        newTemplates.delete(t.id);
+                     } else {
+                        newTemplates.add(t.id);
+                     }
+
+                     return {
+                        ...prev,
+                        templates: newTemplates
+                     };   
+                  });
                }} >
                   <Checkbox checked={searchState.templates.has(t.id)} />
                   <Text>{t.name}</Text></Stack>
@@ -152,18 +436,20 @@
       
                   allCatTemplatesSelected = !allCatTemplatesSelected;
                   
+                  const newTemplates = new Set(searchState.templates);
+
                   if (allCatTemplatesSelected) {
                      catTemplates.forEach( ct => 
-                        searchState.templates.add(ct.id)
+                        newTemplates.add(ct.id)
                      )
-                     setSearchState({ ...searchState, containsParameter: undefined })
                   }
                   else {
                      catTemplates.forEach( ct => 
-                        searchState.templates.delete(ct.id)
+                        newTemplates.delete(ct.id)
                      )
-                     setSearchState({ ...searchState, containsParameter: undefined })
                   }
+
+                  setSearchState(prev => ({...prev, templates: newTemplates}));                  
                   return false;
                }} >
                <Checkbox indeterminate={someCatTemplatesSelected} checked={allCatTemplatesSelected} />
@@ -214,8 +500,6 @@
       items: templateOptions
    };
 
-   const ctemplates = templates.filter(t => searchState.templates.has(t.id));
-
    // Statuses
    let statusOptions: IContextualMenuItem[] = [];
 
@@ -256,12 +540,20 @@
                ev?.preventDefault();
                ev?.stopPropagation();
 
-               if (searchState.statuses.has(status)) {
-                  searchState.statuses.delete(status);
-               } else {
-                  searchState.statuses.add(status);
-               }
-               setSearchState({ ...searchState })
+               setSearchState(prev => {
+                  const newStatuses = new Set(prev.statuses);
+
+                  if (newStatuses.has(status)) {
+                     newStatuses.delete(status);
+                  } else {
+                     newStatuses.add(status);
+                  }
+
+                  return {
+                     ...prev,
+                     statuses: newStatuses
+                  };   
+               });
             }} >
 
                <Checkbox checked={searchState.statuses.has(status)} />
@@ -275,8 +567,6 @@
       items: statusOptions
    };
 
-   const cstatuses = statuses.filter(s => searchState.statuses.has(s))
-
    // results
    const labelColumnWidth = 396;
    const columns = [
@@ -480,200 +770,6 @@
       return <Stack />;
    }
 
-   const queryJobs = async () => {
-
-      setSearchState({ ...searchState, querying: true });
-
-      let jobs: JobData[] = [];
-
-      try {
-
-         let name = searchState.name?.trim();
-         if (!name) {
-            name = undefined;
-         } else {
-            name = "$" + name;
-         }
-
-         let minCL: number | undefined = parseInt(searchState.minCL ?? "");
-         if (isNaN(minCL)) {
-            minCL = undefined;
-         }
-         let maxCL: number | undefined = parseInt(searchState.maxCL ?? "");
-         if (isNaN(maxCL)) {
-            maxCL = undefined;
-         }
-
-         if (!minCL && maxCL) {
-            minCL = 0;
-         }
-
-         if (minCL && !maxCL) {
-            maxCL = 999999999;
-         }
-
-         let minCreateTime = searchState.minDate;
-         let maxCreateTime = searchState.maxDate;
-
-         if (minCreateTime && (!maxCreateTime || minCreateTime === maxCreateTime)) {
-
-            //const maxTime = moment(minCreateTime);
-            //maxTime.add(23, "hours");
-            //maxTime.add(59, "minutes");
-            //maxTime.add(59, "seconds");
-
-            //maxCreateTime = maxTime.toDate();
-
-         } else if (maxCreateTime) {
-
-            const maxTime = moment(maxCreateTime);
-            maxTime.add(23, "hours");
-            maxTime.add(59, "minutes");
-            maxTime.add(59, "seconds");
-
-            maxCreateTime = maxTime.toDate();
-
-         }
-
-         const maxCount = 2048;
-
-         let searchStreamId: string | undefined = streamId;
-
-         if (searchState.streamIdOverride) {
-            searchStreamId = searchState.streamIdOverride
-         } else if (searchState.global && !searchState.containsStep) {
-            searchStreamId = undefined;
-         }
-
-         const stepName = searchState.containsStep?.trim();
-         const stepJobIds: string[] = []
-
-         jobToStep.clear();
-         jobStepName = stepName ?? "";
-         if (stepName) {
-
-            let stepTemplates: string[] = [];
-
-            if (ctemplates?.length) {
-               stepTemplates.push(...ctemplates.map(t => t.id));
-            } else {
-               stepTemplates = templates.filter(t => !!t.id).map(t => t.id);
-            }
-
-            const numTemplates = stepTemplates.length;
-            while (stepTemplates.length) {
-
-               const batch = stepTemplates.slice(0, 5);
-               stepTemplates = stepTemplates.slice(5);
-
-               // @todo: get endpoint to handle multiple templates
-               const results = await Promise.all(batch.map(template => backend.getJobStepHistory(searchStreamId!, stepName, numTemplates === 1 ? 256 : 64, template)));
-               results.forEach(r => {
-                  r.forEach(j => {
-                     stepJobIds.push(j.jobId)
-                     jobToStep.set(j.jobId, j);
-                  });
-               });
-
-               if (stepJobIds.length >= 256) {
-                  break;
-               }
-            }
-         }
-
-         const jobIds = stepJobIds.slice(0, 256);
-
-         if (stepName && !jobIds.length) {
-            jobs = [];
-         } else {
-
-            const query: JobQuery = {
-               id: jobIds.length ? jobIds : undefined,
-               streamId: searchStreamId?.toLowerCase(),
-               name: name,
-               filter: "id,labels,defaultLabel,name,change,createTime,streamId,preflightChange,startedByUserInfo,arguments,abortedByUserInfo,state",
-               minChange: minCL,
-               maxChange: maxCL,
-               t: (!searchState.global && ctemplates.length && !searchState.containsStep) ? ctemplates.map(t => t.id) : undefined,
-               minCreateTime: minCreateTime?.toUTCString(),
-               maxCreateTime: maxCreateTime?.toUTCString(),
-               includePreflight: searchState.preflights || !!searchState.preflightOnly,
-               preflightOnly: !!searchState.preflightOnly,
-               startedByUserId: searchState.userId,
-               count: maxCount
-            }
-
-            jobs = await backend.getJobs(query);
-
-            if (cstatuses.length && cstatuses.length != allStatuses.size) {
-               jobs = jobs.filter(j => {
-                  return (cstatuses.includes('Cancelled') && j.abortedByUserInfo?.id || cstatuses.includes(j.state) && !j.abortedByUserInfo?.id)
-               })
-            }
-
-            const p = searchState.containsParameter?.trim().toLowerCase();
-            if (p) {
-               jobs = jobs.filter(j => {
-                  return !!j.arguments.find(a => {
-                     return a.toLowerCase().indexOf(p) !== -1
-                  })
-               });
-            }
-         }
-
-      } catch (reason) {
-
-      } finally {
-
-         /*
-         if (column.name === "Stream") {
-             const stream = projectStore.streamById(item.job.streamId);
-             if (!stream) {
-                 return null;
-             }
-             return <Stack horizontal verticalFill={true} verticalAlign="center" tokens={{ childrenGap: 0, padding: 0 }} style={{ overflow: "hidden" }} ><Text>{stream.fullname}</Text></Stack>;
-         }
-         */
-
-         const streamIds = new Set<string>();
-         const idToFullname = new Map<string, string>();
-
-         jobs.forEach(j => {
-            const stream = projectStore.streamById(j.streamId);
-            if (!stream || !stream.fullname) {
-               streamIds.add(j.streamId);
-               idToFullname.set(j.streamId, j.streamId);
-               return;
-            };
-            streamIds.add(j.streamId);
-            idToFullname.set(j.streamId, stream.fullname);
-
-         });
-
-         const sorted = Array.from(streamIds).sort((a, b) => {
-            const aname = idToFullname.get(a)!;
-            const bname = idToFullname.get(b)!;
-
-            if (aname < bname) return -1;
-            if (aname > bname) return 1;
-            return 0;
-         })
-
-         // emit groups
-         const groups: IGroup[] = [];
-         const items: JobItem[] = [];
-         sorted.forEach(id => {
-            const streamJobs = jobs.filter(j => j.streamId === id);
-            groups.push({ key: `jobs_${id}`, name: idToFullname.get(id)!, startIndex: items.length, count: streamJobs.length, isCollapsed: false });
-            items.push(...streamJobs.map(j => { return { job: j } }));
-         })
-
-
-         setSearchState({ ...searchState, items: items, groups: groups, querying: false });
-      }
-
-   }
-
    const renderRow: IDetailsListProps['onRenderRow'] = (props) => {
 
       if (props) {
@@ -719,14 +815,22 @@
       }
    }
 
-   if (autoloadIn && autoload === undefined) {
-      queryJobs();
-      setAutoload(autoloadIn);
+   // Can auto load normally (e.g., clicking "related jobs") if not already a search link
+   if(autoloadIn && isFirstRender && !searchParams.has("search")){
+      queryJobs(searchState);
+      setIsFirstRender(false);
       return null;
    }
-
-   return (<Modal isOpen={true} topOffsetFixed={true} styles={{ main: { padding: 8, width: 1400, height: '820px', backgroundColor: hordeTheme.horde.contentBackground, hasBeenOpened: false, top: "80px", position: "absolute" } }} className={hordeClasses.modal} onDismiss={() => { onClose() }}>
-      {searchState.querying && <StatusModal text={"Finding Jobs"} />}
+   
+   return (
+   <Modal 
+      isOpen={true} 
+      topOffsetFixed={true} 
+      styles={{ main: { padding: 8, width: 1400, height: '820px', backgroundColor: hordeTheme.horde.contentBackground, hasBeenOpened: false, top: "80px", position: "absolute" } }} 
+      className={hordeClasses.modal} 
+      onDismiss={() => { clearSearchParams(); onClose(); }}
+   >
+      {searchResults.querying && <StatusModal text={"Finding Jobs"} />}
       <JobLabelCallout controller={controller} />
       <Stack styles={{ root: { paddingTop: 8, paddingLeft: 24, paddingRight: 12, paddingBottom: 8 } }}>
          <Stack tokens={{ childrenGap: 12 }}>
@@ -737,7 +841,7 @@
                <Stack grow horizontalAlign="end">
                   <IconButton
                      iconProps={{ iconName: 'Cancel' }}
-                     onClick={() => { onClose(); }}
+                     onClick={() => { clearSearchParams(); onClose(); }}
                   />
                </Stack>
             </Stack>
@@ -792,7 +896,7 @@
 
                      <Stack tokens={{ childrenGap: 4 }}>
                         <Label>Started By</Label>
-                        <UserSelect handleSelection={(id => {
+                        <UserSelect defaultId={searchState.userId} handleSelection={(id => {
                            setSearchState({ ...searchState, userId: id })
                         })} userIdHints={[]} noResultsFoundText="No users found" />
                      </Stack>
@@ -844,10 +948,20 @@
 
                      <Stack horizontal style={{ paddingTop: 20 }}>
                         <Stack style={{ paddingTop: 8, paddingRight: 24 }}>
-                           <PrimaryButton disabled={searchState.querying} text="Search" onClick={() => { queryJobs() }} />
+                           <PrimaryButton 
+                              disabled={searchResults.querying} 
+                              text="Search" 
+                              onClick={() => {
+                                 queryJobs(searchState);
+                              }} 
+                           />
                         </Stack>
                         <Stack style={{ paddingTop: 8, paddingRight: 24 }}>
-                           <DefaultButton disabled={searchState.querying} text="Clear" onClick={() => { setSearchState({ items: [], groups: [], querying: false, preflights: false, userId: searchState.userId, templates: new Set(), statuses: new Set()}) }} />
+                           <DefaultButton disabled={searchResults.querying} text="Clear" onClick={() => { 
+                                 setSearchState({ preflights: false, userId: undefined, templates: new Set(), statuses: new Set()}) 
+                                 setSearchResults({items: [], groups: [], querying: false});
+                                 clearSearchParams();
+                              }} />
                         </Stack>
                      </Stack>
                   </Stack>
@@ -859,19 +973,19 @@
                         <Label>Jobs</Label>
                      </Stack>
 
-                     {!searchState.items.length && <Stack>
+                     {!searchResults.items.length && <Stack>
                         <Text>No Results</Text>
                      </Stack>}
 
-                     {!!searchState.items.length && <Stack>
+                     {!!searchResults.items.length && <Stack>
                         <div style={{ overflowY: 'auto', overflowX: 'hidden', height: "700px" }} data-is-scrollable={true} onScroll={() => { controller.setState(undefined, true) }}>
                            <Stack tokens={{ childrenGap: 12 }} style={{ paddingRight: 12 }}>
                               <DetailsList
                                  compact={true}
                                  isHeaderVisible={false}
                                  indentWidth={0}
-                                 items={searchState.items}
-                                 groups={searchState.groups}
+                                 items={searchResults.items}
+                                 groups={searchResults.groups}
                                  columns={columns}
                                  setKey="set"
                                  selectionMode={SelectionMode.none}
diff --git a/Engine/Source/Programs/Horde/HordeDashboard/src/components/StreamView.tsx b/Engine/Source/Programs/Horde/HordeDashboard/src/components/StreamView.tsx
--- a/Engine/Source/Programs/Horde/HordeDashboard/src/components/StreamView.tsx
+++ b/Engine/Source/Programs/Horde/HordeDashboard/src/components/StreamView.tsx
@@ -4,7 +4,7 @@ import { DefaultButton, FontIcon, HoverCard, HoverCardType, IContextualMenuProps
 import { action, makeObservable, observable } from 'mobx';
 import { observer } from 'mobx-react-lite';
 import React, { useState } from 'react';
-import { Link, Navigate, useNavigate, useParams } from 'react-router-dom';
+import { Link, Navigate, useNavigate, useParams, useSearchParams } from 'react-router-dom';
 import backend, { useBackend } from '../backend';
 import { GetJobsTabResponse, GetLabelStateResponse, GetStreamTabResponse, GetTemplateRefResponse, JobData, JobsTabData, LabelOutcome, LabelState, ProjectData, TabStyle } from '../backend/Api';
 import dashboard, { StatusColor } from '../backend/Dashboard';
@@ -343,7 +343,9 @@ const StreamViewInner: React.FC = observer(() => {
    const [showOthersPreflights, setShowOthersPreflights] = useState<boolean | undefined>(dashboard.showPreflights);
 
    const [shown, setShown] = useState(query.get("newbuild") ? true : false);
-   const [findJobsShown, setFindJobsShown] = useState(false);
+   const [searchParams] = useSearchParams();
+   const searchInURL: boolean = searchParams.has("search");
+   const [findJobsShown, setFindJobsShown] = useState(searchInURL);
 
    const { projectStore } = useBackend();
 
@@ -512,7 +514,7 @@ const StreamViewInner: React.FC = observer(() => {
                }
             }
          }} />}
-         {findJobsShown && <JobSearchSimpleModal onClose={() => { setFindJobsShown(false) }} streamIdIn={stream.id} jobTab={currentTab} />}
+         {findJobsShown && <JobSearchSimpleModal onClose={() => { setFindJobsShown(false) }} streamIdIn={stream.id} jobTab={currentTab} autoloadIn={searchInURL} />}
          <Stack horizontal style={{ backgroundColor: hordeTheme.horde.neutralBackground }}>
             <div key={`windowsize_streamview_${windowSize.width}_${windowSize.height}`} style={{ width: (vw / 2 - (1440 / 2)) - 12, flexShrink: 0, backgroundColor: modeColors.background }} />
             <Stack tokens={{ childrenGap: 0 }} styles={{ root: { width: "100%" } }}>
diff --git a/Engine/Source/Programs/Horde/HordeDashboard/src/components/UserSelect.tsx b/Engine/Source/Programs/Horde/HordeDashboard/src/components/UserSelect.tsx
--- a/Engine/Source/Programs/Horde/HordeDashboard/src/components/UserSelect.tsx
+++ b/Engine/Source/Programs/Horde/HordeDashboard/src/components/UserSelect.tsx
@@ -1,6 +1,6 @@
 import { CompactPeoplePicker, IBasePickerSuggestionsProps, IInputProps, IPersonaProps, Label, Stack } from "@fluentui/react";
 import { observer } from "mobx-react-lite";
-import React, { useState } from "react";
+import React, { useEffect, useState } from "react";
 import { useWindowSize } from "../base/utilities/hooks";
 import { Breadcrumbs } from "./Breadcrumbs";
 import { TopNav } from "./TopNav";
@@ -11,17 +11,42 @@ import { getHordeStyling } from "../styles/Styles";
 
 const loadedUsers: Map<string, GetUserResponse> = new Map();
 
-export const UserSelect: React.FC<{ handleSelection: (userId: string | undefined) => void, userIdHints?: string[], noResultsFoundText?: string, defaultUser?: GetUserResponse, autoFocus?: boolean }> = observer(({ handleSelection, userIdHints, noResultsFoundText, defaultUser, autoFocus }) => {
+export const UserSelect: React.FC<{ handleSelection: (userId: string | undefined) => void, userIdHints?: string[], noResultsFoundText?: string, defaultUser?: GetUserResponse, autoFocus?: boolean, defaultId?: string }> = observer(({ handleSelection, userIdHints, noResultsFoundText, defaultUser, autoFocus, defaultId }) => {
+   const [selectedItems, setSelectedItems] = useState<IPersonaProps[]>([]);
 
-   const [state, setState] = useState<{ initial?: boolean }>({ initial: true });
+   const loadUser = (user: GetUserResponse) => {
+      if(!user) return;
 
-   if (defaultUser && state.initial) {
-      loadedUsers.set(defaultUser.id, defaultUser);
-      handleSelection(defaultUser.id);
-      setState({});
-      return null;
+      loadedUsers.set(user.id, user);
+      
+      const persona = {
+         text: user.name,
+         imageUrl: user.image24,
+         optionalText: user.id,
+      } as IPersonaProps;
+
+      setSelectedItems([persona]);
+      handleSelection(user.id);
+   }
+
+   const loadUserFromId = async (userId: string) => {
+      if(!userId) return;
+
+      const users = await backend.getUsers({ids: [userId], includeAvatar: true});
+      if(users.length) loadUser(users[0]);
    }
 
+   useEffect(() => {
+      if(defaultUser) {
+         loadUser(defaultUser);
+      } else if(defaultId) {
+         loadUserFromId(defaultId);
+      } else {
+         setSelectedItems([])
+         handleSelection(undefined);
+      }
+   }, [defaultUser, defaultId]);
+
    const suggestionProps: IBasePickerSuggestionsProps = {
       suggestionsHeaderText: '',
       mostRecentlyUsedHeaderText: '',
@@ -109,10 +134,10 @@ export const UserSelect: React.FC<{ handleSelection: (userId: string | undefined
          onEmptyResolveSuggestions={onEmptyResolveSuggestions}
          pickerSuggestionsProps={suggestionProps}
          className={'ms-PeoplePicker'}
-         defaultSelectedItems={defaultUser ? [{ text: defaultUser.name, imageUrl: defaultUser.image24, optionalText: defaultUser.id }] : []}
+         selectedItems={selectedItems}
          inputProps={inputProps}
          onChange={(items) => {
-
+            setSelectedItems(items ?? []);
             if (items && items?.length) {
                handleSelection(items[0].optionalText!);
             } else {
diff --git a/Engine/Source/Programs/Horde/HordeDashboard/src/components/jobDetailsV2/JobOperationsBar.tsx b/Engine/Source/Programs/Horde/HordeDashboard/src/components/jobDetailsV2/JobOperationsBar.tsx
--- a/Engine/Source/Programs/Horde/HordeDashboard/src/components/jobDetailsV2/JobOperationsBar.tsx
+++ b/Engine/Source/Programs/Horde/HordeDashboard/src/components/jobDetailsV2/JobOperationsBar.tsx
@@ -3,7 +3,7 @@
 import { CommandBar, CommandBarButton, ICommandBarItemProps, IContextualMenuItem, Stack } from '@fluentui/react';
 import { observer } from 'mobx-react-lite';
 import React, { useState } from 'react';
-import { Link, useNavigate } from 'react-router-dom';
+import { Link, useNavigate, useSearchParams } from 'react-router-dom';
 import { ArtifactContextType, JobState, JobStepOutcome, JobStepState, StepData } from '../../backend/Api';
 import { getSiteConfig } from '../../backend/Config';
 import dashboard from '../../backend/Dashboard';
@@ -34,7 +34,9 @@ export const JobOperations: React.FC<{ jobDetails: JobDetailsV2 }> = observer(({
    const [abortShown, setAbortShown] = useState(false);
    const [editShown, setEditShown] = useState(false);
    const [retryStepsShown, setRetryStepsShown] = useState(false);
-   const [findJobsShown, setFindJobsShown] = useState(false);
+
+   const [searchParams] = useSearchParams();
+   const [findJobsShown, setFindJobsShown] = useState(searchParams.has("search"));
 
    const [parametersState, setParametersState] = useState(query.get("newbuild") ? ParameterState.Clone : ParameterState.Hidden);
 
