Horde: Add cancellation token to IPluginConfig.PostLoadAsync method parameters

diff --git a/Engine/Source/Programs/Horde/HordeServer.Shared/Plugins/IPluginConfig.cs b/Engine/Source/Programs/Horde/HordeServer.Shared/Plugins/IPluginConfig.cs
--- a/Engine/Source/Programs/Horde/HordeServer.Shared/Plugins/IPluginConfig.cs
+++ b/Engine/Source/Programs/Horde/HordeServer.Shared/Plugins/IPluginConfig.cs
@@ -13,13 +13,14 @@ namespace HordeServer.Plugins
 		/// Called to fixup a plugin's configuration after deserialization
 		/// </summary>
 		/// <param name="configOptions">Options for configuring the plugin</param>
-		Task PostLoadAsync(PluginConfigOptions configOptions);
+		/// <param name="cancellationToken">Cancellation token for the operation</param>
+		Task PostLoadAsync(PluginConfigOptions configOptions, CancellationToken cancellationToken = default);
 	}
 
 	/// <summary>
-	/// Options passed to <see cref="IPluginConfig.PostLoadAsync(PluginConfigOptions)"/>
+	/// Options passed to <see cref="IPluginConfig.PostLoadAsync(PluginConfigOptions,CancellationToken)"/>
 	/// </summary>
-	public record class PluginConfigOptions(ConfigVersion Version, IEnumerable<IPluginConfig> Plugins, AclConfig ParentAcl, IServiceProvider? ServiceProvider = null, CancellationToken CancellationToken = default);
+	public record class PluginConfigOptions(ConfigVersion Version, IEnumerable<IPluginConfig> Plugins, AclConfig ParentAcl, IServiceProvider? ServiceProvider = null);
 
 	/// <summary>
 	/// Empty implementation of <see cref="IPluginConfig"/>
@@ -27,7 +28,7 @@ namespace HordeServer.Plugins
 	public sealed class EmptyPluginConfig : IPluginConfig
 	{
 		/// <inheritdoc/>
-		public Task PostLoadAsync(PluginConfigOptions configOptions)
+		public Task PostLoadAsync(PluginConfigOptions configOptions, CancellationToken cancellationToken = default)
 		{
 			return Task.CompletedTask;
 		}
diff --git a/Engine/Source/Programs/Horde/HordeServer/Server/GlobalConfig.cs b/Engine/Source/Programs/Horde/HordeServer/Server/GlobalConfig.cs
--- a/Engine/Source/Programs/Horde/HordeServer/Server/GlobalConfig.cs
+++ b/Engine/Source/Programs/Horde/HordeServer/Server/GlobalConfig.cs
@@ -156,10 +156,10 @@ namespace HordeServer.Server
 				pluginConfigs = sortedPluginConfigs;
 			}
 
-			PluginConfigOptions pluginConfigOptions = new PluginConfigOptions(VersionEnum, Plugins.Values, Acl, serviceProvider, cancellationToken);
+			PluginConfigOptions pluginConfigOptions = new PluginConfigOptions(VersionEnum, Plugins.Values, Acl, serviceProvider);
 			foreach (IPluginConfig pluginConfig in pluginConfigs)
 			{
-				await pluginConfig.PostLoadAsync(pluginConfigOptions);
+				await pluginConfig.PostLoadAsync(pluginConfigOptions, cancellationToken);
 			}
 
 			foreach (ScheduledDowntime downtime in Downtime)
diff --git a/Engine/Source/Programs/Horde/Plugins/Analytics/HordeServer.Analytics/AnalyticsConfig.cs b/Engine/Source/Programs/Horde/Plugins/Analytics/HordeServer.Analytics/AnalyticsConfig.cs
--- a/Engine/Source/Programs/Horde/Plugins/Analytics/HordeServer.Analytics/AnalyticsConfig.cs
+++ b/Engine/Source/Programs/Horde/Plugins/Analytics/HordeServer.Analytics/AnalyticsConfig.cs
@@ -26,7 +26,7 @@ namespace HordeServer
 		private readonly Dictionary<TelemetryStoreId, TelemetryStoreConfig> _telemetryStoreLookup = new Dictionary<TelemetryStoreId, TelemetryStoreConfig>();
 
 		/// <inheritdoc/>
-		public Task PostLoadAsync(PluginConfigOptions configOptions)
+		public Task PostLoadAsync(PluginConfigOptions configOptions, CancellationToken cancellationToken = default)
 		{
 			_parentAcl = configOptions.ParentAcl;
 
diff --git a/Engine/Source/Programs/Horde/Plugins/Build/HordeServer.Build/BuildConfig.cs b/Engine/Source/Programs/Horde/Plugins/Build/HordeServer.Build/BuildConfig.cs
--- a/Engine/Source/Programs/Horde/Plugins/Build/HordeServer.Build/BuildConfig.cs
+++ b/Engine/Source/Programs/Horde/Plugins/Build/HordeServer.Build/BuildConfig.cs
@@ -99,12 +99,12 @@ namespace HordeServer
 		private readonly Dictionary<StreamId, StreamConfig> _streamLookup = new Dictionary<StreamId, StreamConfig>();
 
 		/// <inheritdoc/>
-		public async Task PostLoadAsync(PluginConfigOptions configOptions)
+		public async Task PostLoadAsync(PluginConfigOptions configOptions, CancellationToken cancellationToken = default)
 		{
 			Acl = configOptions.ParentAcl;
 			Streams = Projects.SelectMany(x => x.Streams).ToList();
 
-			await UpdatePerforceCredentialsAsync(configOptions.Plugins, configOptions.ServiceProvider, configOptions.CancellationToken);
+			await UpdatePerforceCredentialsAsync(configOptions.Plugins, configOptions.ServiceProvider, cancellationToken);
 
 			_projectLookup.Clear();
 			_streamLookup.Clear();
diff --git a/Engine/Source/Programs/Horde/Plugins/Compute/HordeServer.Compute/ComputeConfig.cs b/Engine/Source/Programs/Horde/Plugins/Compute/HordeServer.Compute/ComputeConfig.cs
--- a/Engine/Source/Programs/Horde/Plugins/Compute/HordeServer.Compute/ComputeConfig.cs
+++ b/Engine/Source/Programs/Horde/Plugins/Compute/HordeServer.Compute/ComputeConfig.cs
@@ -69,7 +69,7 @@ namespace HordeServer
 		private readonly Dictionary<PoolId, PoolConfig> _poolLookup = new Dictionary<PoolId, PoolConfig>();
 
 		/// <inheritdoc/>
-		public Task PostLoadAsync(PluginConfigOptions configOptions)
+		public Task PostLoadAsync(PluginConfigOptions configOptions, CancellationToken cancellationToken = default)
 		{
 			Acl = configOptions.ParentAcl;
 			VersionEnum = configOptions.Version;
diff --git a/Engine/Source/Programs/Horde/Plugins/Experimental/HordeServer.Experimental/ExperimentalConfig.cs b/Engine/Source/Programs/Horde/Plugins/Experimental/HordeServer.Experimental/ExperimentalConfig.cs
--- a/Engine/Source/Programs/Horde/Plugins/Experimental/HordeServer.Experimental/ExperimentalConfig.cs
+++ b/Engine/Source/Programs/Horde/Plugins/Experimental/HordeServer.Experimental/ExperimentalConfig.cs
@@ -20,7 +20,7 @@ namespace HordeServer
 		readonly Dictionary<NotificationConfigId, NotificationConfig> _notificationConfigLookup = new Dictionary<NotificationConfigId, NotificationConfig>();
 
 		/// <inheritdoc/>
-		public Task PostLoadAsync(PluginConfigOptions configOptions)
+		public Task PostLoadAsync(PluginConfigOptions configOptions, CancellationToken cancellationToken = default)
 		{
 			_notificationConfigLookup.Clear();
 			foreach (NotificationConfig config in Notifications)
diff --git a/Engine/Source/Programs/Horde/Plugins/Secrets/HordeServer.Secrets/SecretsConfig.cs b/Engine/Source/Programs/Horde/Plugins/Secrets/HordeServer.Secrets/SecretsConfig.cs
--- a/Engine/Source/Programs/Horde/Plugins/Secrets/HordeServer.Secrets/SecretsConfig.cs
+++ b/Engine/Source/Programs/Horde/Plugins/Secrets/HordeServer.Secrets/SecretsConfig.cs
@@ -31,7 +31,7 @@ namespace HordeServer
 		private readonly Dictionary<SecretId, Dictionary<string, SecretProviderConfig>> _providerConfigLookup = new();
 
 		/// <inheritdoc/>
-		public Task PostLoadAsync(PluginConfigOptions configOptions)
+		public Task PostLoadAsync(PluginConfigOptions configOptions, CancellationToken cancellationToken = default)
 		{
 			_secretLookup.Clear();
 			foreach (SecretConfig secret in Secrets)
diff --git a/Engine/Source/Programs/Horde/Plugins/Storage/HordeServer.Storage/StorageConfig.cs b/Engine/Source/Programs/Horde/Plugins/Storage/HordeServer.Storage/StorageConfig.cs
--- a/Engine/Source/Programs/Horde/Plugins/Storage/HordeServer.Storage/StorageConfig.cs
+++ b/Engine/Source/Programs/Horde/Plugins/Storage/HordeServer.Storage/StorageConfig.cs
@@ -74,7 +74,7 @@ namespace HordeServer.Storage
 		/// <summary>
 		/// Called after the config has been read from disk
 		/// </summary>
-		public Task PostLoadAsync(PluginConfigOptions configOptions)
+		public Task PostLoadAsync(PluginConfigOptions configOptions, CancellationToken cancellationToken = default)
 		{
 			Acl = configOptions.ParentAcl;
 
diff --git a/Engine/Source/Programs/Horde/Plugins/Symbols/HordeServer.Symbols/SymbolsConfig.cs b/Engine/Source/Programs/Horde/Plugins/Symbols/HordeServer.Symbols/SymbolsConfig.cs
--- a/Engine/Source/Programs/Horde/Plugins/Symbols/HordeServer.Symbols/SymbolsConfig.cs
+++ b/Engine/Source/Programs/Horde/Plugins/Symbols/HordeServer.Symbols/SymbolsConfig.cs
@@ -22,7 +22,7 @@ namespace HordeServer.Symbols
 		readonly Dictionary<SymbolStoreId, SymbolStoreConfig> _storeLookup = new Dictionary<SymbolStoreId, SymbolStoreConfig>();
 
 		/// <inheritdoc/>
-		public Task PostLoadAsync(PluginConfigOptions configOptions)
+		public Task PostLoadAsync(PluginConfigOptions configOptions, CancellationToken cancellationToken = default)
 		{
 			foreach (SymbolStoreConfig store in Stores)
 			{
diff --git a/Engine/Source/Programs/Horde/Plugins/Tools/HordeServer.Tools/ToolsConfig.cs b/Engine/Source/Programs/Horde/Plugins/Tools/HordeServer.Tools/ToolsConfig.cs
--- a/Engine/Source/Programs/Horde/Plugins/Tools/HordeServer.Tools/ToolsConfig.cs
+++ b/Engine/Source/Programs/Horde/Plugins/Tools/HordeServer.Tools/ToolsConfig.cs
@@ -21,7 +21,7 @@ namespace HordeServer.Tools
 		private readonly Dictionary<ToolId, ToolConfig> _toolLookup = new Dictionary<ToolId, ToolConfig>();
 
 		/// <inheritdoc/>
-		public Task PostLoadAsync(PluginConfigOptions configOptions)
+		public Task PostLoadAsync(PluginConfigOptions configOptions, CancellationToken cancellationToken = default)
 		{
 			_toolLookup.Clear();
 			foreach (ToolConfig tool in Tools)
