Horde: Add ACL actions for read/write access to UBA cache server

diff --git a/Engine/Source/Programs/Horde/Plugins/Build/HordeServer.Build.Tests/Compute/ComputeServiceTest.cs b/Engine/Source/Programs/Horde/Plugins/Build/HordeServer.Build.Tests/Compute/ComputeServiceTest.cs
--- a/Engine/Source/Programs/Horde/Plugins/Build/HordeServer.Build.Tests/Compute/ComputeServiceTest.cs
+++ b/Engine/Source/Programs/Horde/Plugins/Build/HordeServer.Build.Tests/Compute/ComputeServiceTest.cs
@@ -3,6 +3,7 @@
 using System.Diagnostics.CodeAnalysis;
 using System.Diagnostics.Metrics;
 using System.Net;
+using System.Security.Claims;
 using System.Text.RegularExpressions;
 using EpicGames.Horde.Agents;
 using EpicGames.Horde.Agents.Pools;
@@ -11,6 +12,7 @@ using HordeServer.Acls;
 using HordeServer.Agents;
 using HordeServer.Compute;
 using HordeServer.Plugins;
+using HordeServer.Utilities;
 using Microsoft.Extensions.DependencyInjection;
 using Microsoft.Extensions.Logging.Abstractions;
 using Microsoft.Extensions.Options;
@@ -31,7 +33,21 @@ namespace HordeServer.Tests.Compute
 		}
 		
 		private readonly ClusterId _cluster1 = new("cluster1");
-
+		private readonly AclClaimConfig _claimConfig = new (HordeClaimTypes.User, "some.user");
+		private readonly AclConfig _ubaCacheReadAcl;
+		private readonly AclConfig _ubaCacheWriteAcl;
+		private readonly AclConfig _ubaCacheReadWriteAcl;
+		private readonly ClaimsPrincipal _user;
+		private readonly UbaComputeClusterConfig _uccc = new() { CacheEndpoint = "1.2.3.4:3333", CacheHttpEndpoint = "1.2.3.4:4444", CacheHttpCrypto = "aaaaaaaaaaaaaaaaaaaaaaaaaaaaaaaa" };
+		
+		public ComputeServiceTest()
+		{
+			_user = new ClaimsPrincipal(new ClaimsIdentity(new List<Claim> { _claimConfig.ToClaim() }, "TestAuthType"));
+			_ubaCacheReadAcl = new AclConfig { Entries = [new AclEntryConfig(_claimConfig, [ComputeAclAction.UbaCacheRead])] };
+			_ubaCacheWriteAcl = new AclConfig { Entries = [new AclEntryConfig(_claimConfig, [ComputeAclAction.UbaCacheWrite])] };
+			_ubaCacheReadWriteAcl = new AclConfig { Entries = [new AclEntryConfig(_claimConfig, [ComputeAclAction.UbaCacheRead, ComputeAclAction.UbaCacheWrite])] };
+		}
+		
 		#region Resource needs
 		[TestMethod]
 		public async Task ResourceNeedsMetricAsync()
@@ -204,19 +220,19 @@ namespace HordeServer.Tests.Compute
 			IAgent agent4 = await CreateAgentAsync(new PoolId("qux-myComputeId"), properties: props);
 			ClusterId clusterId = new("default");
 
-			AllocateResourceParams arp1 = new(clusterId, ComputeProtocol.Latest, new Requirements { Pool = "foo" }) { RequestId = "req1", RequesterIp = ip, ParentLeaseId = null };
+			AllocateResourceParams arp1 = new(clusterId, _user, ComputeProtocol.Latest, new Requirements { Pool = "foo" }) { RequestId = "req1", RequesterIp = ip, ParentLeaseId = null };
 			ComputeResource resource1 = await ComputeService.TryAllocateResourceAsync(arp1, CancellationToken.None);
 			Assert.AreEqual(agent1.Id, resource1!.AgentId);
 
-			AllocateResourceParams arp2 = new(clusterId, ComputeProtocol.Latest, new Requirements { Pool = "bar-%REQUESTER_NETWORK_ID%" }) { RequestId = "req2", RequesterIp = IPAddress.Parse("15.0.0.1"), ParentLeaseId = null };
+			AllocateResourceParams arp2 = new(clusterId, _user, ComputeProtocol.Latest, new Requirements { Pool = "bar-%REQUESTER_NETWORK_ID%" }) { RequestId = "req2", RequesterIp = IPAddress.Parse("15.0.0.1"), ParentLeaseId = null };
 			ComputeResource resource2 = await ComputeService.TryAllocateResourceAsync(arp2, CancellationToken.None);
 			Assert.AreEqual(agent2.Id, resource2!.AgentId);
 
-			AllocateResourceParams arp3 = new(clusterId, ComputeProtocol.Latest, new Requirements { Pool = "bar-%REQUESTER_NETWORK_ID%" }) { RequestId = "req3", RequesterIp = ip, ParentLeaseId = null };
+			AllocateResourceParams arp3 = new(clusterId, _user, ComputeProtocol.Latest, new Requirements { Pool = "bar-%REQUESTER_NETWORK_ID%" }) { RequestId = "req3", RequesterIp = ip, ParentLeaseId = null };
 			ComputeResource resource3 = await ComputeService.TryAllocateResourceAsync(arp3, CancellationToken.None);
 			Assert.AreEqual(agent3.Id, resource3!.AgentId);
 
-			AllocateResourceParams arp4 = new(clusterId, ComputeProtocol.Latest, new Requirements { Pool = "qux-%REQUESTER_COMPUTE_ID%" }) { RequestId = "req4", RequesterIp = ip, ParentLeaseId = null };
+			AllocateResourceParams arp4 = new(clusterId, _user, ComputeProtocol.Latest, new Requirements { Pool = "qux-%REQUESTER_COMPUTE_ID%" }) { RequestId = "req4", RequesterIp = ip, ParentLeaseId = null };
 			ComputeResource resource4 = await ComputeService.TryAllocateResourceAsync(arp4, CancellationToken.None);
 			Assert.AreEqual(agent4.Id, resource4!.AgentId);
 		}
@@ -226,7 +242,7 @@ namespace HordeServer.Tests.Compute
 		{
 			AllocateResourceParams CreateParams(string ip)
 			{
-				return new AllocateResourceParams(new ClusterId("default"), ComputeProtocol.Latest, new Requirements { Pool = "foo" }) { RequesterIp = IPAddress.Parse(ip) };
+				return new AllocateResourceParams(new ClusterId("default"), _user, ComputeProtocol.Latest, new Requirements { Pool = "foo" }) { RequesterIp = IPAddress.Parse(ip) };
 			}
 
 			IAgent agent = await CreateAgentAsync(new PoolId("foo"), properties: ["ComputeIp=11.0.0.1", "ComputePort=5000"]);
@@ -243,7 +259,7 @@ namespace HordeServer.Tests.Compute
 		{
 			AllocateResourceParams CreateParams(string pool)
 			{
-				return new AllocateResourceParams(new ClusterId("default"), ComputeProtocol.Latest, new Requirements { Pool = pool });
+				return new AllocateResourceParams(new ClusterId("default"), _user, ComputeProtocol.Latest, new Requirements { Pool = pool });
 			}
 
 			IAgent agent1 = await CreateAgentAsync(new PoolId("foo"), properties: ["ComputeIp=11.0.0.1", "ComputePort=5000"]);
@@ -276,7 +292,7 @@ namespace HordeServer.Tests.Compute
 		{
 			AllocateResourceParams CreateParams(string pool)
 			{
-				return new AllocateResourceParams(new ClusterId("default"), ComputeProtocol.Latest, new Requirements { Pool = pool, Exclusive = true });
+				return new AllocateResourceParams(new ClusterId("default"), _user, ComputeProtocol.Latest, new Requirements { Pool = pool, Exclusive = true });
 			}
 			
 			await CreateAgentAsync(new PoolId("foo"), properties: ["ComputeIp=11.0.0.1", "ComputePort=5000"]);
@@ -309,18 +325,18 @@ namespace HordeServer.Tests.Compute
 			IAgent agent2 = await CreateAgentAsync(new PoolId("bar"), properties: props);
 
 			{
-				AllocateResourceParams arpFoo = new(cluster1, ComputeProtocol.Latest, new Requirements { Pool = poolFoo.ToString() });
+				AllocateResourceParams arpFoo = new(cluster1, _user, ComputeProtocol.Latest, new Requirements { Pool = poolFoo.ToString() });
 				Assert.AreEqual(agent1.Id, (await ComputeService.TryAllocateResourceAsync(arpFoo, CancellationToken.None))!.AgentId);
 
-				AllocateResourceParams arpBar = new(cluster1, ComputeProtocol.Latest, new Requirements { Pool = poolBar.ToString() });
+				AllocateResourceParams arpBar = new(cluster1, _user, ComputeProtocol.Latest, new Requirements { Pool = poolBar.ToString() });
 				await Assert.ThrowsExceptionAsync<NoComputeResourcesException>(() => ComputeService.TryAllocateResourceAsync(arpBar, CancellationToken.None));
 			}
 
 			{
-				AllocateResourceParams arpFoo = new(cluster2, ComputeProtocol.Latest, new Requirements { Pool = poolFoo.ToString() });
+				AllocateResourceParams arpFoo = new(cluster2, _user, ComputeProtocol.Latest, new Requirements { Pool = poolFoo.ToString() });
 				await Assert.ThrowsExceptionAsync<NoComputeResourcesException>(() => ComputeService.TryAllocateResourceAsync(arpFoo, CancellationToken.None));
 
-				AllocateResourceParams arpBar = new(cluster2, ComputeProtocol.Latest, new Requirements { Pool = poolBar.ToString() });
+				AllocateResourceParams arpBar = new(cluster2, _user, ComputeProtocol.Latest, new Requirements { Pool = poolBar.ToString() });
 				Assert.AreEqual(agent2.Id, (await ComputeService.TryAllocateResourceAsync(arpBar, CancellationToken.None))!.AgentId);
 			}
 		}
@@ -426,19 +442,46 @@ namespace HordeServer.Tests.Compute
 		[TestMethod]
 		public async Task Connection_UbaCache_CacheServerDefined_Async()
 		{
-			UbaComputeClusterConfig uccc = new() { CacheEndpoint = "1.2.3.4:3333", CacheHttpEndpoint = "1.2.3.4:4444", CacheHttpCrypto = "aaaaaaaaaaaaaaaaaaaaaaaaaaaaaaaa" }; 
-			ComputeResource? cr = await AllocateWithAgentAsync(useUbaCache: true, ubaConfig: uccc);
+			ComputeResource? cr = await AllocateWithAgentAsync(useUbaCache: true, ubaConfig: _uccc, acl: _ubaCacheReadAcl);
 			Assert.AreEqual("1.2.3.4:3333", cr!.Uba!.CacheEndpoint);
 			Assert.IsTrue(Regex.IsMatch(cr.Uba!.CacheSessionKey, "^[0-9a-fA-F]{32}$"));
 		}
 		
+		[TestMethod]
+		public async Task Connection_UbaCache_ReadPermission_Async()
+		{
+			ComputeResource? cr = await AllocateWithAgentAsync(useUbaCache: true, ubaConfig: _uccc, acl: _ubaCacheReadAcl);
+			Assert.IsFalse(cr!.Uba!.WriteAccess);
+		}
+		
+		[TestMethod]
+		public async Task Connection_UbaCache_WritePermission_Async()
+		{
+			ComputeResource? cr = await AllocateWithAgentAsync(useUbaCache: true, ubaConfig: _uccc, acl: _ubaCacheWriteAcl);
+			Assert.IsTrue(cr!.Uba!.WriteAccess);
+		}
+		
+		[TestMethod]
+		public async Task Connection_UbaCache_ReadWritePermission_Async()
+		{
+			ComputeResource? cr = await AllocateWithAgentAsync(useUbaCache: true, ubaConfig: _uccc, acl: _ubaCacheReadWriteAcl);
+			Assert.IsTrue(cr!.Uba!.WriteAccess);
+		}
+		
 		[TestMethod]
 		[Ignore] // Uncomment to manually test allocation with a local UBA cache server
 		public async Task Connection_UbaCache_RealServer_Async()
 		{
-			UbaComputeClusterConfig uccc = new() { CacheEndpoint = "localhost:7009", CacheHttpEndpoint = "localhost:7010", CacheHttpCrypto = "aaaaaaaaaaaaaaaaaaaaaaaaaaaaaaaa" };
-			ComputeResource? cr = await AllocateWithAgentAsync(useUbaCache: true, ubaConfig: uccc);
+			UbaComputeClusterConfig uccc = new()
+			{
+				CacheEndpoint = "localhost:7009",
+				CacheHttpEndpoint = "localhost:7010",
+				CacheHttpCrypto = "6C4AD3E4406C08C9C435DE3B4F400FD4"
+			};
+			using HttpClient httpClient = new ();
+			ComputeResource? cr = await AllocateWithAgentAsync(useUbaCache: true, ubaConfig: uccc, httpClient: httpClient, acl: _ubaCacheReadAcl);
 			Assert.AreEqual(uccc.CacheEndpoint, cr!.Uba!.CacheEndpoint);
+			Console.WriteLine("CacheSessionKey: " + cr.Uba!.CacheSessionKey);
 			Assert.IsTrue(Regex.IsMatch(cr.Uba!.CacheSessionKey, "^[0-9a-fA-F]{32}$"));
 		}
 		
@@ -480,10 +523,18 @@ namespace HordeServer.Tests.Compute
 			bool useUbaCache = false,
 			UbaComputeClusterConfig? ubaConfig = null,
 			ComputeProtocol protocol = ComputeProtocol.Latest,
-			string pool = "myPool")
+			string pool = "myPool",
+			HttpClient? httpClient = null,
+			AclConfig? acl = null)
 		{
-			await using ComputeService cs = await CreateComputeServiceWithAgentAsync(tunnelAddress, pool, new ComputeClusterConfig { Id = _cluster1, Uba = ubaConfig });
-			AllocateResourceParams arp = new(_cluster1, protocol, new Requirements() { Pool = pool })
+			ComputeClusterConfig ccc = new() { Id = _cluster1, Uba = ubaConfig };
+			if (acl != null)
+			{
+				ccc.Acl = acl;
+			}
+			await using ComputeService cs = await CreateComputeServiceWithAgentAsync(tunnelAddress, pool, ccc, httpClient);
+			
+			AllocateResourceParams arp = new(_cluster1, _user, protocol, new Requirements() { Pool = pool })
 			{
 				ConnectionMode = connectionMode,
 				Ports = ports ?? new Dictionary<string, int>(),
@@ -502,7 +553,7 @@ namespace HordeServer.Tests.Compute
 			bool usePublicIp = false,
 			ComputeProtocol protocol = ComputeProtocol.Latest)
 		{
-			AllocateResourceParams arp = new(new ClusterId("default"), protocol, new Requirements() { Pool = pool })
+			AllocateResourceParams arp = new(new ClusterId("default"), _user, protocol, new Requirements() { Pool = pool })
 			{
 				ConnectionMode = connectionMode,
 				Ports = ports ?? new Dictionary<string, int>(),
diff --git a/Engine/Source/Programs/Horde/Plugins/Compute/HordeServer.Compute/Compute/ComputeAclAction.cs b/Engine/Source/Programs/Horde/Plugins/Compute/HordeServer.Compute/Compute/ComputeAclAction.cs
--- a/Engine/Source/Programs/Horde/Plugins/Compute/HordeServer.Compute/Compute/ComputeAclAction.cs
+++ b/Engine/Source/Programs/Horde/Plugins/Compute/HordeServer.Compute/Compute/ComputeAclAction.cs
@@ -2,21 +2,30 @@
 
 using EpicGames.Horde.Acls;
 
-namespace HordeServer.Compute
+namespace HordeServer.Compute;
+
+/// <summary>
+/// ACL actions for compute service
+/// </summary>
+public static class ComputeAclAction
 {
 	/// <summary>
-	/// ACL actions for compute service
+	/// User can add tasks to the compute cluster
 	/// </summary>
-	public static class ComputeAclAction
-	{
-		/// <summary>
-		/// User can add tasks to the compute cluster
-		/// </summary>
-		public static AclAction AddComputeTasks { get; } = new AclAction("AddComputeTasks");
+	public static AclAction AddComputeTasks { get; } = new ("AddComputeTasks");
 
-		/// <summary>
-		/// User can get and list tasks from the compute cluster
-		/// </summary>
-		public static AclAction GetComputeTasks { get; } = new AclAction("GetComputeTasks");
-	}
+	/// <summary>
+	/// User can get and list tasks from the compute cluster
+	/// </summary>
+	public static AclAction GetComputeTasks { get; } = new ("GetComputeTasks");
+	
+	/// <summary>
+	/// User can read from Unreal Build Accelerator cache
+	/// </summary>
+	public static AclAction UbaCacheRead { get; } = new ("UbaCacheRead");
+	
+	/// <summary>
+	/// User can write to Unreal Build Accelerator cache
+	/// </summary>
+	public static AclAction UbaCacheWrite { get; } = new ("UbaCacheWrite");
 }
diff --git a/Engine/Source/Programs/Horde/Plugins/Compute/HordeServer.Compute/Compute/ComputeController.cs b/Engine/Source/Programs/Horde/Plugins/Compute/HordeServer.Compute/Compute/ComputeController.cs
--- a/Engine/Source/Programs/Horde/Plugins/Compute/HordeServer.Compute/Compute/ComputeController.cs
+++ b/Engine/Source/Programs/Horde/Plugins/Compute/HordeServer.Compute/Compute/ComputeController.cs
@@ -109,7 +109,7 @@ namespace HordeServer.Compute
 				return Forbid(ComputeAclAction.AddComputeTasks, clusterId);
 			}
 
-			AllocateResourceParams arp = new(clusterId, (ComputeProtocol)request.Protocol, request.Requirements)
+			AllocateResourceParams arp = new(clusterId, User, (ComputeProtocol)request.Protocol, request.Requirements)
 			{
 				RequestId = request.RequestId,
 				RequesterIp = HttpContext.Connection.RemoteIpAddress,
@@ -155,7 +155,7 @@ namespace HordeServer.Compute
 				Nonce = StringUtils.FormatHexString(computeResource.Task.Nonce.Span),
 				Key = StringUtils.FormatHexString(computeResource.Task.Key.Span),
 				Certificate = StringUtils.FormatHexString(computeResource.Task.Certificate.Span),
-				Uba = computeResource.Uba == null ? null : new UbaConfig(computeResource.Uba.CacheEndpoint, computeResource.Uba.CacheSessionKey),
+				Uba = computeResource.Uba == null ? null : new UbaConfig(computeResource.Uba.CacheEndpoint, computeResource.Uba.CacheSessionKey, computeResource.Uba.WriteAccess),
 				ClusterId = clusterId,
 				AgentId = computeResource.AgentId,
 				AgentVersion = computeResource.AgentVersion,
diff --git a/Engine/Source/Programs/Horde/Plugins/Compute/HordeServer.Compute/Compute/ComputeService.cs b/Engine/Source/Programs/Horde/Plugins/Compute/HordeServer.Compute/Compute/ComputeService.cs
--- a/Engine/Source/Programs/Horde/Plugins/Compute/HordeServer.Compute/Compute/ComputeService.cs
+++ b/Engine/Source/Programs/Horde/Plugins/Compute/HordeServer.Compute/Compute/ComputeService.cs
@@ -2,6 +2,7 @@
 
 using System.Diagnostics.Metrics;
 using System.Net;
+using System.Security.Claims;
 using System.Security.Cryptography;
 using System.Text;
 using System.Text.Json;
@@ -88,6 +89,11 @@ namespace HordeServer.Compute
 		/// <summary>
 		/// User allocating the resource
 		/// </summary>
+		public ClaimsPrincipal User { get; init; }
+
+		/// <summary>
+		/// User ID allocating the resource
+		/// </summary>
 		public UserId? UserId { get; init; }
 
 		/// <summary>
@@ -139,9 +145,10 @@ namespace HordeServer.Compute
 		/// <summary>
 		/// Constructor
 		/// </summary>
-		public AllocateResourceParams(ClusterId clusterId, ComputeProtocol protocol, Requirements? requirements = null)
+		public AllocateResourceParams(ClusterId clusterId, ClaimsPrincipal user, ComputeProtocol protocol, Requirements? requirements = null)
 		{
 			ClusterId = clusterId;
+			User = user;
 			Protocol = protocol;
 			Requirements = requirements ?? new Requirements();
 		}
@@ -423,12 +430,21 @@ namespace HordeServer.Compute
 				{
 					throw new ComputeServiceException($"No UBA cache server defined for compute cluster {arp.ClusterId}") { ShowToUser = true };
 				}
+				
+				bool hasReadAccess = clusterConfig.Authorize(ComputeAclAction.UbaCacheRead, arp.User);
+				bool hasWriteAccess = clusterConfig.Authorize(ComputeAclAction.UbaCacheWrite, arp.User);
+				bool hasAccess = hasReadAccess || hasWriteAccess;
+				if (!hasAccess)
+				{
+					throw new ComputeServiceException($"User does not have UBA cache server access (check {ComputeAclAction.UbaCacheRead} or {ComputeAclAction.UbaCacheWrite})") { ShowToUser = true };
+				}
+				
 				// Allocate UBA cache session key in the background, await the task later
 				// Validate the config above as the method call below is not awaited and exceptions would not immediately be thrown
 				ubaCache = AllocateUbaCacheServerAsync(_httpClient,
 					clusterConfig.Uba.CacheEndpoint,
 					clusterConfig.Uba.CacheHttpEndpoint,
-					clusterConfig.Uba.CacheHttpCrypto, false, 120, cancellationToken);
+					clusterConfig.Uba.CacheHttpCrypto, !hasWriteAccess, 120, cancellationToken);
 			}
 
 			try
@@ -907,7 +923,7 @@ namespace HordeServer.Compute
 				HttpResponseMessage res = await httpClient.GetAsync(uri, timeoutCts.Token);
 				res.EnsureSuccessStatusCode();
 
-				return new UbaComputeResource(ubaCacheEndpoint, sessionKeyHex);
+				return new UbaComputeResource(ubaCacheEndpoint, sessionKeyHex, !readOnly);
 			}
 			catch (TaskCanceledException ex)
 			{
diff --git a/Engine/Source/Programs/Horde/Plugins/Compute/HordeServer.Compute/Compute/ComputeTaskSource.cs b/Engine/Source/Programs/Horde/Plugins/Compute/HordeServer.Compute/Compute/ComputeTaskSource.cs
--- a/Engine/Source/Programs/Horde/Plugins/Compute/HordeServer.Compute/Compute/ComputeTaskSource.cs
+++ b/Engine/Source/Programs/Horde/Plugins/Compute/HordeServer.Compute/Compute/ComputeTaskSource.cs
@@ -16,13 +16,16 @@ using Microsoft.Extensions.Options;
 namespace HordeServer.Compute
 {
 	/// <inheritdoc cref="UbaConfig" />
-	public class UbaComputeResource(string cacheEndpoint, string cacheSessionKey)
+	public class UbaComputeResource(string cacheEndpoint, string cacheSessionKey, bool writeAccess)
 	{
 		/// <inheritdoc cref="UbaConfig.CacheEndpoint" />
 		public string CacheEndpoint { get; } = cacheEndpoint;
 		
 		/// <inheritdoc cref="UbaConfig.CacheSessionKey" />
 		public string CacheSessionKey { get; } = cacheSessionKey;
+		
+		/// <inheritdoc cref="UbaConfig.WriteAccess" />
+		public bool WriteAccess { get; } = writeAccess;
 	}
 	
 	/// <summary>
diff --git a/Engine/Source/Programs/Shared/EpicGames.Horde/Compute/ComputeMessages.cs b/Engine/Source/Programs/Shared/EpicGames.Horde/Compute/ComputeMessages.cs
--- a/Engine/Source/Programs/Shared/EpicGames.Horde/Compute/ComputeMessages.cs
+++ b/Engine/Source/Programs/Shared/EpicGames.Horde/Compute/ComputeMessages.cs
@@ -47,7 +47,7 @@ namespace EpicGames.Horde.Compute
 	/// Configuration settings for Unreal Build Accelerator (UBA) for a compute task invocation
 	/// Contains cache server connectivity and other UBA-related settings.
 	/// </summary>
-	public class UbaConfig(string cacheEndpoint, string cacheSessionKey)
+	public class UbaConfig(string cacheEndpoint, string cacheSessionKey, bool writeAccess)
 	{
 		/// <summary>
 		/// Network endpoint (IP address and port) for connecting to the UBA cache server
@@ -58,6 +58,11 @@ namespace EpicGames.Horde.Compute
 		/// AES encryption key used for secure client-server communication with the UBA cache server
 		/// </summary>
 		public string CacheSessionKey { get; } = cacheSessionKey;
+		
+		/// <summary>
+		/// Session has write access to cache
+		/// </summary>
+		public bool WriteAccess { get; } = writeAccess;
 	}
 
 	/// <summary>
