- Fixing history graph wrong duration conversion
- Adding pointer change when hovering history graph link
- Properly display disabled Artifacts button

#rnx

diff --git a/Engine/Source/Programs/Horde/HordeDashboard/plugins/test-automation/phaseHistoryGraph.ts b/Engine/Source/Programs/Horde/HordeDashboard/plugins/test-automation/phaseHistoryGraph.ts
--- a/Engine/Source/Programs/Horde/HordeDashboard/plugins/test-automation/phaseHistoryGraph.ts
+++ b/Engine/Source/Programs/Horde/HordeDashboard/plugins/test-automation/phaseHistoryGraph.ts
@@ -147,6 +147,7 @@ export class PhaseHistoryGraph {
             .append("g")
             .attr("class", "circleScale")
             .attr("id", i => `circleScale${sessions[i as any].id}`)
+            .style("cursor", "pointer")
             .each((d, i, nodes) => SessionCircle(nodes[i], sessions[d as any], 5, !!this.phaseSessionId && sessions[d as any].id === this.phaseSessionId));
 
         const tooltip = this.tooltip = d3.select(container)
diff --git a/Engine/Source/Programs/Horde/HordeDashboard/plugins/test-automation/testAutomationEvents.tsx b/Engine/Source/Programs/Horde/HordeDashboard/plugins/test-automation/testAutomationEvents.tsx
--- a/Engine/Source/Programs/Horde/HordeDashboard/plugins/test-automation/testAutomationEvents.tsx
+++ b/Engine/Source/Programs/Horde/HordeDashboard/plugins/test-automation/testAutomationEvents.tsx
@@ -204,7 +204,7 @@ export const PhaseEventsPane: React.FC<{phase: TestPhaseStatus, onEventsLoaded?:
                                 <Checkbox styles={styles.defaultCheckbox} label="Info" checked={filterInfo} onChange={(_: any, checked?: boolean) => setFilterInfo(!!checked)} disabled={!events.length}/>
                             </Stack>
                             <Stack horizontalAlign="end">
-                                <DefaultButton disabled={true} style={{ fontSize: 11, padding: 0, height: 18, backgroundColor: styles.defaultButton.backgroundColor, color: 'white' }}
+                                <DefaultButton disabled={true} style={{ fontSize: 11, padding: 0, height: 18, color: 'white' }} title="Unavailable"
                                     onClick={(ev) => {
                                         ev.stopPropagation();
                                     }}>Artifacts</DefaultButton>
diff --git a/Engine/Source/Programs/Horde/HordeDashboard/plugins/test-automation/testHistoryGraph.ts b/Engine/Source/Programs/Horde/HordeDashboard/plugins/test-automation/testHistoryGraph.ts
--- a/Engine/Source/Programs/Horde/HordeDashboard/plugins/test-automation/testHistoryGraph.ts
+++ b/Engine/Source/Programs/Horde/HordeDashboard/plugins/test-automation/testHistoryGraph.ts
@@ -1,5 +1,4 @@
 import * as d3 from "d3";
-import moment from 'moment';
 import { TestDataHandler, TestNameRef, MetadataRef, TestMetaStatus, TestSessionResult } from './testData';
 import dashboard, { StatusColor } from 'horde/backend/Dashboard';
 import { getHumanTime, getShortNiceTime, msecToElapsed } from 'horde/base/utilities/timeUtils';
@@ -194,6 +193,7 @@ export class TestHistoryGraph {
             .append("g")
             .attr("class", "circleScale")
             .attr("id", i => `circleScale${sessions[i as any].id}`)
+            .style("cursor", "pointer")
             .each((d, i, nodes) => SessionConcentricCircle(nodes[i], sessions[d as any], 5));
 
         const tooltip = this.tooltip = d3.select(container)
@@ -268,7 +268,7 @@ export class TestHistoryGraph {
                 let desc = "";
                 if (!!cmeta) desc += `# ${cmeta} <br/>`;
                 desc += `Commit ${closest.commitId} <br/>`;
-                desc += `Duration ${msecToElapsed(moment.duration(closest.duration).asMilliseconds(), true, true)} <br/>`;
+                desc += `Duration ${msecToElapsed(closest.duration * 1000, true, true)} <br/>`;
                 desc += `${date} <br/>`;
                 if (closest.phasesSucceededCount + closest.phasesFailedCount + closest.phasesUnspecifiedCount > 1)
                     desc += renderToString(SessionValues(closest, {fontSize: 9}));
