Horde: Add REST API endpoint for requesting a UBA cache server session

diff --git a/Engine/Source/Programs/Horde/Plugins/Compute/HordeServer.Compute/Compute/ComputeController.cs b/Engine/Source/Programs/Horde/Plugins/Compute/HordeServer.Compute/Compute/ComputeController.cs
--- a/Engine/Source/Programs/Horde/Plugins/Compute/HordeServer.Compute/Compute/ComputeController.cs
+++ b/Engine/Source/Programs/Horde/Plugins/Compute/HordeServer.Compute/Compute/ComputeController.cs
@@ -172,6 +172,25 @@ namespace HordeServer.Compute
 			return response;
 		}
 
+		/// <summary>
+		/// Request a UBA cache server session
+		/// </summary>
+		/// <param name="clusterId">ID of the compute cluster</param>
+		/// <param name="cancellationToken">Cancellation token for the operation</param>
+		[HttpPost]
+		[Authorize]
+		[Route("/api/v2/compute/{clusterId}/uba-cache")]
+		public async Task<ActionResult<UbaConfig>> AllocateUbaCacheServerAsync(ClusterId clusterId, CancellationToken cancellationToken)
+		{
+			if (!_computeConfig.Value.TryGetComputeCluster(clusterId, out ComputeClusterConfig? clusterConfig))
+			{
+				return NotFound(clusterId);
+			}
+			
+			UbaComputeResource ucr = await _computeService.AllocateUbaCacheServerAsync(User, clusterConfig, 120, cancellationToken);
+			return new UbaConfig(ucr.CacheEndpoint, ucr.CacheSessionKey, ucr.WriteAccess);
+		}
+
 		/// <summary>
 		/// Get current resource needs for active sessions
 		/// </summary>
diff --git a/Engine/Source/Programs/Horde/Plugins/Compute/HordeServer.Compute/Compute/ComputeService.cs b/Engine/Source/Programs/Horde/Plugins/Compute/HordeServer.Compute/Compute/ComputeService.cs
--- a/Engine/Source/Programs/Horde/Plugins/Compute/HordeServer.Compute/Compute/ComputeService.cs
+++ b/Engine/Source/Programs/Horde/Plugins/Compute/HordeServer.Compute/Compute/ComputeService.cs
@@ -426,25 +426,8 @@ namespace HordeServer.Compute
 			Task<UbaComputeResource>? ubaCache = null;
 			if (arp.UseUbaCache ?? false)
 			{
-				if (clusterConfig.Uba?.CacheEndpoint == null || clusterConfig.Uba.CacheHttpEndpoint == null || clusterConfig.Uba.CacheHttpCrypto == null)
-				{
-					throw new ComputeServiceException($"No UBA cache server defined for compute cluster {arp.ClusterId}") { ShowToUser = true };
-				}
-				
-				bool hasReadAccess = clusterConfig.Authorize(ComputeAclAction.UbaCacheRead, arp.User);
-				bool hasWriteAccess = clusterConfig.Authorize(ComputeAclAction.UbaCacheWrite, arp.User);
-				bool hasAccess = hasReadAccess || hasWriteAccess;
-				if (!hasAccess)
-				{
-					throw new ComputeServiceException($"User does not have UBA cache server access (check {ComputeAclAction.UbaCacheRead} or {ComputeAclAction.UbaCacheWrite})") { ShowToUser = true };
-				}
-				
 				// Allocate UBA cache session key in the background, await the task later
-				// Validate the config above as the method call below is not awaited and exceptions would not immediately be thrown
-				ubaCache = AllocateUbaCacheServerAsync(_httpClient,
-					clusterConfig.Uba.CacheEndpoint,
-					clusterConfig.Uba.CacheHttpEndpoint,
-					clusterConfig.Uba.CacheHttpCrypto, !hasWriteAccess, 120, cancellationToken);
+				ubaCache = AllocateUbaCacheServerAsync(arp.User, clusterConfig, 120, cancellationToken);
 			}
 
 			try
@@ -531,7 +514,25 @@ namespace HordeServer.Compute
 				throw new ComputeServiceException("Unable to allocate compute resource: " + are.Message, are) { ShowToUser = true };
 			}
 		}
-
+		
+		internal static bool VerifyUbaCache(ClaimsPrincipal user, ComputeClusterConfig clusterConfig)
+		{
+			if (clusterConfig.Uba?.CacheEndpoint == null || clusterConfig.Uba.CacheHttpEndpoint == null || clusterConfig.Uba.CacheHttpCrypto == null)
+			{
+				throw new ComputeServiceException($"No UBA cache server defined for compute cluster {clusterConfig.Id}") { ShowToUser = true };
+			}
+			
+			bool hasReadAccess = clusterConfig.Authorize(ComputeAclAction.UbaCacheRead, user);
+			bool hasWriteAccess = clusterConfig.Authorize(ComputeAclAction.UbaCacheWrite, user);
+			bool hasAccess = hasReadAccess || hasWriteAccess;
+			if (!hasAccess)
+			{
+				throw new ComputeServiceException($"User does not have UBA cache server access (check {ComputeAclAction.UbaCacheRead} or {ComputeAclAction.UbaCacheWrite})") { ShowToUser = true };
+			}
+			
+			return hasWriteAccess;
+		}
+		
 		private string? ResolvePoolId(string? poolId, IPAddress? ipAddress)
 		{
 			if (poolId == null)
@@ -898,10 +899,34 @@ namespace HordeServer.Compute
 			return relayIps[0];
 		}
 		
-		internal static async Task<UbaComputeResource> AllocateUbaCacheServerAsync(HttpClient httpClient,
-			string ubaCacheEndpoint, string ubaCacheHttpEndpoint, string ubaCacheHttpCrypto,
-			bool readOnly, int ttl, CancellationToken cancellationToken)
+		/// <summary>
+		/// Allocate a UBA cache server session
+		/// </summary>
+		/// <param name="user">Requesting user</param>
+		/// <param name="clusterConfig">Compute cluster</param>
+		/// <param name="ttl">Time-to-live in seconds for the session</param>
+		/// <param name="cancellationToken">Cancellation token</param>
+		/// <param name="httpClient">Optional HTTP client override for tests</param>
+		/// <returns>Settings for a UBA cache server session</returns>
+		public async Task<UbaComputeResource> AllocateUbaCacheServerAsync(
+			ClaimsPrincipal user, ComputeClusterConfig clusterConfig, int ttl, CancellationToken cancellationToken, HttpClient? httpClient = null)
 		{
+			string? cacheEndpoint = clusterConfig.Uba?.CacheEndpoint;
+			string? httpEndpoint = clusterConfig.Uba?.CacheHttpEndpoint;
+			string? httpCrypto = clusterConfig.Uba?.CacheHttpCrypto;
+			if (cacheEndpoint == null || httpEndpoint == null || httpCrypto == null)
+			{
+				throw new ComputeServiceException($"No UBA cache server defined for compute cluster {clusterConfig.Id}") { ShowToUser = true };
+			}
+			
+			bool hasReadAccess = clusterConfig.Authorize(ComputeAclAction.UbaCacheRead, user);
+			bool hasWriteAccess = clusterConfig.Authorize(ComputeAclAction.UbaCacheWrite, user);
+			bool hasAccess = hasReadAccess || hasWriteAccess;
+			if (!hasAccess)
+			{
+				throw new ComputeServiceException($"User does not have UBA cache server access (check {ComputeAclAction.UbaCacheRead} or {ComputeAclAction.UbaCacheWrite})") { ShowToUser = true };
+			}
+			
 			try
 			{
 				using CancellationTokenSource timeoutCts = CancellationTokenSource.CreateLinkedTokenSource(cancellationToken);
@@ -910,7 +935,7 @@ namespace HordeServer.Compute
 				// Generate a session encryption key that we activate in the UBA cache server
 				byte[] sessionKey = RandomNumberGenerator.GetBytes(16);
 				string sessionKeyHex = Convert.ToHexString(sessionKey);
-				byte[] cryptoKey = Convert.FromHexString(ubaCacheHttpCrypto);
+				byte[] cryptoKey = Convert.FromHexString(httpCrypto);
 				if (cryptoKey.Length != 16)
 				{
 					throw new Exception($"Invalid crypto key length: {cryptoKey.Length}");
@@ -918,12 +943,12 @@ namespace HordeServer.Compute
 				
 				using UbaCrypto crypto = new(cryptoKey);
 				string encryptedSessionKeyHex = Convert.ToHexString(crypto.Encrypt(sessionKey));
-				string readOnlyStr = readOnly ? "true" : "false";
-				Uri uri = new ($"http://{ubaCacheHttpEndpoint}/addsession?crypto={encryptedSessionKeyHex},readonly={readOnlyStr},expiration={ttl}");
-				HttpResponseMessage res = await httpClient.GetAsync(uri, timeoutCts.Token);
+				string readOnlyStr = hasWriteAccess ? "false" : "true";
+				Uri uri = new ($"http://{httpEndpoint}/addsession?crypto={encryptedSessionKeyHex},readonly={readOnlyStr},expiration={ttl}");
+				HttpResponseMessage res = await (httpClient ?? _httpClient).GetAsync(uri, timeoutCts.Token);
 				res.EnsureSuccessStatusCode();
 
-				return new UbaComputeResource(ubaCacheEndpoint, sessionKeyHex, !readOnly);
+				return new UbaComputeResource(cacheEndpoint, sessionKeyHex, hasWriteAccess);
 			}
 			catch (TaskCanceledException ex)
 			{
diff --git a/Engine/Source/Programs/Shared/EpicGames.Horde/Compute/ComputeMessages.cs b/Engine/Source/Programs/Shared/EpicGames.Horde/Compute/ComputeMessages.cs
--- a/Engine/Source/Programs/Shared/EpicGames.Horde/Compute/ComputeMessages.cs
+++ b/Engine/Source/Programs/Shared/EpicGames.Horde/Compute/ComputeMessages.cs
@@ -43,6 +43,13 @@ namespace EpicGames.Horde.Compute
 		}
 	}
 	
+	/// <summary>
+	/// Request for allocating a UBA cache server/session
+	/// </summary>
+	public class AllocateUbaCacheServerRequest
+	{
+	}
+	
 	/// <summary>
 	/// Configuration settings for Unreal Build Accelerator (UBA) for a compute task invocation
 	/// Contains cache server connectivity and other UBA-related settings.
