Horde: Updating schema documentation

#rnx

diff --git a/Engine/Source/Programs/Horde/Docs/Config/Schema/AclActions.md b/Engine/Source/Programs/Horde/Docs/Config/Schema/AclActions.md
--- a/Engine/Source/Programs/Horde/Docs/Config/Schema/AclActions.md
+++ b/Engine/Source/Programs/Horde/Docs/Config/Schema/AclActions.md
@@ -193,6 +193,8 @@
 | ---- | ----------- |
 | `AddComputeTasks` | User can add tasks to the compute cluster |
 | `GetComputeTasks` | User can get and list tasks from the compute cluster |
+| `UbaCacheRead` | User can read from Unreal Build Accelerator cache |
+| `UbaCacheWrite` | User can write to Unreal Build Accelerator cache |
 
 ## Devices
 
@@ -349,6 +351,8 @@
 | ---- | ----------- |
 | `AddComputeTasks` |  |
 | `GetComputeTasks` |  |
+| `UbaCacheRead` |  |
+| `UbaCacheWrite` |  |
 
 ## Ddc
 
@@ -518,6 +522,8 @@
 | ---- | ----------- |
 | `AddComputeTasks` |  |
 | `GetComputeTasks` |  |
+| `UbaCacheRead` |  |
+| `UbaCacheWrite` |  |
 
 ## Ddc
 
@@ -687,6 +693,8 @@
 | ---- | ----------- |
 | `AddComputeTasks` |  |
 | `GetComputeTasks` |  |
+| `UbaCacheRead` |  |
+| `UbaCacheWrite` |  |
 
 ## Ddc
 
@@ -862,6 +870,8 @@
 | ---- | ----------- |
 | `AddComputeTasks` |  |
 | `GetComputeTasks` |  |
+| `UbaCacheRead` |  |
+| `UbaCacheWrite` |  |
 
 ## Ddc
 
@@ -1050,6 +1060,8 @@
 | ---- | ----------- |
 | `AddComputeTasks` |  |
 | `GetComputeTasks` |  |
+| `UbaCacheRead` |  |
+| `UbaCacheWrite` |  |
 
 ## Ddc
 
@@ -1244,6 +1256,8 @@
 | ---- | ----------- |
 | `AddComputeTasks` |  |
 | `GetComputeTasks` |  |
+| `UbaCacheRead` |  |
+| `UbaCacheWrite` |  |
 
 ## Ddc
 
@@ -1445,6 +1459,8 @@
 | ---- | ----------- |
 | `AddComputeTasks` |  |
 | `GetComputeTasks` |  |
+| `UbaCacheRead` |  |
+| `UbaCacheWrite` |  |
 
 ## Ddc
 
