Added an example of SubmitTool metric and dashboard category to Horde Analytics

diff --git a/Engine/Source/Programs/Horde/HordeServer/Defaults/default-metrics.telemetry.json b/Engine/Source/Programs/Horde/HordeServer/Defaults/default-metrics.telemetry.json
--- a/Engine/Source/Programs/Horde/HordeServer/Defaults/default-metrics.telemetry.json
+++ b/Engine/Source/Programs/Horde/HordeServer/Defaults/default-metrics.telemetry.json
@@ -844,6 +844,13 @@
 			"function": "Sum",
 			"interval": "00:30:00"
 		},
+		{
+			"id": "submit-tool-session-count",
+			"filter": "$[?(@.Payload.EventName == 'SubmitTool.StandAlone.Start')]",
+			"groupBy": "$.Payload.Stream,$.Payload.Version",
+			"function": "Count",
+			"interval": "00:30:00"
+		},
 		{
 			"id": "debug-probe",
 			"filter": "$[?(@.Payload.SchemaVersion >= 0 )]",
diff --git a/Engine/Source/Programs/Horde/HordeServer/Defaults/default-views.telemetry.json b/Engine/Source/Programs/Horde/HordeServer/Defaults/default-views.telemetry.json
--- a/Engine/Source/Programs/Horde/HordeServer/Defaults/default-views.telemetry.json
+++ b/Engine/Source/Programs/Horde/HordeServer/Defaults/default-views.telemetry.json
@@ -769,6 +769,21 @@
 						}
 					]
 				},
+				{
+					"name": "Submit Tool",
+					"charts": [
+						{
+							"name": "Session Count",
+							"display": "Value",
+							"min": 0,
+							"metrics": [
+								{
+									"id": "submit-tool-session-count"
+								}
+							]
+						}
+					]
+				},
 				{
 					"name": "Debug",
 					"charts": [
