Changed Horde Analytics dashbaord category name from SubmitTool to Submit

diff --git a/Engine/Source/Programs/Horde/HordeServer/Defaults/default-views.telemetry.json b/Engine/Source/Programs/Horde/HordeServer/Defaults/default-views.telemetry.json
--- a/Engine/Source/Programs/Horde/HordeServer/Defaults/default-views.telemetry.json
+++ b/Engine/Source/Programs/Horde/HordeServer/Defaults/default-views.telemetry.json
@@ -770,7 +770,7 @@
 					]
 				},
 				{
-					"name": "Submit Tool",
+					"name": "Submit",
 					"charts": [
 						{
 							"name": "Session Count",
