Horde Dashboard: portable search url inconsistency

diff --git a/Engine/Source/Programs/Horde/HordeDashboard/src/components/JobSearchSimple.tsx b/Engine/Source/Programs/Horde/HordeDashboard/src/components/JobSearchSimple.tsx
--- a/Engine/Source/Programs/Horde/HordeDashboard/src/components/JobSearchSimple.tsx
+++ b/Engine/Source/Programs/Horde/HordeDashboard/src/components/JobSearchSimple.tsx
@@ -75,8 +75,6 @@ export const JobSearchSimpleModal: React.FC<{ streamIdIn: string, changeListIn?:
 
    // Update the search params in the url for sharing and bookmarking queires
    const updateSearch = (state: SearchState) => {
-      console.log(state);
-
       // Preserve current tab
       const tab = searchParams.get("tab")!;
       setSearchParams("");
@@ -98,10 +96,13 @@ export const JobSearchSimpleModal: React.FC<{ streamIdIn: string, changeListIn?:
       setSearchParams(search);
    }
 
-   const queryJobs = async (state: SearchState) => {      
+   const queryJobs = async (state: SearchState) => {
       updateSearch(state);
       setSearchResults(prev => ({ ...prev, querying: true }));
 
+      const c_templates = (templateData?.templates ?? []).filter(t => state.templates.has(t.id));
+      const c_statuses = statuses.filter(s => state.statuses.has(s));
+
       let jobs: JobData[] = [];
 
       try {
@@ -163,8 +164,8 @@ export const JobSearchSimpleModal: React.FC<{ streamIdIn: string, changeListIn?:
 
             let stepTemplates: string[] = [];
 
-            if (ctemplates?.length) {
-               stepTemplates.push(...ctemplates.map(t => t.id));
+            if (c_templates?.length) {
+               stepTemplates.push(...c_templates.map(t => t.id));
             } else {
                stepTemplates = templates.filter(t => !!t.id).map(t => t.id);
             }
@@ -203,7 +204,7 @@ export const JobSearchSimpleModal: React.FC<{ streamIdIn: string, changeListIn?:
                filter: "id,labels,defaultLabel,name,change,createTime,streamId,preflightChange,startedByUserInfo,arguments,abortedByUserInfo,state",
                minChange: minCL,
                maxChange: maxCL,
-               t: (!state.global && ctemplates.length && !state.containsStep) ? ctemplates.map(t => t.id) : undefined,
+               t: (!state.global && c_templates.length && !state.containsStep) ? c_templates.map(t => t.id) : undefined,
                minCreateTime: minCreateTime?.toUTCString(),
                maxCreateTime: maxCreateTime?.toUTCString(),
                includePreflight: state.preflights || !!state.preflightOnly,
@@ -214,9 +215,9 @@ export const JobSearchSimpleModal: React.FC<{ streamIdIn: string, changeListIn?:
 
             jobs = await backend.getJobs(query);
 
-            if (cstatuses.length && cstatuses.length != allStatuses.size) {
+            if (c_statuses.length && c_statuses.length != allStatuses.size) {
                jobs = jobs.filter(j => {
-                  return (cstatuses.includes('Cancelled') && j.abortedByUserInfo?.id || cstatuses.includes(j.state) && !j.abortedByUserInfo?.id)
+                  return (c_statuses.includes('Cancelled') && j.abortedByUserInfo?.id || c_statuses.includes(j.state) && !j.abortedByUserInfo?.id)
                })
             }
 
@@ -234,16 +235,6 @@ export const JobSearchSimpleModal: React.FC<{ streamIdIn: string, changeListIn?:
          console.error(reason);
       } finally {
 
-         /*
-         if (column.name === "Stream") {
-             const stream = projectStore.streamById(item.job.streamId);
-             if (!stream) {
-                 return null;
-             }
-             return <Stack horizontal verticalFill={true} verticalAlign="center" tokens={{ childrenGap: 0, padding: 0 }} style={{ overflow: "hidden" }} ><Text>{stream.fullname}</Text></Stack>;
-         }
-         */
-
          const streamIds = new Set<string>();
          const idToFullname = new Map<string, string>();
 
@@ -278,6 +269,7 @@ export const JobSearchSimpleModal: React.FC<{ streamIdIn: string, changeListIn?:
          })
 
          setSearchResults({items, groups, querying: false });
+         setSearchState(state);
       }
 
    }
@@ -295,10 +287,10 @@ export const JobSearchSimpleModal: React.FC<{ streamIdIn: string, changeListIn?:
 
    // Sets state from search params and autoloads jobs
    useEffect(() => {
-      if(!searchParams.has("search") || !isFirstRender) return;
+      if(!searchParams.has("search") || !isFirstRender || !templateData) return;
       
       setIsFirstRender(false);
-
+      
       // Decompress URL data
       const stringyState = JSON.parse(LZString.decompressFromEncodedURIComponent(searchParams.get("search")!));
       
@@ -309,16 +301,15 @@ export const JobSearchSimpleModal: React.FC<{ streamIdIn: string, changeListIn?:
          templates: new Set(stringyState.templates),
          statuses: new Set(stringyState.statuses)
       };
-
+      
       if(state.global){
          streamChooserId++;
          state.templates = new Set();
          state.streamIdOverride = undefined;
       }
-
+      
       queryJobs(state)
-      setSearchState(state);
-   }, [])
+   }, [templateData])
 
    if (!stream) {
       console.error("unable to get stream");
