Horde: Change IPluginConfig.PostLoad to non-async

diff --git a/Engine/Source/Programs/Horde/HordeServer.Shared/Plugins/IPluginConfig.cs b/Engine/Source/Programs/Horde/HordeServer.Shared/Plugins/IPluginConfig.cs
--- a/Engine/Source/Programs/Horde/HordeServer.Shared/Plugins/IPluginConfig.cs
+++ b/Engine/Source/Programs/Horde/HordeServer.Shared/Plugins/IPluginConfig.cs
@@ -13,12 +13,11 @@ namespace HordeServer.Plugins
 		/// Called to fixup a plugin's configuration after deserialization
 		/// </summary>
 		/// <param name="configOptions">Options for configuring the plugin</param>
-		/// <param name="cancellationToken">Cancellation token for the operation</param>
-		Task PostLoadAsync(PluginConfigOptions configOptions, CancellationToken cancellationToken = default);
+		void PostLoad(PluginConfigOptions configOptions);
 	}
 
 	/// <summary>
-	/// Options passed to <see cref="IPluginConfig.PostLoadAsync(PluginConfigOptions,CancellationToken)"/>
+	/// Options passed to <see cref="IPluginConfig.PostLoad(PluginConfigOptions)"/>
 	/// </summary>
 	public record class PluginConfigOptions(ConfigVersion Version, IEnumerable<IPluginConfig> Plugins, AclConfig ParentAcl, IServiceProvider? ServiceProvider = null);
 
@@ -28,9 +27,7 @@ namespace HordeServer.Plugins
 	public sealed class EmptyPluginConfig : IPluginConfig
 	{
 		/// <inheritdoc/>
-		public Task PostLoadAsync(PluginConfigOptions configOptions, CancellationToken cancellationToken = default)
-		{
-			return Task.CompletedTask;
-		}
+		public void PostLoad(PluginConfigOptions configOptions)
+		{ }
 	}
 }
diff --git a/Engine/Source/Programs/Horde/HordeServer.Tests.Shared/ServerTestSetup.cs b/Engine/Source/Programs/Horde/HordeServer.Tests.Shared/ServerTestSetup.cs
--- a/Engine/Source/Programs/Horde/HordeServer.Tests.Shared/ServerTestSetup.cs
+++ b/Engine/Source/Programs/Horde/HordeServer.Tests.Shared/ServerTestSetup.cs
@@ -75,7 +75,7 @@ namespace HordeServer.Tests
 		protected async Task SetConfigAsync(GlobalConfig globalConfig)
 		{
 			await ConfigService.WaitForInitialConfigAsync();
-			await globalConfig.PostLoadAsync(ServerSettings, _pluginCollection.LoadedPlugins, DefaultAclModifiers, ServiceProvider);
+			globalConfig.PostLoad(ServerSettings, _pluginCollection.LoadedPlugins, DefaultAclModifiers, ServiceProvider);
 			ConfigService.OverrideConfig(globalConfig);
 		}
 
@@ -84,7 +84,7 @@ namespace HordeServer.Tests
 			await ConfigService.WaitForInitialConfigAsync();
 			GlobalConfig globalConfig = GlobalConfig.CurrentValue;
 			action(globalConfig);
-			await globalConfig.PostLoadAsync(ServerSettings, _pluginCollection.LoadedPlugins, DefaultAclModifiers, ServiceProvider);
+			globalConfig.PostLoad(ServerSettings, _pluginCollection.LoadedPlugins, DefaultAclModifiers, ServiceProvider);
 			ConfigService.OverrideConfig(globalConfig);
 		}
 
diff --git a/Engine/Source/Programs/Horde/HordeServer.Tests/Configuration/ConfigTests.cs b/Engine/Source/Programs/Horde/HordeServer.Tests/Configuration/ConfigTests.cs
--- a/Engine/Source/Programs/Horde/HordeServer.Tests/Configuration/ConfigTests.cs
+++ b/Engine/Source/Programs/Horde/HordeServer.Tests/Configuration/ConfigTests.cs
@@ -173,7 +173,7 @@ namespace HordeServer.Tests.Configuration
 		}
 
 		[TestMethod]
-		public async Task WorkspaceConfigAsync()
+		public void WorkspaceConfig()
 		{
 			Dictionary<string, WorkspaceConfig> inputWorkspaces = new();
 			inputWorkspaces["base"] = new WorkspaceConfig { Identifier = "base", Cluster = "myCluster", MinScratchSpace = 111 };
@@ -186,7 +186,7 @@ namespace HordeServer.Tests.Configuration
 			GlobalConfig gc = new();
 			gc.Plugins.AddComputeConfig(new ComputeConfig());
 			gc.Plugins.AddBuildConfig(buildConfig);
-			await gc.PostLoadAsync(new ServerSettings(), Array.Empty<ILoadedPlugin>(), Array.Empty<IDefaultAclModifier>());
+			gc.PostLoad(new ServerSettings(), Array.Empty<ILoadedPlugin>(), Array.Empty<IDefaultAclModifier>());
 
 			Dictionary<string,WorkspaceConfig> workspaces = gc.Plugins.GetBuildConfig().Projects[0].Streams[0].WorkspaceTypes;
 			Assert.AreEqual(3, workspaces.Count);
@@ -196,7 +196,7 @@ namespace HordeServer.Tests.Configuration
 		}
 
 		[TestMethod]
-		public async Task WorkspaceInheritFromProjectConfigAsync()
+		public void WorkspaceInheritFromProjectConfig()
 		{
 			List<string> autoSdkViews = ["foo", "bar"];
 
@@ -221,7 +221,7 @@ namespace HordeServer.Tests.Configuration
 			GlobalConfig gc = new();
 			gc.Plugins.AddComputeConfig(new ComputeConfig());
 			gc.Plugins.AddBuildConfig(buildConfig);
-			await gc.PostLoadAsync(new ServerSettings(), Array.Empty<ILoadedPlugin>(), Array.Empty<IDefaultAclModifier>());
+			gc.PostLoad(new ServerSettings(), Array.Empty<ILoadedPlugin>(), Array.Empty<IDefaultAclModifier>());
 
 			Dictionary<string, WorkspaceConfig> workspaces = gc.Plugins.GetBuildConfig().Projects[0].Streams[0].WorkspaceTypes;
 			Assert.AreEqual(4, workspaces.Count);
diff --git a/Engine/Source/Programs/Horde/HordeServer.Tests/Server/DowntimeServiceTests.cs b/Engine/Source/Programs/Horde/HordeServer.Tests/Server/DowntimeServiceTests.cs
--- a/Engine/Source/Programs/Horde/HordeServer.Tests/Server/DowntimeServiceTests.cs
+++ b/Engine/Source/Programs/Horde/HordeServer.Tests/Server/DowntimeServiceTests.cs
@@ -15,12 +15,12 @@ namespace HordeServer.Tests.Server
 	[TestClass]
 	public class DowntimeServiceTest
 	{
-		private static async Task<ScheduledDowntime> GetDowntimeFromJsonAsync(string json)
+		private static ScheduledDowntime GetDowntimeFromJson(string json)
 		{
 			ScheduledDowntime scheduledDowntime = JsonSerializer.Deserialize<ScheduledDowntime>(json, JsonUtils.DefaultSerializerOptions)!;
 			GlobalConfig config = new();
 			config.Downtime.Add(scheduledDowntime);
-			await config.PostLoadAsync(new ServerSettings(), [], []);
+			config.PostLoad(new ServerSettings(), [], []);
 			return scheduledDowntime;
 		}
 
@@ -47,7 +47,7 @@ namespace HordeServer.Tests.Server
 				FinishTime = clock.UtcNow + TimeSpan.FromSeconds(90)
 			};
 			config.Downtime.Add(downtime);
-			await config.PostLoadAsync(new ServerSettings(), [], []);
+			config.PostLoad(new ServerSettings(), [], []);
 			TestOptionsMonitor<GlobalConfig> optionsMonitor = new(config);
 			await using DowntimeService downtimeService = new(clock, optionsMonitor, NullLogger<DowntimeService>.Instance);
 			Assert.IsFalse(downtimeService.IsDowntimeActive);
@@ -57,9 +57,9 @@ namespace HordeServer.Tests.Server
 		}
 
 		[TestMethod]
-		public async Task IsNotActiveAsync()
+		public void IsNotActive()
 		{
-			ScheduledDowntime downtime = await GetDowntimeFromJsonAsync("""
+			ScheduledDowntime downtime = GetDowntimeFromJson("""
 				{
 					"startTime": "2020-09-06T03:00:00+00:00",
 					"finishTime": "2020-09-06T03:45:00+00:00"
@@ -69,9 +69,9 @@ namespace HordeServer.Tests.Server
 		}
 
 		[TestMethod]
-		public async Task IsActiveAsync()
+		public void IsActive()
 		{
-			ScheduledDowntime downtime = await GetDowntimeFromJsonAsync("""
+			ScheduledDowntime downtime = GetDowntimeFromJson("""
 				{
 					"startTime": "2020-09-06T03:00:00+00:00",
 					"finishTime": "2020-09-06T03:45:00+00:00"
@@ -81,9 +81,9 @@ namespace HordeServer.Tests.Server
 		}
 
 		[TestMethod]
-		public async Task IsNotActiveDailyAsync()
+		public void IsNotActiveDaily()
 		{
-			ScheduledDowntime downtime = await GetDowntimeFromJsonAsync("""
+			ScheduledDowntime downtime = GetDowntimeFromJson("""
 				{
 					"startTime": "2020-09-06T03:00:00+00:00",
 					"finishTime": "2020-09-06T03:45:00+00:00",
@@ -94,9 +94,9 @@ namespace HordeServer.Tests.Server
 		}
 
 		[TestMethod]
-		public async Task IsActiveDailyAsync()
+		public void IsActiveDaily()
 		{
-			ScheduledDowntime downtime = await GetDowntimeFromJsonAsync("""
+			ScheduledDowntime downtime = GetDowntimeFromJson("""
 				{
 					"startTime": "2020-09-06T03:00:00+00:00",
 					"finishTime": "2020-09-06T03:45:00+00:00",
@@ -107,9 +107,9 @@ namespace HordeServer.Tests.Server
 		}
 
 		[TestMethod]
-		public async Task IsNotActiveWeeklyAsync()
+		public void IsNotActiveWeekly()
 		{
-			ScheduledDowntime downtime = await GetDowntimeFromJsonAsync("""
+			ScheduledDowntime downtime = GetDowntimeFromJson("""
 				{
 					"startTime": "2020-09-06T03:00:00+00:00",
 					"finishTime": "2020-09-06T03:45:00+00:00",
@@ -120,9 +120,9 @@ namespace HordeServer.Tests.Server
 		}
 
 		[TestMethod]
-		public async Task IsActiveWeeklyAsync()
+		public void IsActiveWeekly()
 		{
-			ScheduledDowntime downtime = await GetDowntimeFromJsonAsync("""
+			ScheduledDowntime downtime = GetDowntimeFromJson("""
 				{
 					"startTime": "2020-09-06T03:00:00+00:00",
 					"finishTime": "2020-09-06T03:45:00+00:00",
@@ -133,9 +133,9 @@ namespace HordeServer.Tests.Server
 		}
 
 		[TestMethod]
-		public async Task FinishTimeFromDurationAsync()
+		public void FinishTimeFromDuration()
 		{
-			ScheduledDowntime downtime = await GetDowntimeFromJsonAsync("""
+			ScheduledDowntime downtime = GetDowntimeFromJson("""
 				{
 					"startTime": "2020-09-06T03:00:00+00:00",
 					"duration": "0:45"
@@ -145,9 +145,9 @@ namespace HordeServer.Tests.Server
 		}
 
 		[TestMethod]
-		public async Task DurationFromFinishTimeAsync()
+		public void DurationFromFinishTime()
 		{
-			ScheduledDowntime downtime = await GetDowntimeFromJsonAsync("""
+			ScheduledDowntime downtime = GetDowntimeFromJson("""
 				{
 					"startTime": "2020-09-06T03:00:00+00:00",
 					"finishTime": "2020-09-06T03:30:00+00:00",
@@ -157,9 +157,9 @@ namespace HordeServer.Tests.Server
 		}
 
 		[TestMethod]
-		public async Task DurationHigherPrecedenceAsync()
+		public void DurationHigherPrecedence()
 		{
-			ScheduledDowntime downtime = await GetDowntimeFromJsonAsync("""
+			ScheduledDowntime downtime = GetDowntimeFromJson("""
 				{
 					"startTime": "2020-09-06T03:00:00+00:00",
 					"finishTime": "2020-09-06T03:30:00+00:00",
@@ -170,9 +170,9 @@ namespace HordeServer.Tests.Server
 		}
 
 		[TestMethod]
-		public async Task StartTimeAndDurationMinuteAbbreviationAsync()
+		public void StartTimeAndDurationMinuteAbbreviation()
 		{
-			ScheduledDowntime downtime = await GetDowntimeFromJsonAsync("""
+			ScheduledDowntime downtime = GetDowntimeFromJson("""
 				{
 					"startTime": "Sun Sep 6, 2020, 3AM",
 					"duration": "0:45"
@@ -183,23 +183,21 @@ namespace HordeServer.Tests.Server
 		}
 		
 		[TestMethod]
-		public async Task WrongDayOfWeekNameErrorAsync()
+		[ExpectedException(typeof(FormatException))]
+		public void WrongDayOfWeekNameError()
 		{
-			await Assert.ThrowsExceptionAsync<FormatException>(async () =>
-			{
-				await GetDowntimeFromJsonAsync("""
-					{
-						"startTime": "Mon Sep 6, 2020, 3AM",
-						"finishTime": "Sun Sep 6, 2020, 3:30AM",
-					}
-				""");
-			});
+			GetDowntimeFromJson("""
+				{
+					"startTime": "Mon Sep 6, 2020, 3AM",
+					"finishTime": "Sun Sep 6, 2020, 3:30AM",
+				}
+			""");
 		}
 
 		[TestMethod]
-		public async Task StartTimeAndDurationHourAbbreviationAsync()
+		public void StartTimeAndDurationHourAbbreviation()
 		{
-			ScheduledDowntime downtime = await GetDowntimeFromJsonAsync("""
+			ScheduledDowntime downtime = GetDowntimeFromJson("""
 				{
 					"startTime": "Sun Sep 6, 2020, 3AM",
 					"duration": "1:30"
@@ -210,9 +208,9 @@ namespace HordeServer.Tests.Server
 		}
 
 		[TestMethod]
-		public async Task StartTimeAndFinishTimeAbbreviationAsync()
+		public void StartTimeAndFinishTimeAbbreviation()
 		{
-			ScheduledDowntime downtime = await GetDowntimeFromJsonAsync("""
+			ScheduledDowntime downtime = GetDowntimeFromJson("""
 				{
 					"startTime": "Dec 1, 2024, 3AM",
 					"finishTime": "Sunday December 1 2024, 3:15AM"
@@ -223,9 +221,9 @@ namespace HordeServer.Tests.Server
 		}
 
 		[TestMethod]
-		public async Task WithTimeZoneAsync()
+		public void WithTimeZone()
 		{
-			ScheduledDowntime downtime = await GetDowntimeFromJsonAsync("""
+			ScheduledDowntime downtime = GetDowntimeFromJson("""
 				{
 					"startTime": "Jan 1, 2020, 3AM",
 					"duration": "1:30",
@@ -237,9 +235,9 @@ namespace HordeServer.Tests.Server
 		}
 
 		[TestMethod]
-		public async Task WithTimeZoneAndDaylightSavingAsync()
+		public void WithTimeZoneAndDaylightSaving()
 		{
-			ScheduledDowntime downtime = await GetDowntimeFromJsonAsync("""
+			ScheduledDowntime downtime = GetDowntimeFromJson("""
 				{
 					"startTime": "Sun Sep 6, 2020, 3AM",
 					"duration": "1:30",
@@ -251,9 +249,9 @@ namespace HordeServer.Tests.Server
 		}
 
 		[TestMethod]
-		public async Task WithTimeZoneIanaAsync()
+		public void WithTimeZoneIana()
 		{
-			ScheduledDowntime downtime = await GetDowntimeFromJsonAsync("""
+			ScheduledDowntime downtime = GetDowntimeFromJson("""
 				{
 					"startTime": "Jan 1, 2020, 3AM",
 					"duration": "1:30",
@@ -265,9 +263,9 @@ namespace HordeServer.Tests.Server
 		}
 
 		[TestMethod]
-		public async Task WithTimeZoneIanaAndDaylightSavingAsync()
+		public void WithTimeZoneIanaAndDaylightSaving()
 		{
-			ScheduledDowntime downtime = await GetDowntimeFromJsonAsync("""
+			ScheduledDowntime downtime = GetDowntimeFromJson("""
 				{
 					"startTime": "Sun Sep 6, 2020, 3AM",
 					"duration": "1:30",
@@ -279,31 +277,29 @@ namespace HordeServer.Tests.Server
 		}
 
 		[TestMethod]
-		public async Task WrongTimeZoneNameAsync()
+		[ExpectedException(typeof(TimeZoneNotFoundException))]
+		public void WrongTimeZoneName()
 		{
-			await Assert.ThrowsExceptionAsync<TimeZoneNotFoundException>(async () =>
-			{
-				await GetDowntimeFromJsonAsync("""
-					{
-						"startTime": "Jan 1, 2020, 3AM",
-						"timezone": "ET"
-					}
-				""");
-			});
+			GetDowntimeFromJson("""
+				{
+					"startTime": "Jan 1, 2020, 3AM",
+					"timezone": "ET"
+				}
+			""");
 		}
 
 		[TestMethod]
-		public async Task DowntimeToJsonAsync()
+		public void DowntimeToJson()
 		{
 			// -4 hours is Eastern Daylight Time
-			string json = JsonSerializer.Serialize(await GetDowntimeFromJsonAsync("""
+			string json = JsonSerializer.Serialize(GetDowntimeFromJson("""
 					{
 						"startTime": "Sun Sep 6, 2020, 3AM -4",
 						"duration": "0:30"
 					}
 				"""),
 				JsonUtils.DefaultSerializerOptions);
-			ScheduledDowntime downtime = await GetDowntimeFromJsonAsync(json);
+			ScheduledDowntime downtime = GetDowntimeFromJson(json);
 			Assert.AreEqual(DateTimeOffsetParse("2020-09-06T03:00:00-04:00"), downtime.StartTime);
 			Assert.AreEqual(DateTimeOffsetParse("2020-09-06T03:30:00-04:00"), downtime.FinishTime);
 			Assert.AreEqual(TimeSpan.FromMinutes(30), downtime.Duration);
@@ -312,9 +308,9 @@ namespace HordeServer.Tests.Server
 		}
 
 		[TestMethod]
-		public async Task DowntimeWithTimeZoneToJsonAsync()
+		public void DowntimeWithTimeZoneToJson()
 		{
-			string json = JsonSerializer.Serialize(await GetDowntimeFromJsonAsync("""
+			string json = JsonSerializer.Serialize(GetDowntimeFromJson("""
 					{
 						"startTime": "Jan 1, 2025, 3AM",
 						"duration": "0:30",
@@ -322,7 +318,7 @@ namespace HordeServer.Tests.Server
 					}
 				"""),
 				JsonUtils.DefaultSerializerOptions);
-			ScheduledDowntime downtime = await GetDowntimeFromJsonAsync(json);
+			ScheduledDowntime downtime = GetDowntimeFromJson(json);
 			Assert.AreEqual(DateTimeOffsetParse("2025-01-01T03:00:00-05:00"), downtime.StartTime);
 			Assert.AreEqual(DateTimeOffsetParse("2025-01-01T03:30:00-05:00"), downtime.FinishTime);
 			Assert.AreEqual(TimeSpan.FromMinutes(30), downtime.Duration);
diff --git a/Engine/Source/Programs/Horde/HordeServer/Configuration/ConfigService.cs b/Engine/Source/Programs/Horde/HordeServer/Configuration/ConfigService.cs
--- a/Engine/Source/Programs/Horde/HordeServer/Configuration/ConfigService.cs
+++ b/Engine/Source/Programs/Horde/HordeServer/Configuration/ConfigService.cs
@@ -285,7 +285,7 @@ namespace HordeServer.Configuration
 			try
 			{
 				GlobalConfig globalConfig = await ReadGlobalConfigAsync(context, cancellationToken);
-				await globalConfig.PostLoadAsync(_serverSettings, _pluginCollection.LoadedPlugins, _defaultAclModifiers, _serviceProvider, cancellationToken);
+				globalConfig.PostLoad(_serverSettings, _pluginCollection.LoadedPlugins, _defaultAclModifiers, _serviceProvider);
 
 				foreach (OverrideConfigFile file in overrideFiles.Values)
 				{
@@ -395,7 +395,7 @@ namespace HordeServer.Configuration
 					data = SerializeSnapshot(snapshot);
 				}
 
-				ConfigState state = new ConfigState(IoHash.Compute(data.Span), await CreateGlobalConfigAsync(data));
+				ConfigState state = new ConfigState(IoHash.Compute(data.Span), CreateGlobalConfig(data));
 				await _health.UpdateAsync(HealthStatus.Healthy);
 
 				return state;
@@ -611,7 +611,7 @@ namespace HordeServer.Configuration
 				}
 
 				// Execute a PostLoad before returning so we can validate that everything is valid
-				await globalConfig.PostLoadAsync(_serverSettings, _pluginCollection.LoadedPlugins, _defaultAclModifiers, _serviceProvider, cancellationToken);
+				globalConfig.PostLoad(_serverSettings, _pluginCollection.LoadedPlugins, _defaultAclModifiers, _serviceProvider);
 
 				return snapshot;
 			}
@@ -844,13 +844,13 @@ namespace HordeServer.Configuration
 			IoHash hash = IoHash.Compute(data.Span);
 			if (hash != state.Hash && state.Hash != IoHash.Zero) // Don't replace any explicit config override
 			{
-				GlobalConfig globalConfig = await CreateGlobalConfigAsync(data);
+				GlobalConfig globalConfig = CreateGlobalConfig(data);
 				_logger.LogInformation("Updating config state from Redis (hash: {OldHash} -> {NewHash}, new size: {Size})", state.Hash, hash, data.Length);
 				Set(hash, globalConfig);
 			}
 		}
 
-		async Task<GlobalConfig> CreateGlobalConfigAsync(ReadOnlyMemory<byte> data)
+		GlobalConfig CreateGlobalConfig(ReadOnlyMemory<byte> data)
 		{
 			// Decompress the snapshot object
 			ConfigSnapshot snapshot = DeserializeSnapshot(data);
@@ -860,7 +860,7 @@ namespace HordeServer.Configuration
 			globalConfig.Revision = IoHash.Compute(data.Span).ToString();
 
 			// Run the postload callbacks on all the config objects
-			await globalConfig.PostLoadAsync(_serverSettings, _pluginCollection.LoadedPlugins, _defaultAclModifiers, _serviceProvider, CancellationToken.None);
+			globalConfig.PostLoad(_serverSettings, _pluginCollection.LoadedPlugins, _defaultAclModifiers, _serviceProvider);
 			return globalConfig;
 		}
 	}
diff --git a/Engine/Source/Programs/Horde/HordeServer/Server/GlobalConfig.cs b/Engine/Source/Programs/Horde/HordeServer/Server/GlobalConfig.cs
--- a/Engine/Source/Programs/Horde/HordeServer/Server/GlobalConfig.cs
+++ b/Engine/Source/Programs/Horde/HordeServer/Server/GlobalConfig.cs
@@ -6,8 +6,6 @@ using System.Diagnostics.CodeAnalysis;
 using System.Linq;
 using System.Security.Claims;
 using System.Text.Json.Serialization;
-using System.Threading;
-using System.Threading.Tasks;
 using EpicGames.Core;
 using EpicGames.Horde.Acls;
 using HordeServer.Accounts;
@@ -106,7 +104,7 @@ namespace HordeServer.Server
 		/// <summary>
 		/// Called after the config file has been read
 		/// </summary>
-		public async Task PostLoadAsync(ServerSettings serverSettings, IReadOnlyList<ILoadedPlugin> loadedPlugins, IEnumerable<IDefaultAclModifier> defaultAclModifiers, IServiceProvider? serviceProvider = null, CancellationToken cancellationToken = default)
+		public void PostLoad(ServerSettings serverSettings, IReadOnlyList<ILoadedPlugin> loadedPlugins, IEnumerable<IDefaultAclModifier> defaultAclModifiers, IServiceProvider? serviceProvider = null)
 		{
 			ServerSettings = serverSettings;
 			AclAction[] aclActions = AclConfig.GetActions(
@@ -159,7 +157,7 @@ namespace HordeServer.Server
 			PluginConfigOptions pluginConfigOptions = new PluginConfigOptions(VersionEnum, Plugins.Values, Acl, serviceProvider);
 			foreach (IPluginConfig pluginConfig in pluginConfigs)
 			{
-				await pluginConfig.PostLoadAsync(pluginConfigOptions, cancellationToken);
+				pluginConfig.PostLoad(pluginConfigOptions);
 			}
 
 			foreach (ScheduledDowntime downtime in Downtime)
diff --git a/Engine/Source/Programs/Horde/Plugins/Analytics/HordeServer.Analytics/AnalyticsConfig.cs b/Engine/Source/Programs/Horde/Plugins/Analytics/HordeServer.Analytics/AnalyticsConfig.cs
--- a/Engine/Source/Programs/Horde/Plugins/Analytics/HordeServer.Analytics/AnalyticsConfig.cs
+++ b/Engine/Source/Programs/Horde/Plugins/Analytics/HordeServer.Analytics/AnalyticsConfig.cs
@@ -26,7 +26,7 @@ namespace HordeServer
 		private readonly Dictionary<TelemetryStoreId, TelemetryStoreConfig> _telemetryStoreLookup = new Dictionary<TelemetryStoreId, TelemetryStoreConfig>();
 
 		/// <inheritdoc/>
-		public Task PostLoadAsync(PluginConfigOptions configOptions, CancellationToken cancellationToken = default)
+		public void PostLoad(PluginConfigOptions configOptions)
 		{
 			_parentAcl = configOptions.ParentAcl;
 
@@ -36,8 +36,6 @@ namespace HordeServer
 				_telemetryStoreLookup.Add(telemetryStore.Id, telemetryStore);
 				telemetryStore.PostLoad(configOptions.ParentAcl);
 			}
-
-			return Task.CompletedTask;
 		}
 
 		/// <inheritdoc cref="AclConfig.Authorize(AclAction, ClaimsPrincipal)"/>
diff --git a/Engine/Source/Programs/Horde/Plugins/Build/HordeServer.Build.Tests/Compute/ComputeServiceTest.cs b/Engine/Source/Programs/Horde/Plugins/Build/HordeServer.Build.Tests/Compute/ComputeServiceTest.cs
--- a/Engine/Source/Programs/Horde/Plugins/Build/HordeServer.Build.Tests/Compute/ComputeServiceTest.cs
+++ b/Engine/Source/Programs/Horde/Plugins/Build/HordeServer.Build.Tests/Compute/ComputeServiceTest.cs
@@ -318,7 +318,7 @@ namespace HordeServer.Tests.Compute
 				new() { Id = cluster1, Condition = $"pool == 'fOO'" }, // Test case-insensitivity of pool ID matching
 				new() { Id = cluster2, Condition = $"pool == '{poolBar.ToString()}'" }
 			};
-			await GlobalConfig.CurrentValue.PostLoadAsync(new ServerSettings(), Array.Empty<ILoadedPlugin>(), Array.Empty<IDefaultAclModifier>());
+			GlobalConfig.CurrentValue.PostLoad(new ServerSettings(), Array.Empty<ILoadedPlugin>(), Array.Empty<IDefaultAclModifier>());
 
 			List<string> props = ["ComputeIp=11.0.0.1", "ComputePort=5000"];
 			IAgent agent1 = await CreateAgentAsync(new PoolId("foo"), properties: props);
@@ -354,7 +354,7 @@ namespace HordeServer.Tests.Compute
 		}
 
 		[TestMethod]
-		public async Task Cluster_FindBest_Async()
+		public void Cluster_FindBest()
 		{
 			ComputeConfig computeConfig = new()
 			{
@@ -365,7 +365,7 @@ namespace HordeServer.Tests.Compute
 			};
 
 			PluginConfigOptions configOptions = new PluginConfigOptions(ConfigVersion.Latest, Enumerable.Empty<IPluginConfig>(), new HordeServer.Acls.AclConfig());
-			await computeConfig.PostLoadAsync(configOptions);
+			computeConfig.PostLoad(configOptions);
 
 			Assert.AreEqual(new ClusterId("compute1"), ComputeService.FindBestComputeClusterId(computeConfig, IPAddress.Parse("11.0.0.1")));
 			Assert.AreEqual(new ClusterId("compute2"), ComputeService.FindBestComputeClusterId(computeConfig, IPAddress.Parse("12.0.1.1")));
@@ -375,7 +375,7 @@ namespace HordeServer.Tests.Compute
 			{
 				Networks = [new NetworkConfig { CidrBlock = "0.0.0.0/0", Id = "catchAll", ComputeId = "catchAll" }]
 			};
-			await globalConfigCatchAll.PostLoadAsync(configOptions);
+			globalConfigCatchAll.PostLoad(configOptions);
 			Assert.AreEqual(new ClusterId("catchAll"), ComputeService.FindBestComputeClusterId(globalConfigCatchAll, IPAddress.Parse("123.123.123.123")));
 		}
 		#endregion Cluster
@@ -509,7 +509,7 @@ namespace HordeServer.Tests.Compute
 				NullLogger<ComputeService>.Instance);
 			await CreateAgentAsync(new PoolId(pool), properties: ["ComputeIp=11.0.0.1", "ComputePort=5000"]);
 			computeConfig.CurrentValue.Clusters = [ccc];
-			await computeConfig.CurrentValue.PostLoadAsync(new PluginConfigOptions(ConfigVersion.Latest, Enumerable.Empty<IPluginConfig>(), new AclConfig()));
+			computeConfig.CurrentValue.PostLoad(new PluginConfigOptions(ConfigVersion.Latest, Enumerable.Empty<IPluginConfig>(), new AclConfig()));
 			AgentRelayService.SetRandomSeed(1);
 			return cs;
 		}
diff --git a/Engine/Source/Programs/Horde/Plugins/Build/HordeServer.Build.Tests/Fixture.cs b/Engine/Source/Programs/Horde/Plugins/Build/HordeServer.Build.Tests/Fixture.cs
--- a/Engine/Source/Programs/Horde/Plugins/Build/HordeServer.Build.Tests/Fixture.cs
+++ b/Engine/Source/Programs/Horde/Plugins/Build/HordeServer.Build.Tests/Fixture.cs
@@ -79,7 +79,7 @@ namespace HordeServer.Tests
 
 			GlobalConfig globalConfig = new GlobalConfig();
 			globalConfig.Plugins.AddBuildConfig(buildConfig);
-			await globalConfig.PostLoadAsync(serverSettings, pluginCollection.LoadedPlugins, Array.Empty<IDefaultAclModifier>());
+			globalConfig.PostLoad(serverSettings, pluginCollection.LoadedPlugins, Array.Empty<IDefaultAclModifier>());
 			configService.OverrideConfig(globalConfig);
 
 			StreamId = streamId;
diff --git a/Engine/Source/Programs/Horde/Plugins/Build/HordeServer.Build/BuildConfig.cs b/Engine/Source/Programs/Horde/Plugins/Build/HordeServer.Build/BuildConfig.cs
--- a/Engine/Source/Programs/Horde/Plugins/Build/HordeServer.Build/BuildConfig.cs
+++ b/Engine/Source/Programs/Horde/Plugins/Build/HordeServer.Build/BuildConfig.cs
@@ -99,12 +99,12 @@ namespace HordeServer
 		private readonly Dictionary<StreamId, StreamConfig> _streamLookup = new Dictionary<StreamId, StreamConfig>();
 
 		/// <inheritdoc/>
-		public async Task PostLoadAsync(PluginConfigOptions configOptions, CancellationToken cancellationToken = default)
+		public void PostLoad(PluginConfigOptions configOptions)
 		{
 			Acl = configOptions.ParentAcl;
 			Streams = Projects.SelectMany(x => x.Streams).ToList();
 
-			await UpdatePerforceCredentialsAsync(configOptions.Plugins, configOptions.ServiceProvider, cancellationToken);
+			UpdatePerforceCredentials(configOptions.Plugins, configOptions.ServiceProvider);
 
 			_projectLookup.Clear();
 			_streamLookup.Clear();
@@ -127,11 +127,11 @@ namespace HordeServer
 			UpdateWorkspacesForPools(configOptions.Plugins);
 		}
 
-		private static async Task<string?> ResolveSecretAsync(string? value, ISecretCollection secretCollection, SecretsConfig secretsConfig, CancellationToken cancellationToken)
+		private static string? ResolveSecret(string? value, ISecretCollection secretCollection, SecretsConfig secretsConfig)
 		{
 			if (!String.IsNullOrEmpty(value))
 			{
-				ISecretProperty? newValue = await secretCollection.ResolveAsync(value, secretsConfig, cancellationToken);
+				ISecretProperty? newValue = secretCollection.ResolveAsync(value, secretsConfig).Result;
 				if (newValue != null)
 				{
 					return newValue.Value;
@@ -140,7 +140,7 @@ namespace HordeServer
 			return value;
 		}
 
-		private async Task UpdatePerforceCredentialsAsync(IEnumerable<IPluginConfig> plugins, IServiceProvider? serviceProvider, CancellationToken cancellationToken)
+		private void UpdatePerforceCredentials(IEnumerable<IPluginConfig> plugins, IServiceProvider? serviceProvider)
 		{
 			if (!ResolvePerforceCredentialsWithSecrets)
 			{
@@ -164,8 +164,8 @@ namespace HordeServer
 			{
 				foreach (PerforceCredentials credentials in perforceCluster.Credentials)
 				{
-					credentials.Password = await ResolveSecretAsync(credentials.Password, secretCollection, secretsConfig, cancellationToken);
-					credentials.Ticket = await ResolveSecretAsync(credentials.Ticket, secretCollection, secretsConfig, cancellationToken);
+					credentials.Password = ResolveSecret(credentials.Password, secretCollection, secretsConfig);
+					credentials.Ticket = ResolveSecret(credentials.Ticket, secretCollection, secretsConfig);
 				}
 			}
 		}
diff --git a/Engine/Source/Programs/Horde/Plugins/Compute/HordeServer.Compute/ComputeConfig.cs b/Engine/Source/Programs/Horde/Plugins/Compute/HordeServer.Compute/ComputeConfig.cs
--- a/Engine/Source/Programs/Horde/Plugins/Compute/HordeServer.Compute/ComputeConfig.cs
+++ b/Engine/Source/Programs/Horde/Plugins/Compute/HordeServer.Compute/ComputeConfig.cs
@@ -69,7 +69,7 @@ namespace HordeServer
 		private readonly Dictionary<PoolId, PoolConfig> _poolLookup = new Dictionary<PoolId, PoolConfig>();
 
 		/// <inheritdoc/>
-		public Task PostLoadAsync(PluginConfigOptions configOptions, CancellationToken cancellationToken = default)
+		public void PostLoad(PluginConfigOptions configOptions)
 		{
 			Acl = configOptions.ParentAcl;
 			VersionEnum = configOptions.Version;
@@ -87,8 +87,6 @@ namespace HordeServer
 				_poolLookup.Add(pool.Id, pool);
 			}
 			ConfigObject.MergeDefaults<string, PoolConfig>(Pools.Select(x => (x.Id.ToString(), x.Base?.ToString(), x)));
-
-			return Task.CompletedTask;
 		}
 
 		/// <summary>
diff --git a/Engine/Source/Programs/Horde/Plugins/Ddc/HordeServer.Ddc.Tests/Ddc/Functional/ContentId/CompressedBlobTests.cs b/Engine/Source/Programs/Horde/Plugins/Ddc/HordeServer.Ddc.Tests/Ddc/Functional/ContentId/CompressedBlobTests.cs
--- a/Engine/Source/Programs/Horde/Plugins/Ddc/HordeServer.Ddc.Tests/Ddc/Functional/ContentId/CompressedBlobTests.cs
+++ b/Engine/Source/Programs/Horde/Plugins/Ddc/HordeServer.Ddc.Tests/Ddc/Functional/ContentId/CompressedBlobTests.cs
@@ -32,7 +32,7 @@ namespace HordeServer.Tests.Ddc.FunctionalTests.CompressedBlobs
 		}
 
 		[TestInitialize]
-		public async Task AddConfigsAsync()
+		public void AddConfigs()
 		{
 			ServerSettings serverSettings = ServiceProvider.GetRequiredService<IOptions<ServerSettings>>().Value;
 
@@ -42,7 +42,7 @@ namespace HordeServer.Tests.Ddc.FunctionalTests.CompressedBlobs
 
 			GlobalConfig globalConfig = ServiceProvider.GetRequiredService<IOptions<GlobalConfig>>().Value;
 			globalConfig.Plugins.AddStorageConfig(storageConfig);
-			await globalConfig.PostLoadAsync(serverSettings, Array.Empty<ILoadedPlugin>(), Array.Empty<IDefaultAclModifier>());
+			globalConfig.PostLoad(serverSettings, Array.Empty<ILoadedPlugin>(), Array.Empty<IDefaultAclModifier>());
 
 			ConfigService configService = ServiceProvider.GetRequiredService<ConfigService>();
 			configService.OverrideConfig(globalConfig);
diff --git a/Engine/Source/Programs/Horde/Plugins/Ddc/HordeServer.Ddc.Tests/Ddc/Functional/Objects/ReferencesTests.cs b/Engine/Source/Programs/Horde/Plugins/Ddc/HordeServer.Ddc.Tests/Ddc/Functional/Objects/ReferencesTests.cs
--- a/Engine/Source/Programs/Horde/Plugins/Ddc/HordeServer.Ddc.Tests/Ddc/Functional/Objects/ReferencesTests.cs
+++ b/Engine/Source/Programs/Horde/Plugins/Ddc/HordeServer.Ddc.Tests/Ddc/Functional/Objects/ReferencesTests.cs
@@ -51,7 +51,7 @@ namespace HordeServer.Tests.Ddc.FunctionalTests.References
 		}
 
 		[TestInitialize]
-		public async Task AddConfigsAsync()
+		public void AddConfigs()
 		{
 			ServerSettings serverSettings = ServiceProvider.GetRequiredService<IOptions<ServerSettings>>().Value;
 
@@ -61,7 +61,7 @@ namespace HordeServer.Tests.Ddc.FunctionalTests.References
 
 			GlobalConfig globalConfig = ServiceProvider.GetRequiredService<IOptions<GlobalConfig>>().Value;
 			globalConfig.Plugins.AddStorageConfig(storageConfig);
-			await globalConfig.PostLoadAsync(serverSettings, Array.Empty<ILoadedPlugin>(), Array.Empty<IDefaultAclModifier>());
+			globalConfig.PostLoad(serverSettings, Array.Empty<ILoadedPlugin>(), Array.Empty<IDefaultAclModifier>());
 
 			ConfigService configService = ServiceProvider.GetRequiredService<ConfigService>();
 			configService.OverrideConfig(globalConfig);
diff --git a/Engine/Source/Programs/Horde/Plugins/Experimental/HordeServer.Experimental/ExperimentalConfig.cs b/Engine/Source/Programs/Horde/Plugins/Experimental/HordeServer.Experimental/ExperimentalConfig.cs
--- a/Engine/Source/Programs/Horde/Plugins/Experimental/HordeServer.Experimental/ExperimentalConfig.cs
+++ b/Engine/Source/Programs/Horde/Plugins/Experimental/HordeServer.Experimental/ExperimentalConfig.cs
@@ -20,7 +20,7 @@ namespace HordeServer
 		readonly Dictionary<NotificationConfigId, NotificationConfig> _notificationConfigLookup = new Dictionary<NotificationConfigId, NotificationConfig>();
 
 		/// <inheritdoc/>
-		public Task PostLoadAsync(PluginConfigOptions configOptions, CancellationToken cancellationToken = default)
+		public void PostLoad(PluginConfigOptions configOptions)
 		{
 			_notificationConfigLookup.Clear();
 			foreach (NotificationConfig config in Notifications)
@@ -28,8 +28,6 @@ namespace HordeServer
 				config.PostLoad();
 				_notificationConfigLookup.Add(config.Id, config);
 			}
-
-			return Task.CompletedTask;
 		}
 
 		/// <summary>
diff --git a/Engine/Source/Programs/Horde/Plugins/Secrets/HordeServer.Secrets.Tests/Secrets/SecretProviderTests.cs b/Engine/Source/Programs/Horde/Plugins/Secrets/HordeServer.Secrets.Tests/Secrets/SecretProviderTests.cs
--- a/Engine/Source/Programs/Horde/Plugins/Secrets/HordeServer.Secrets.Tests/Secrets/SecretProviderTests.cs
+++ b/Engine/Source/Programs/Horde/Plugins/Secrets/HordeServer.Secrets.Tests/Secrets/SecretProviderTests.cs
@@ -45,7 +45,7 @@ namespace HordeServer.Tests.Secrets
 				return new HttpClient();
 			}
 		}
-		private static async Task<SecretCollection> CreateSecretCollectionAsync()
+		private static SecretCollection CreateSecretCollection()
 		{
 			SecretCollectionInternal secretCollectionInternal = new([
 				new InMemoryProvider(),
@@ -186,14 +186,14 @@ namespace HordeServer.Tests.Secrets
 			SecretCollection secretCollection = new(secretCollectionInternal, new OptionsSnapshotStub<SecretsConfig>(secretsConfig));
 			GlobalConfig config = new();
 			config.Plugins.AddSecretsConfig(secretsConfig);
-			await config.PostLoadAsync(new ServerSettings(), [], []);
+			config.PostLoad(new ServerSettings(), [], []);
 			return secretCollection;
 		}
 
 		[TestMethod]
 		public async Task GetSecretAsync()
 		{
-			SecretCollection secretCollection = await CreateSecretCollectionAsync(); 
+			SecretCollection secretCollection = CreateSecretCollection(); 
 			ISecret? secret = await secretCollection.GetAsync(new SecretId("Secret"), CancellationToken.None);
 			Assert.AreEqual(secret?.Data["username"], "j.doe");
 			Assert.AreEqual(secret?.Data["password"], "password123");
@@ -202,7 +202,7 @@ namespace HordeServer.Tests.Secrets
 		[TestMethod]
 		public async Task GetSecretFromProviderConfigAsync()
 		{
-			SecretCollection secretCollection = await CreateSecretCollectionAsync(); 
+			SecretCollection secretCollection = CreateSecretCollection(); 
 			ISecret? secret = await secretCollection.GetAsync(new SecretId("SecretUsingProviderConfig"), CancellationToken.None);
 			Assert.AreEqual(secret?.Data["username"], "user01");
 			Assert.AreEqual(secret?.Data["password"], "password01");
@@ -211,7 +211,7 @@ namespace HordeServer.Tests.Secrets
 		[TestMethod]
 		public async Task GetEmptySecretAsync()
 		{
-			SecretCollection secretCollection = await CreateSecretCollectionAsync();
+			SecretCollection secretCollection = CreateSecretCollection();
 			ISecret? secret = await secretCollection.GetAsync(new SecretId("EmptySecret"), CancellationToken.None);
 			Assert.AreEqual(secret?.Data.Count, 0);
 		}
@@ -220,7 +220,7 @@ namespace HordeServer.Tests.Secrets
 		[TestMethod]
 		public async Task GetSecretFromVaultAsync()
 		{
-			SecretCollection secretCollection = await CreateSecretCollectionAsync(); 
+			SecretCollection secretCollection = CreateSecretCollection(); 
 			ISecret? secret = await secretCollection.GetAsync(new SecretId("SecretUsingHcpVault"), CancellationToken.None);
 			Assert.AreEqual(secret?.Data["username"], "User01");
 			Assert.AreEqual(secret?.Data["password"], "Password123");
@@ -230,7 +230,7 @@ namespace HordeServer.Tests.Secrets
 		[TestMethod]
 		public async Task GetMultipleSecretFromVaultAsync()
 		{
-			SecretCollection secretCollection = await CreateSecretCollectionAsync(); 
+			SecretCollection secretCollection = CreateSecretCollection(); 
 			ISecret? secret = await secretCollection.GetAsync(new SecretId("MultipleSecretUsingHcpVault"), CancellationToken.None);
 			Assert.AreEqual(secret?.Data["one"], "1");
 			Assert.AreEqual(secret?.Data["two"], "2");
@@ -241,7 +241,7 @@ namespace HordeServer.Tests.Secrets
 		[TestMethod]
 		public async Task SecretNotFoundFromVaultAsync()
 		{
-			SecretCollection secretCollection = await CreateSecretCollectionAsync(); 
+			SecretCollection secretCollection = CreateSecretCollection(); 
 			InvalidOperationException ex = await Assert.ThrowsExceptionAsync<InvalidOperationException>(() =>
 				secretCollection.GetAsync(new SecretId("MissingSecretUsingHcpVault"), CancellationToken.None)
 			);
@@ -257,7 +257,7 @@ namespace HordeServer.Tests.Secrets
 			// 2. set the environment variable AWS_PROFILE to a profile that can call AssumeRole for the required ARN
 			// 3. restart IDE to pick up the environment change
 			// 4. update the SecretProviderConfig options for the SecretId used in this test 
-			SecretCollection secretCollection = await CreateSecretCollectionAsync(); 
+			SecretCollection secretCollection = CreateSecretCollection(); 
 			ISecret? secret = await secretCollection.GetAsync(new SecretId("SecretUsingHcpVaultAwsAuth"), CancellationToken.None);
 			Assert.AreEqual(secret?.Data["test"], "test");
 		}
@@ -283,7 +283,7 @@ namespace HordeServer.Tests.Secrets
 			SecretCollection secretCollection = new(secretCollectionInternal, new OptionsSnapshotStub<SecretsConfig>(secretsConfig));
 			GlobalConfig config = new();
 			config.Plugins.AddSecretsConfig(secretsConfig);
-			await config.PostLoadAsync(new ServerSettings(), [], []);
+			config.PostLoad(new ServerSettings(), [], []);
 			ISecret? secret = await secretCollection.GetAsync(new SecretId("basic-secret"), CancellationToken.None);
 			Assert.AreEqual(secret?.Data["username"], "j.doe");
 			Assert.AreEqual(secret?.Data["password"], "password123");
diff --git a/Engine/Source/Programs/Horde/Plugins/Secrets/HordeServer.Secrets.Tests/Secrets/SecretResolveTests.cs b/Engine/Source/Programs/Horde/Plugins/Secrets/HordeServer.Secrets.Tests/Secrets/SecretResolveTests.cs
--- a/Engine/Source/Programs/Horde/Plugins/Secrets/HordeServer.Secrets.Tests/Secrets/SecretResolveTests.cs
+++ b/Engine/Source/Programs/Horde/Plugins/Secrets/HordeServer.Secrets.Tests/Secrets/SecretResolveTests.cs
@@ -29,7 +29,7 @@ namespace HordeServer.Tests.Secrets
 		}
 
 		[ClassInitialize]
-		public static async Task CreateSecretCollectionAsync(TestContext _)
+		public static void CreateSecretCollection(TestContext _)
 		{
 			string json = """
 				{
@@ -43,7 +43,7 @@ namespace HordeServer.Tests.Secrets
 			SecretCollection secretCollection = new(new SecretCollectionInternal([], s_logger.Object), new OptionsSnapshotStub<SecretsConfig>(secretsConfig));
 			GlobalConfig config = new();
 			config.Plugins.AddSecretsConfig(secretsConfig);
-			await config.PostLoadAsync(new ServerSettings(), [], []);
+			config.PostLoad(new ServerSettings(), [], []);
 			s_globalConfig = config;
 			s_secretCollection = secretCollection;
 		}
diff --git a/Engine/Source/Programs/Horde/Plugins/Secrets/HordeServer.Secrets/SecretsConfig.cs b/Engine/Source/Programs/Horde/Plugins/Secrets/HordeServer.Secrets/SecretsConfig.cs
--- a/Engine/Source/Programs/Horde/Plugins/Secrets/HordeServer.Secrets/SecretsConfig.cs
+++ b/Engine/Source/Programs/Horde/Plugins/Secrets/HordeServer.Secrets/SecretsConfig.cs
@@ -31,7 +31,7 @@ namespace HordeServer
 		private readonly Dictionary<SecretId, Dictionary<string, SecretProviderConfig>> _providerConfigLookup = new();
 
 		/// <inheritdoc/>
-		public Task PostLoadAsync(PluginConfigOptions configOptions, CancellationToken cancellationToken = default)
+		public void PostLoad(PluginConfigOptions configOptions)
 		{
 			_secretLookup.Clear();
 			foreach (SecretConfig secret in Secrets)
@@ -63,8 +63,6 @@ namespace HordeServer
 					}
 				}
 			}
-
-			return Task.CompletedTask;
 		}
 
 		/// <summary>
diff --git a/Engine/Source/Programs/Horde/Plugins/Storage/HordeServer.Storage/StorageConfig.cs b/Engine/Source/Programs/Horde/Plugins/Storage/HordeServer.Storage/StorageConfig.cs
--- a/Engine/Source/Programs/Horde/Plugins/Storage/HordeServer.Storage/StorageConfig.cs
+++ b/Engine/Source/Programs/Horde/Plugins/Storage/HordeServer.Storage/StorageConfig.cs
@@ -74,7 +74,7 @@ namespace HordeServer.Storage
 		/// <summary>
 		/// Called after the config has been read from disk
 		/// </summary>
-		public Task PostLoadAsync(PluginConfigOptions configOptions, CancellationToken cancellationToken = default)
+		public void PostLoad(PluginConfigOptions configOptions)
 		{
 			Acl = configOptions.ParentAcl;
 
@@ -123,8 +123,6 @@ namespace HordeServer.Storage
 			// Create a unique id to allow detect changes
 			byte[] data = JsonSerializer.SerializeToUtf8Bytes<StorageConfig>(this, JsonUtils.DefaultSerializerOptions);
 			Revision = IoHash.Compute(data).ToString();
-
-			return Task.CompletedTask;
 		}
 
 		static BackendConfig MergeBackendConfigs(BackendId backendId, Dictionary<BackendId, BackendConfig> baseIdToConfig, Dictionary<BackendId, BackendConfig> mergedIdToConfig)
diff --git a/Engine/Source/Programs/Horde/Plugins/Symbols/HordeServer.Symbols/SymbolsConfig.cs b/Engine/Source/Programs/Horde/Plugins/Symbols/HordeServer.Symbols/SymbolsConfig.cs
--- a/Engine/Source/Programs/Horde/Plugins/Symbols/HordeServer.Symbols/SymbolsConfig.cs
+++ b/Engine/Source/Programs/Horde/Plugins/Symbols/HordeServer.Symbols/SymbolsConfig.cs
@@ -22,7 +22,7 @@ namespace HordeServer.Symbols
 		readonly Dictionary<SymbolStoreId, SymbolStoreConfig> _storeLookup = new Dictionary<SymbolStoreId, SymbolStoreConfig>();
 
 		/// <inheritdoc/>
-		public Task PostLoadAsync(PluginConfigOptions configOptions, CancellationToken cancellationToken = default)
+		public void PostLoad(PluginConfigOptions configOptions)
 		{
 			foreach (SymbolStoreConfig store in Stores)
 			{
@@ -34,8 +34,6 @@ namespace HordeServer.Symbols
 			{
 				_storeLookup.Add(store.Id, store);
 			}
-
-			return Task.CompletedTask;
 		}
 
 		/// <summary>
diff --git a/Engine/Source/Programs/Horde/Plugins/Tools/HordeServer.Tools.Tests/Tools/ToolsControllerTests.cs b/Engine/Source/Programs/Horde/Plugins/Tools/HordeServer.Tools.Tests/Tools/ToolsControllerTests.cs
--- a/Engine/Source/Programs/Horde/Plugins/Tools/HordeServer.Tools.Tests/Tools/ToolsControllerTests.cs
+++ b/Engine/Source/Programs/Horde/Plugins/Tools/HordeServer.Tools.Tests/Tools/ToolsControllerTests.cs
@@ -46,7 +46,7 @@ public class ToolsControllerTests
 		globalConfig.Plugins.AddToolsConfig(toolsConfig);
 
 		ServerSettings serverSettings = new() { AuthMethod = AuthMethod.Horde };
-		await globalConfig.PostLoadAsync(serverSettings, pluginCollection.LoadedPlugins, Array.Empty<IDefaultAclModifier>());
+		globalConfig.PostLoad(serverSettings, pluginCollection.LoadedPlugins, Array.Empty<IDefaultAclModifier>());
 
 		Dictionary<string, string> settings = new() { { "Horde:AuthMethod", AuthMethod.Horde.ToString() } };
 		await using FakeHordeWebApp app = new(settings: settings);
diff --git a/Engine/Source/Programs/Horde/Plugins/Tools/HordeServer.Tools/ToolsConfig.cs b/Engine/Source/Programs/Horde/Plugins/Tools/HordeServer.Tools/ToolsConfig.cs
--- a/Engine/Source/Programs/Horde/Plugins/Tools/HordeServer.Tools/ToolsConfig.cs
+++ b/Engine/Source/Programs/Horde/Plugins/Tools/HordeServer.Tools/ToolsConfig.cs
@@ -21,7 +21,7 @@ namespace HordeServer.Tools
 		private readonly Dictionary<ToolId, ToolConfig> _toolLookup = new Dictionary<ToolId, ToolConfig>();
 
 		/// <inheritdoc/>
-		public Task PostLoadAsync(PluginConfigOptions configOptions, CancellationToken cancellationToken = default)
+		public void PostLoad(PluginConfigOptions configOptions)
 		{
 			_toolLookup.Clear();
 			foreach (ToolConfig tool in Tools)
@@ -29,8 +29,6 @@ namespace HordeServer.Tools
 				_toolLookup.Add(tool.Id, tool);
 				tool.PostLoad(configOptions.ParentAcl);
 			}
-
-			return Task.CompletedTask;
 		}
 
 		/// <summary>
