Horde: Reassign Horde issue to Jira user when Jira and issue owners differ

diff --git a/Engine/Source/Programs/Horde/Plugins/Build/HordeServer.Build/Issues/External/JiraService.cs b/Engine/Source/Programs/Horde/Plugins/Build/HordeServer.Build/Issues/External/JiraService.cs
--- a/Engine/Source/Programs/Horde/Plugins/Build/HordeServer.Build/Issues/External/JiraService.cs
+++ b/Engine/Source/Programs/Horde/Plugins/Build/HordeServer.Build/Issues/External/JiraService.cs
@@ -317,7 +317,7 @@ namespace HordeServer.Issues.External
 
 				foreach (IIssue issue in openIssues)
 				{
-					if (String.IsNullOrEmpty(issue.ExternalIssueKey) || issue.OwnerId != null)
+					if (String.IsNullOrEmpty(issue.ExternalIssueKey))
 					{
 						continue;
 					}
@@ -333,8 +333,9 @@ namespace HordeServer.Issues.External
 
 							IUser? user = await _userCollection.FindUserByEmailAsync(jiraIssue.AssigneeEmailAddress, cancellationToken: stoppingToken);
 
-							if (user != null)
+							if (user != null && user.Id != issue.OwnerId)
 							{
+
 								await _issueService.UpdateIssueAsync(issue.Id, ownerId: user.Id, cancellationToken: stoppingToken);
 
 								_logger.LogInformation("Assigning issue {IssueId} to user {UserId} by matching {EmailAddress}", issue.Id, user.Id, jiraIssue.AssigneeEmailAddress);
