Horde: When a string cannot be resolved to a secret then throw an exception if the string contains the secret prefix

diff --git a/Engine/Source/Programs/Horde/Plugins/Build/HordeServer.Build.Tests/Perforce/PerforceSecretsTests.cs b/Engine/Source/Programs/Horde/Plugins/Build/HordeServer.Build.Tests/Perforce/PerforceSecretsTests.cs
--- a/Engine/Source/Programs/Horde/Plugins/Build/HordeServer.Build.Tests/Perforce/PerforceSecretsTests.cs
+++ b/Engine/Source/Programs/Horde/Plugins/Build/HordeServer.Build.Tests/Perforce/PerforceSecretsTests.cs
@@ -73,18 +73,25 @@ public class PerforceSecretsTests : BuildTestSetup
 	}
 
 	[TestMethod]
+	[ExpectedException(typeof(KeyNotFoundException))]
 	[DataRow("horde:secret:notfound.property", DisplayName = "Resolve Perforce Secret, secret not found")]
 	[DataRow("horde:secret:p4secret.notfound", DisplayName = "Resolve Perforce Secret, property not found")]
 	public async Task ResolvePerforceCredentialsWithSecretNotFoundAsync(string value)
 	{
 		PerforceCredentials credentials = new() { Password = value, Ticket = value };
 		SecretConfig secret = new() { Id = new SecretId("p4secret") };
-		await AddCredentialsAsync(credentials, secret);
-		Assert.AreEqual(credentials.Password, value);
-		Assert.AreEqual(credentials.Ticket, value);
+		try
+		{
+			await AddCredentialsAsync(credentials, secret);
+		}
+		catch (AggregateException ae)
+		{
+			throw ae.InnerExceptions[0];
+		}
 	}
 
 	[TestMethod]
+	[ExpectedException(typeof(KeyNotFoundException))]
 	public async Task ResolvePerforceCredentialsWithSecretMisconfigurationAsync()
 	{
 		PerforceCredentials credentials = new() { Password = "horde:secret:p4secret.property" };
@@ -93,7 +100,13 @@ public class PerforceSecretsTests : BuildTestSetup
 			Id = new SecretId("p4secret"),
 			Sources = [ new ExternalSecretConfig() { Provider = "SecretProviderKeyNotFound" } ]
 		};
-		await AddCredentialsAsync(credentials, secret);
-		Assert.AreEqual(credentials.Password, "horde:secret:p4secret.property");
+		try
+		{
+			await AddCredentialsAsync(credentials, secret);
+		}
+		catch (AggregateException ae)
+		{
+			throw ae.InnerExceptions[0];
+		}
 	}
 }
diff --git a/Engine/Source/Programs/Horde/Plugins/Secrets/HordeServer.Secrets.Tests/Secrets/SecretProviderTests.cs b/Engine/Source/Programs/Horde/Plugins/Secrets/HordeServer.Secrets.Tests/Secrets/SecretProviderTests.cs
--- a/Engine/Source/Programs/Horde/Plugins/Secrets/HordeServer.Secrets.Tests/Secrets/SecretProviderTests.cs
+++ b/Engine/Source/Programs/Horde/Plugins/Secrets/HordeServer.Secrets.Tests/Secrets/SecretProviderTests.cs
@@ -50,7 +50,7 @@ namespace HordeServer.Tests.Secrets
 			SecretCollectionInternal secretCollectionInternal = new([
 				new InMemoryProvider(),
 				new HcpVaultSecretProvider(new StubHttpClientFactory(), NullLogger<HcpVaultSecretProvider>.Instance)
-			], NullLogger<SecretCollectionInternal>.Instance);
+			]);
 			string json = """
 				{
 					"secrets": [
@@ -279,7 +279,7 @@ namespace HordeServer.Tests.Secrets
 				}
 			""";
 			SecretsConfig secretsConfig = JsonSerializer.Deserialize<SecretsConfig>(json, JsonUtils.DefaultSerializerOptions)!;
-			SecretCollectionInternal secretCollectionInternal = new([], NullLogger<SecretCollectionInternal>.Instance);
+			SecretCollectionInternal secretCollectionInternal = new([]);
 			SecretCollection secretCollection = new(secretCollectionInternal, new OptionsSnapshotStub<SecretsConfig>(secretsConfig));
 			GlobalConfig config = new();
 			config.Plugins.AddSecretsConfig(secretsConfig);
diff --git a/Engine/Source/Programs/Horde/Plugins/Secrets/HordeServer.Secrets.Tests/Secrets/SecretResolveTests.cs b/Engine/Source/Programs/Horde/Plugins/Secrets/HordeServer.Secrets.Tests/Secrets/SecretResolveTests.cs
--- a/Engine/Source/Programs/Horde/Plugins/Secrets/HordeServer.Secrets.Tests/Secrets/SecretResolveTests.cs
+++ b/Engine/Source/Programs/Horde/Plugins/Secrets/HordeServer.Secrets.Tests/Secrets/SecretResolveTests.cs
@@ -6,9 +6,7 @@ using HordeServer.Plugins;
 using HordeServer.Secrets;
 using HordeServer.Server;
 using HordeServer.Utilities;
-using Microsoft.Extensions.Logging;
 using Microsoft.Extensions.Options;
-using Moq;
 
 namespace HordeServer.Tests.Secrets
 {
@@ -17,7 +15,6 @@ namespace HordeServer.Tests.Secrets
 	{
 		private static ISecretCollection s_secretCollection = default!;
 		private static GlobalConfig s_globalConfig = default!;
-		private static readonly Mock<ILogger<SecretCollectionInternal>> s_logger = new();
 
 		private class OptionsSnapshotStub<T>(T value) : IOptionsSnapshot<T> where T : class
 		{
@@ -40,7 +37,7 @@ namespace HordeServer.Tests.Secrets
 				}
 			""";
 			SecretsConfig secretsConfig = JsonSerializer.Deserialize<SecretsConfig>(json, JsonUtils.DefaultSerializerOptions)!;
-			SecretCollection secretCollection = new(new SecretCollectionInternal([], s_logger.Object), new OptionsSnapshotStub<SecretsConfig>(secretsConfig));
+			SecretCollection secretCollection = new(new SecretCollectionInternal([]), new OptionsSnapshotStub<SecretsConfig>(secretsConfig));
 			GlobalConfig config = new();
 			config.Plugins.AddSecretsConfig(secretsConfig);
 			config.PostLoad(new ServerSettings(), [], []);
@@ -81,8 +78,6 @@ namespace HordeServer.Tests.Secrets
 		[DataRow("", DisplayName = "Empty string")]
 		[DataRow(null, DisplayName = "Null string")]
 		[DataRow("invalid string", DisplayName = "Incorrect format")]
-		[DataRow("horde:secret:doesNotExist.password", DisplayName = "Secret does not exist")]
-		[DataRow("horde:secret:mySecret.doesNotExist", DisplayName = "Secret property does not exist")]
 		public async Task ResolveInvalidSecretAsync(string value)
 		{
 			ISecretProperty? secret = await s_secretCollection.ResolveAsync(value);
@@ -90,19 +85,19 @@ namespace HordeServer.Tests.Secrets
 		}
 
 		[TestMethod]
+		[ExpectedException(typeof(KeyNotFoundException))]
+		[DataRow("horde:secret:doesNotExist.password", DisplayName = "Secret does not exist")]
+		[DataRow("horde:secret:mySecret.doesNotExist", DisplayName = "Secret property does not exist")]
+		public async Task ResolveSecretNotFoundAsync(string value)
+		{
+			_ = await s_secretCollection.ResolveAsync(value);
+		}
+
+		[TestMethod]
+		[ExpectedException(typeof(KeyNotFoundException))]
 		public async Task ResolveSecretWithMisconfiguredConfigAsync()
 		{
-			ISecretProperty? secret = await s_secretCollection.ResolveAsync("horde:secret:misconfiguredSecret.property");
-			Assert.IsNull(secret);
-			s_logger.Verify(x =>
-				x.Log(
-					LogLevel.Error,
-					It.IsAny<EventId>(),
-					It.Is<It.IsAnyType>((value, _) => value.ToString()!.Contains("Unable to resolve the string")),
-					It.IsAny<KeyNotFoundException>(),
-					It.IsAny<Func<It.IsAnyType, Exception?, string>>()
-				)
-			);
+			_ = await s_secretCollection.ResolveAsync("horde:secret:misconfiguredSecret.property");
 		}
 	}
 }
diff --git a/Engine/Source/Programs/Horde/Plugins/Secrets/HordeServer.Secrets/Secrets/SecretCollection.cs b/Engine/Source/Programs/Horde/Plugins/Secrets/HordeServer.Secrets/Secrets/SecretCollection.cs
--- a/Engine/Source/Programs/Horde/Plugins/Secrets/HordeServer.Secrets/Secrets/SecretCollection.cs
+++ b/Engine/Source/Programs/Horde/Plugins/Secrets/HordeServer.Secrets/Secrets/SecretCollection.cs
@@ -3,7 +3,6 @@
 using System.Text.Json;
 using System.Text.RegularExpressions;
 using EpicGames.Horde.Secrets;
-using Microsoft.Extensions.Logging;
 using Microsoft.Extensions.Options;
 
 namespace HordeServer.Secrets
@@ -77,15 +76,13 @@ namespace HordeServer.Secrets
 		}
 
 		readonly Dictionary<string, ISecretProvider> _providers;
-		readonly ILogger _logger;
 
 		/// <summary>
 		/// Constructor
 		/// </summary>
-		public SecretCollectionInternal(IEnumerable<ISecretProvider> providers, ILogger<SecretCollectionInternal> logger)
+		public SecretCollectionInternal(IEnumerable<ISecretProvider> providers)
 		{
 			_providers = providers.ToDictionary(x => x.Name, x => x);
-			_logger = logger;
 		}
 
 		/// <summary>
@@ -185,24 +182,17 @@ namespace HordeServer.Secrets
 			}
 
 			SecretId secretId = new(match.Groups["SecretId"].Value);
-			ISecret? secret = null;
-			try
-			{
-				secret = await GetAsync(secretId, secretsConfig, cancellationToken);
-			}
-			catch (Exception e)
-			{
-				_logger.LogError(e, "Unable to resolve the string '{Value}'. The string was not recognized as a valid secret. Check the configuration of the secrets config: '{Message}'", value, e.Message);
-			}
+
+			ISecret? secret = await GetAsync(secretId, secretsConfig, cancellationToken);
 			if (secret == null)
 			{
-				return null;
+				throw new KeyNotFoundException($"The given secret '{secretId}' was not found");
 			}
 
 			string propertyName = match.Groups["PropertyName"].Value;
 			if (!secret.Data.TryGetValue(propertyName, out string? propertyValue))
 			{
-				return null;
+				throw new KeyNotFoundException($"The given property '{propertyName}' for the secret '{secretId}' was not found");
 			}
 
 			ISecretProperty secretProperty = new SecretProperty(secretId, propertyName, propertyValue);
