Horde: Avoid issue service exceptions during suspect finding

diff --git a/Engine/Source/Programs/Horde/Plugins/Build/HordeServer.Build/Issues/IssueService.cs b/Engine/Source/Programs/Horde/Plugins/Build/HordeServer.Build/Issues/IssueService.cs
--- a/Engine/Source/Programs/Horde/Plugins/Build/HordeServer.Build/Issues/IssueService.cs
+++ b/Engine/Source/Programs/Horde/Plugins/Build/HordeServer.Build/Issues/IssueService.cs
@@ -1335,7 +1335,26 @@ namespace HordeServer.Issues
 				int maxRank = -1;
 				foreach (ICommit commit in changes)
 				{
-					int rank = await GetChangeScoreAsync(commit, filter, suspectFiles, suspectNotes, suspectSymbols, _tracer);
+					int rank = -1;
+					try
+					{
+						rank = await GetChangeScoreAsync(commit, filter, suspectFiles, suspectNotes, suspectSymbols, _tracer);
+					}
+					catch (NullReferenceException nullRefEx)
+					{
+						_logger.LogWarning("Exception while ranking suspect changes: {Message} - {Callstack}", nullRefEx.Message,  nullRefEx.StackTrace);
+
+					}
+					catch (EpicGames.Perforce.PerforceException perforceEx)
+					{
+						_logger.LogWarning("Exception while ranking suspect changes: {Message} - {Callstack}", perforceEx.Message, perforceEx.StackTrace);
+
+						// re-throw any perforce exceptions that aren't known or expected here
+						if (!perforceEx.Message.Contains("does not contain any files in"))
+						{
+							throw perforceEx;
+						}
+					}
 					if (rank > maxRank)
 					{
 						suspects.Clear();
