Horde: Fix warnings during exception handling in suspect finding
#rnx

diff --git a/Engine/Source/Programs/Horde/Plugins/Build/HordeServer.Build/Issues/IssueService.cs b/Engine/Source/Programs/Horde/Plugins/Build/HordeServer.Build/Issues/IssueService.cs
--- a/Engine/Source/Programs/Horde/Plugins/Build/HordeServer.Build/Issues/IssueService.cs
+++ b/Engine/Source/Programs/Horde/Plugins/Build/HordeServer.Build/Issues/IssueService.cs
@@ -1350,9 +1350,9 @@ namespace HordeServer.Issues
 						_logger.LogWarning("Exception while ranking suspect changes: {Message} - {Callstack}", perforceEx.Message, perforceEx.StackTrace);
 
 						// re-throw any perforce exceptions that aren't known or expected here
-						if (!perforceEx.Message.Contains("does not contain any files in"))
+						if (!perforceEx.Message.Contains("does not contain any files in", System.StringComparison.OrdinalIgnoreCase))
 						{
-							throw perforceEx;
+							throw;
 						}
 					}
 					if (rank > maxRank)
