Horde: Update release notes
#skipci
#submitool bypass

diff --git a/Engine/Source/Programs/Horde/Docs/ReleaseNotes.md b/Engine/Source/Programs/Horde/Docs/ReleaseNotes.md
--- a/Engine/Source/Programs/Horde/Docs/ReleaseNotes.md
+++ b/Engine/Source/Programs/Horde/Docs/ReleaseNotes.md
@@ -1,5 +1,48 @@
 # Release Notes
 
+## 2025-06-19
+
+* Avoid issue service exceptions during suspect finding (43629946)
+* When a string cannot be resolved to a secret then throw an exception if the string contains the secret prefix (43615113)
+* Reassign Horde issue to Jira user when Jira and issue owners differ (43614130)
+* Change IPluginConfig.PostLoad to non-async (43610732)
+* Add method to compute client for allocating a UBA cache server session (43610555)
+* Horde Dashboard: portable search url inconsistency (43609394)
+* Changed Horde Analytics dashbaord category name from SubmitTool to Submit (43607713)
+* Fix broken UbaCrypto test (43607690)
+* Added an example of SubmitTool metric and dashboard category to Horde Analytics (43607541)
+* Add REST API endpoint for requesting a UBA cache server session (43581723)
+* Add ACL actions for read/write access to UBA cache server (43574053)
+* Use zeroed IV for UbaCrypto to match HTTP endpoint in UBA cache server (43554905)
+* Horde Dashboard: job search url parameters (43545493)
+* Horde Server: Fix v2 device reservation exception with unit test (43498441)
+* Horde Dashboard: A few tweaks to make the agent modal telemetry more consistent (43496811)
+* Add support in compute client for requesting a UBA cache server (43496230)
+* Horde Dashboard: Agent telemetry fixes and improvements for Agent and Log views (43495864)
+* Perforce passwords and tickets in the BuildConfig can be resolved to the value of a ISecret during post load replacing the need for them to be plain text values (43492230)
+* Update UbaCrypto to match C++ impl (43491786)
+* Use Win32 API for calculating memory telemetry on Windows
+* Use t alias to shorten URLs in JobHandler
+* Add shortened template alias to api/v1/jobs/streams/{streamId} endpoint (43433353)
+* Default to read/write access for UBA cache for now until per user permissions are added (43390743)
+* Horde Dashboard: Show "time to step" in trends panel (43350691)
+* Fix device checkout expiration link and clarify instructions (43350167)
+* Use new encrypted /addsession endpoint for allocating UBA cache sessions (43349167)
+* Adding Horde job/step metatagging (43276253)
+* Horde Server: Add job start time to JobStepRef response (43274814)
+* Add C# implementation of UBA crypto (43257937)
+* Horde Dashboard: Disable job step retries with expired TempStorage (43233123)
+* Fix sorting agents by status (43230711)
+* Support debugging jobs as scheduler (43229680)
+* Allow sessions to get job responses and post job metadata (43228360)
+* Add HTTP config values for UBA cache in compute cluster config (43225086)
+* Add API to resolve a string to a secret property value (43224431)
+* Update Secrets plugin with a new server configuration option WithAws that controls if the secret provider AwsParameterStore is to be enabled (43222303)
+* Adding capabilities to the warning defaults, cleaning up tests, adding more warning+error strings for easier long term adoption & toggling. (43196850)
+* Add a "No Suspects" emoji indication to build health slack threads that lack suspects (43191214)
+* Use preferred colors for device status, slightly increased brightness threshold (43187852)
+* Use code commit tag to get suspect CLs (43161937)
+
 ## 2025-06-05
 
 * Add C# implementation of UBA crypto (43257937)
