Horde Dashboard: job status feedback

diff --git a/Engine/Source/Programs/Horde/HordeDashboard/src/backend/JobHandler.ts b/Engine/Source/Programs/Horde/HordeDashboard/src/backend/JobHandler.ts
--- a/Engine/Source/Programs/Horde/HordeDashboard/src/backend/JobHandler.ts
+++ b/Engine/Source/Programs/Horde/HordeDashboard/src/backend/JobHandler.ts
@@ -5,7 +5,7 @@ import moment from 'moment';
 import backend from '../backend';
 import { JobData, JobState, JobStepOutcome, JobStreamQuery, StreamData } from '../backend/Api';
 
-export type FilterStatus = "Running" | "Complete" | "Succeeded" | "Failed" | "Waiting";
+export type FilterStatus = "Running" | "Complete" | "Succeeded" | "Failed" | "Cancelled" | "Waiting";
 
 const jobPageSize = 50;
 const jobsRefreshTime = 10000;
@@ -77,7 +77,7 @@ export class JobHandler {
         try {
 
             // discover new/updated jobs
-            let filter = "id,streamId,name,change,preflightChange,templateId,templateHash,graphHash,startedByUserInfo,abortedByUserInfo,createTime,state,arguments,updateTime,labels,defaultLabel,preflightDescription";
+            let filter = "id,streamId,name,change,preflightChange,templateId,templateHash,graphHash,startedByUserInfo,abortedByUserInfo,createTime,state,arguments,updateTime,labels,defaultLabel,preflightDescription,cancellationReason";
 
             if (this.includeBatches) {
                 filter += ",batches";
@@ -210,14 +210,21 @@ export class JobHandler {
 
     }
 
-    getFilteredJobs(status?: FilterStatus[]) {
-        if (!status || !status.length) {
+    getFilteredJobs(status?: FilterStatus[], minCL?: number, maxCL?: number) {
+        if ((!status || !status.length) && !minCL && !maxCL) {
             return this.jobs;
         }
 
         return this.jobs.filter(job => {
+            if((minCL || maxCL) && !job.change) return false;
+            
+            let withinBounds = true;
+            if((minCL && job.change! < minCL) || (maxCL && job.change! > maxCL)) withinBounds = false;
+            if((minCL || maxCL) && !withinBounds) return false;
 
-            return status.every(status => {
+            if(!status || !status.length) return true;
+
+            return status.some(status => {
 
                 if ((status === "Complete" || status === "Succeeded") && job.state !== JobState.Complete) {
                     return false;
@@ -243,6 +250,10 @@ export class JobHandler {
                     }
                 }
 
+                if(status === "Cancelled" && !job.abortedByUserInfo) {
+                    return false;
+                }
+
                 return true;
             })
 
diff --git a/Engine/Source/Programs/Horde/HordeDashboard/src/components/ChangeButton.tsx b/Engine/Source/Programs/Horde/HordeDashboard/src/components/ChangeButton.tsx
--- a/Engine/Source/Programs/Horde/HordeDashboard/src/components/ChangeButton.tsx
+++ b/Engine/Source/Programs/Horde/HordeDashboard/src/components/ChangeButton.tsx
@@ -21,6 +21,7 @@ export type JobParameters = {
    preflightChange?: number;
    preflightDescription?: string;
    abortedByUserInfo?: GetThinUserInfoResponse;
+   cancellationReason?: string;
    startedByUserInfo?: GetThinUserInfoResponse;
    batches?: GetBatchResponse[];
 }
@@ -72,11 +73,6 @@ export const ChangeContextMenu: React.FC<{ target: ChangeContextMenuTarget, onDi
          }
       });
 
-      menuItems.push({
-         key: "commit_divider",
-         itemType: ContextualMenuItemType.Divider
-      });
-
    } else if (commit) {
 
       menuItems.push({
@@ -99,11 +95,28 @@ export const ChangeContextMenu: React.FC<{ target: ChangeContextMenuTarget, onDi
          }
       });
 
+   }
+   
+   // Should only show cancel text if not in an incremental view where "commit" exists
+   if(!commit && job.abortedByUserInfo && job.cancellationReason) { 
+      menuItems.push({
+         key: "cancelled",
+         onRender: () => {
+            return (
+               <Stack style={{ padding: 18, paddingTop: job.preflightDescription ? 0 : 18, paddingBottom: 18, maxWidth: 800 }} horizontal tokens={{ childrenGap: 12 }}>
+                  <Text variant="small" style={{ fontFamily: "Horde Open Sans Bold" }}>Cancellation Reason:</Text>
+                  <Text variant="small" >{job.cancellationReason}</Text>
+               </Stack>
+            )
+         }
+      });
+   }
+
+   if(job.cancellationReason || commit || job.preflightDescription){
       menuItems.push({
          key: "commit_divider",
          itemType: ContextualMenuItemType.Divider
       });
-
    }
 
    menuItems.push({
@@ -269,8 +282,24 @@ export const ChangeButton: React.FC<{ job?: JobParameters, stepRef?: GetJobStepR
 
    return (<Stack verticalAlign="center" verticalFill={true} horizontalAlign="start"> <div style={{ paddingBottom: "1px" }}>
       <Stack tokens={{ childrenGap: 4 }}>
-         <span ref={spanRef} style={{ padding: "2px 6px 2px 6px", height: "15px", cursor: "pointer", color: "#FFFFFF", backgroundColor: buttonColor ? `${buttonColor}` : defaultBackgroundColor }} className={job.startedByUserInfo ? "cl-callout-button-user" : "cl-callout-button"} onClick={(ev) => { ev.stopPropagation(); ev.preventDefault(); setMenuShown(!menuShown) }} >{change}</span>
-         {(!!showStatus) && <span ref={spanRef} style={{ padding: "2px 6px 2px 6px", height: "16px", cursor: "pointer", userSelect: "none", fontFamily: "Horde Open Sans SemiBold", fontSize: "10px", backgroundColor: color, color: "rgb(255, 255, 255)" }} onClick={(ev) => { ev.preventDefault(); setMenuShown(!menuShown) }}>{text}</span>}
+         <span 
+            ref={spanRef} 
+            style={{ padding: "2px 6px 2px 6px", height: "15px", cursor: "pointer", color: "#FFFFFF", backgroundColor: buttonColor ? `${buttonColor}` : defaultBackgroundColor }} 
+            className={job.startedByUserInfo ? "cl-callout-button-user" : "cl-callout-button"} 
+            onClick={(ev) => { ev.stopPropagation(); ev.preventDefault(); setMenuShown(!menuShown) }}
+         >
+            {change}
+         </span>
+         {(!!showStatus) && 
+            <span 
+               ref={spanRef} 
+               style={{ padding: "2px 6px 2px 6px", height: "16px", cursor: "pointer", userSelect: "none", fontFamily: "Horde Open Sans SemiBold", fontSize: "10px", backgroundColor: color, color: "rgb(255, 255, 255)" }} 
+               className={job.startedByUserInfo ? "cl-callout-button-user" : "cl-callout-button"} 
+               onClick={(ev) => { ev.stopPropagation(); ev.preventDefault(); setMenuShown(!menuShown) }}
+            >
+               {text}
+            </span>
+         }
       </Stack>
       {menuShown && <ChangeContextMenu target={{ ref: spanRef }} job={job} commit={commit} stepRef={stepRef} rangeCL={rangeCL} onDismiss={() => setMenuShown(false)} />}
    </div></Stack>);
diff --git a/Engine/Source/Programs/Horde/HordeDashboard/src/components/JobFilterBar.tsx b/Engine/Source/Programs/Horde/HordeDashboard/src/components/JobFilterBar.tsx
--- a/Engine/Source/Programs/Horde/HordeDashboard/src/components/JobFilterBar.tsx
+++ b/Engine/Source/Programs/Horde/HordeDashboard/src/components/JobFilterBar.tsx
@@ -1,40 +1,26 @@
 // Copyright Epic Games, Inc. All Rights Reserved.
 
 import { action, makeObservable, observable } from "mobx";
-import { DefaultButton, Dropdown, IContextualMenuItem, IContextualMenuProps, Label, Stack } from "@fluentui/react";
-import React, { useEffect, useState } from "react";
+import { Checkbox, ContextualMenuItemType, DefaultButton, Dropdown, IContextualMenuItem, IContextualMenuProps, Label, Stack, Text, TextField } from "@fluentui/react";
+import React, { useEffect, useState, useSyncExternalStore } from "react";
 import { useBackend } from "../backend";
 import { GetJobsTabResponse, TabType, GetTemplateRefResponse, StreamData } from "../backend/Api";
-import { FilterStatus } from "../backend/JobHandler";
+import { FilterStatus, JobHandler } from "../backend/JobHandler";
 import templateCache from '../backend/TemplateCache';
 import { getHordeStyling } from "../styles/Styles";
+import { observer } from "mobx-react-lite";
+import dashboard from "horde/backend/Dashboard";
 
 // todo, store filter for streamId so same when return
 // todo, put on a timeout so can capture a number of changes before update
 
-const dropDownStyle: any = () => {
-
+const statusItems = ["Running", "Complete", "Succeeded", "Failed", "Waiting", "Cancelled"].map(status => {
    return {
-      dropdown: {},
-      callout: {
-         selectors: {
-            ".ms-Callout-main": {
-               padding: "4px 4px 12px 12px",
-               overflow: "hidden"
-            }
-         }
-      },
-      dropdownItemHeader: { fontSize: 12 },
-      dropdownOptionText: { fontSize: 12 },
-      dropdownItem: {
-         minHeight: 28, lineHeight: 28
-      },
-      dropdownItemSelected: {
-         minHeight: 28, lineHeight: 28, backgroundColor: "inherit"
-      }
-   }
-}
-
+      key: status,
+      text: status,
+      status: status
+   };
+});
 
 export class JobFilter {
 
@@ -91,20 +77,42 @@ export class JobFilter {
       }
    }
 
+   setMinCL(value: number) {
+      if(value !== this.minCL){
+         this.setUpdated();
+      }
+
+      this.minCL = value;
+   }
+
+   setMaxCL(value: number) {
+      if(value !== this.maxCL){
+         this.setUpdated();
+      }
+
+      this.maxCL = value;
+   }
+
    @observable
    updated: number = 0;
 
    templates?: GetTemplateRefResponse[];
    status?: FilterStatus[];
+   minCL: number;
+   maxCL: number;
 }
 
 export const jobFilter = new JobFilter();
 
-const TemplateSelector: React.FC<{ stream: StreamData, templates: GetTemplateRefResponse[], selected: GetTemplateRefResponse | undefined, setSelectedTemplate: (template: GetTemplateRefResponse | undefined) => void }> = ({ stream, templates, selected, setSelectedTemplate }) => {
+const TemplateSelector: React.FC<{ stream: StreamData, templates: GetTemplateRefResponse[] }> = observer(({ stream, templates }) => {
+   const [selectedTemplates, setSelectedTemplates] = useState<Set<GetTemplateRefResponse>>(new Set());
 
-   let templateOptions: IContextualMenuItem[] = [];
+   useEffect(() => {
+      jobFilter.set(Array.from(selectedTemplates), jobFilter.status);
+   }, [selectedTemplates]);
 
-   const sorted = new Map<string, GetTemplateRefResponse[]>();
+   const templateOptions: IContextualMenuItem[] = [];
+   const grouped = new Map<string, GetTemplateRefResponse[]>();
 
    templates.forEach(t => {
 
@@ -118,31 +126,109 @@ const TemplateSelector: React.FC<{ stream: StreamData, templates: GetTemplateRef
             return;
          }
 
-         if (!sorted.has(jtab.title)) {
-            sorted.set(jtab.title, []);
+         if (!grouped.has(jtab.title)) {
+            grouped.set(jtab.title, []);
          }
 
-         sorted.get(jtab.title)!.push(t);
+         grouped.get(jtab.title)!.push(t);
 
       })
    })
 
-   Array.from(sorted.keys()).sort((a, b) => a < b ? -1 : 1).forEach(cat => {
+   Array.from(grouped.keys()).sort((a, b) => a < b ? -1 : 1).forEach(cat => {
 
-      const templates = sorted.get(cat);
-      if (!templates?.length) {
+      const catTemplates = grouped.get(cat);
+      if (!catTemplates?.length) {
          return;
       }
 
-      const subItems = templates.sort((a, b) => a.name < b.name ? -1 : 1).map(t => {
-         return { key: t.id, text: t.name, onClick: () => { setSelectedTemplate(t) } };
+      const subItems: IContextualMenuItem[] = catTemplates.sort((a, b) => a.name < b.name ? -1 : 1).map(t => {
+
+         return {
+            itemType: ContextualMenuItemType.Normal,
+            key: t.id,
+            text: t.name,
+            data: t,
+            onRender: () => {
+               return <Stack horizontal verticalFill verticalAlign="center" style={{ padding: "4px 12px 4px 12px", cursor: "pointer" }} onClick={(ev) => {
+                  ev?.preventDefault();
+                  ev?.stopPropagation();
+
+                  const newTemplates = new Set(selectedTemplates);
+                  newTemplates.has(t) ? newTemplates.delete(t) : newTemplates.add(t);
+
+                  setSelectedTemplates(newTemplates);
+               }} >
+                  <Checkbox checked={selectedTemplates.has(t)} />
+                  <Text>{t.name}</Text></Stack>
+            }
+         };
       })
 
-      templateOptions.push({ key: `${cat}_category`, text: cat, subMenuProps: { items: subItems } });
+      const isSelected = (ct: GetTemplateRefResponse) => selectedTemplates.has(ct);
+      let allCatTemplatesSelected: boolean = catTemplates.every(isSelected);
+      let someCatTemplatesSelected: boolean = catTemplates.some(isSelected) && !allCatTemplatesSelected;
+
+      if (catTemplates?.length > 1) {
+         subItems.unshift({
+            key: `select_all_${cat}`,
+            onRender: () => {
+               return <Stack horizontal verticalFill verticalAlign="center" style={{ padding: "4px 12px 4px 12px", cursor: "pointer" }} onClick={(ev) => { 
+                  ev?.stopPropagation();
+                  ev?.preventDefault();
+      
+                  allCatTemplatesSelected = !allCatTemplatesSelected;
+                  
+                  const newTemplates = new Set(selectedTemplates);
+
+                  if (allCatTemplatesSelected) {
+                     catTemplates.forEach( ct => 
+                        newTemplates.add(ct)
+                     )
+                  }
+                  else {
+                     catTemplates.forEach( ct => 
+                        newTemplates.delete(ct)
+                     )
+                  }
+
+                  setSelectedTemplates(newTemplates);        
+                  return false;
+               }} >
+               <Checkbox indeterminate={someCatTemplatesSelected} checked={allCatTemplatesSelected} />
+               <Text style={{ color: dashboard.darktheme ? "#ABABAB" : "#575757" }}>Select All</Text></Stack>
+            }
+         });
+      }
+
+      templateOptions.push({ 
+         key: `${cat}_category`,
+         text: cat,
+         iconProps: { iconName: allCatTemplatesSelected ? 'Tick' : someCatTemplatesSelected ? 'Form' : '', style: { fontSize: '14px' } },
+         subMenuProps: { items: subItems }
+      });
 
    })
 
-   templateOptions.push({ key: `show_all_templates`, text: "All Templates", onClick: (ev) => { ev?.stopPropagation(); ev?.preventDefault(); setSelectedTemplate(undefined) } });
+   let allTemplatesSelected: boolean = selectedTemplates.size === templates.length;
+   let someTemplatesSelected: boolean = selectedTemplates.size > 0 && !allTemplatesSelected;
+
+   templateOptions.unshift({
+      key: `select_all_templates`,
+      onRender: () => {
+         return <Stack horizontal verticalFill verticalAlign="center" style={{ padding: "4px 12px 4px 7px", cursor: "pointer" }} onClick={(ev) => { 
+            ev?.stopPropagation();
+            ev?.preventDefault();
+
+            allTemplatesSelected = !allTemplatesSelected;
+            
+            setSelectedTemplates(allTemplatesSelected ? new Set(templates) : new Set());
+            return false;
+         }} >
+         <Checkbox indeterminate={someTemplatesSelected} checked={allTemplatesSelected} />
+         <Text style={{ padding: "6px 2px"}}>Select All</Text></Stack>
+      }
+   });
 
    const templateMenuProps: IContextualMenuProps = {
       shouldFocusOnMount: true,
@@ -150,96 +236,169 @@ const TemplateSelector: React.FC<{ stream: StreamData, templates: GetTemplateRef
       items: templateOptions,
    };
 
-   return <Stack style={{ paddingTop: 4 }}>
-      <Label>Template</Label>
-      <DefaultButton style={{ width: 352, textAlign: "left", height: 30 }} menuProps={templateMenuProps} text={selected?.name ?? "All"} />
+   let templateText = "None Selected";
+   if (selectedTemplates.size) {
+      templateText = Array.from(selectedTemplates).map(t => t.name).join(", ");
+      if (templateText.length > 84) {
+         templateText = templateText.slice(0, 84);
+         templateText += "...";
+      }
+   }
+
+   return <Stack>
+      <Label>{selectedTemplates.size > 1 ? "Templates" : "Template"}</Label>
+      <DefaultButton style={{ width: 352, textAlign: "left" }} menuProps={templateMenuProps} text={templateText} />
    </Stack>
 
-}
+})
 
-// Job filter bar for "all" jobs view
-export const JobFilterBar: React.FC<{ streamId: string }> = ({ streamId }) => {
+const StatusSelector: React.FC = () => {
+   const [selectedStatuses, setSelectedStatuses] = useState<Set<string>>(new Set());
 
-   const { projectStore } = useBackend();
+   useEffect(() => {
+      jobFilter.set(jobFilter.templates, Array.from(selectedStatuses) as FilterStatus[]);
+   }, [selectedStatuses])
 
-   const [state, setState] = useState<{ streamId?: string; templates?: GetTemplateRefResponse[], selectedTemplate?: GetTemplateRefResponse, status?: FilterStatus | "All" | undefined }>({});
+   let statusOptions: IContextualMenuItem[] = [];
 
-   const stream = projectStore.streamById(streamId);
+   let allStatusesSelected: boolean = selectedStatuses.size === statusItems.length;
+   let someStatusesSelected: boolean = selectedStatuses.size > 0 && !allStatusesSelected;
 
-   useEffect(() => {
+   const statuses = statusItems.map(s => s.text);
 
-      return () => {
-         jobFilter.set(undefined, undefined);
-      };
+   statusOptions.push({
+      itemType: ContextualMenuItemType.Normal,
+      key: `select_all_statuses`,
+      onRender: () => {
+         return (
+            <Stack horizontal verticalFill verticalAlign="center" style={{ padding: "8px 12px 8px 12px", cursor: "pointer" }} onClick={(ev) => {
+               ev?.preventDefault();
+               ev?.stopPropagation();
 
-   }, []);
+               allStatusesSelected = !allStatusesSelected;
 
-   const { hordeClasses } = getHordeStyling();
-
-   if (!stream) {
-      console.error("unable to get stream");
-      return <div>unable to get stream</div>;
-   }
+               setSelectedStatuses(allStatusesSelected ? new Set(statuses) : new Set());
+            }} >
+               <Checkbox indeterminate={someStatusesSelected} checked={allStatusesSelected} />
+               <Text>Select All</Text>
+            </Stack>
+         )
+      }
+      
+   });
 
-   if (!state.templates || streamId !== state.streamId) {
-      templateCache.getStreamTemplates(stream).then(data => {
-         setState({ streamId: streamId, templates: data, status: "All" });
-         jobFilter.set(data);
+   statuses.forEach((status) => {
+      statusOptions.push({ 
+         itemType: ContextualMenuItemType.Normal,
+         key: `${status}_key`,
+         onRender: () => {
+            return (
+               <Stack horizontal verticalFill verticalAlign="center" style={{ padding: "8px 12px 8px 12px", cursor: "pointer" }} onClick={(ev) => {
+                  ev?.preventDefault();
+                  ev?.stopPropagation();
+
+                  const newStatuses = new Set(selectedStatuses);
+
+                  if (newStatuses.has(status)) {
+                     newStatuses.delete(status);
+                  } else {
+                     newStatuses.add(status);
+                  }
+
+                  setSelectedStatuses(newStatuses);
+               }} >
+                  <Checkbox checked={selectedStatuses.has(status)} />
+                  <Text>{status}</Text>
+               </Stack>
+            ) 
+         }
       });
-      return null;
-   }
+   })
 
-   let templates: GetTemplateRefResponse[] = state.templates.map(t => t);
+   const statusMenuProps: IContextualMenuProps = {
+      shouldFocusOnMount: true,
+      items: statusOptions
+   };
 
-   if (!templates.length) {
-      return null;
+   let statusText = "None Selected";
+   if (selectedStatuses.size) {
+      statusText = Array.from(selectedStatuses).join(", ");
+      if (statusText.length > 84) {
+         statusText = statusText.slice(0, 84);
+         statusText += "...";
+      }
    }
 
-   const statusItems = ["Running", "Complete", "Succeeded", "Failed", "Waiting", "All"].map(status => {
-      return {
-         key: status,
-         text: status,
-         status: status
-      };
-   });
+   return (
+      <Stack>
+         <Label>{selectedStatuses.size > 1 ? "Statuses" : "Status"}</Label>
+         <DefaultButton style={{ width: 352, textAlign: "left" }} menuProps={statusMenuProps} text={statusText} />
+      </Stack>
+   )
+}
 
-   let filterTemplates: GetTemplateRefResponse[] = [];
-   if (state.selectedTemplate) {
-      filterTemplates = [state.selectedTemplate];
-   }
+// Job filter bar for "all" jobs view
+export const JobFilterBar: React.FC<{ streamId: string }> = observer(({ streamId }) => {
 
-   if (!jobFilter?.status?.find(s => s === state.status)) {
-      if ((!state.status || state.status === "All") && jobFilter.status?.length) {
-         jobFilter.set(jobFilter.templates, undefined);
-         return null;
-      }
+   const { projectStore } = useBackend();
+
+   const [templates, setTemplates] = useState<GetTemplateRefResponse[]>();
 
-      if (state.status && state.status !== "All") {
-         jobFilter.set(jobFilter.templates, [state.status]);
-         return null;
+   const stream = projectStore.streamById(streamId);
+   
+   const { hordeClasses } = getHordeStyling();
+
+   const [changelistBounds, setChangelistBounds] = useState<{min?: string, max?: string}>({min: undefined, max: undefined})
+
+   useEffect(() => {
+      // Cancellation flag for managing stale async calls (prevent updating the UI with old data)
+      let cancelled = false;
+      if(stream) {
+         templateCache.getStreamTemplates(stream).then(data => {
+            if(cancelled) return;
+
+            setTemplates(data);
+            jobFilter.set(data, jobFilter.status); // Init to all templates
+         })
       }
+      return () => {
+         cancelled = true;
+         jobFilter.set(undefined, undefined);
+      };
+
+   }, [stream]);
+
+   useEffect(() => {
+      jobFilter.setMinCL(Number(changelistBounds.min));
+      jobFilter.setMaxCL(Number(changelistBounds.max));
+   }, [changelistBounds])
+
+   if (!stream || !templates) {
+      console.error("unable to get stream or templates");
+      return <div>unable to get stream or templates</div>;
    }
 
-   return <Stack horizontal tokens={{ childrenGap: 24 }} className={hordeClasses.modal}>
-      <TemplateSelector stream={stream} templates={templates} selected={state.selectedTemplate} setSelectedTemplate={(t) => {
-         jobFilter.set(t ? [t] : [...templates], (state.status === "All" || !state.status) ? undefined : [state.status]);
-         setState({ ...state, selectedTemplate: t });         
-      }} />
-      <Dropdown
-         style={{ width: 120 }}
-         styles={dropDownStyle}
-         label="Status"
-         selectedKey={state.status}
-         options={statusItems}
-         onDismiss={() => {
-            jobFilter.set(filterTemplates.length ? filterTemplates : undefined, (state.status === "All" || !state.status) ? undefined : [state.status]);
-         }}
-         onChange={(event, option, index) => {
-            if (option) {
-               setState({ streamId: streamId, templates: state.templates, selectedTemplate: state.selectedTemplate, status: option.key as FilterStatus });
+   return (
+      <Stack horizontal tokens={{ childrenGap: 24 }} className={hordeClasses.modal}>
+         <TemplateSelector 
+            stream={stream}
+            templates={templates} 
+         />
+         <StatusSelector />
+         <TextField value={changelistBounds.min ?? ""} style={{ width: 168 }} label="Min Changelist" onChange={(ev, newValue) => {
+            newValue = newValue?.trim();
+            if (!newValue) {
+               newValue = undefined;
             }
-         }}
-      />
-
-   </Stack>;
-};
\ No newline at end of file
+            setChangelistBounds({...changelistBounds, min: newValue});
+         }} />
+         <TextField value={changelistBounds.max ?? ""} style={{ width: 168 }} label="Max Changelist" onChange={(ev, newValue) => {
+            newValue = newValue?.trim();
+            if (!newValue) {
+               newValue = undefined;
+            }
+            setChangelistBounds({ ...changelistBounds, max: newValue })
+         }} />
+      </Stack>
+   );
+});
\ No newline at end of file
diff --git a/Engine/Source/Programs/Horde/HordeDashboard/src/components/JobSearchSimple.tsx b/Engine/Source/Programs/Horde/HordeDashboard/src/components/JobSearchSimple.tsx
--- a/Engine/Source/Programs/Horde/HordeDashboard/src/components/JobSearchSimple.tsx
+++ b/Engine/Source/Programs/Horde/HordeDashboard/src/components/JobSearchSimple.tsx
@@ -4,7 +4,7 @@ import moment from "moment";
 import { useEffect, useMemo, useState } from "react";
 import { Link, useSearchParams } from "react-router-dom";
 import backend from "../backend";
-import { GetJobResponse, GetJobsTabResponse, GetJobStepRefResponse, JobData, JobQuery, TabType, GetTemplateRefResponse, GetJobsTabLabelColumnResponse, GetLabelStateResponse, LabelOutcome, LabelState, GetJobsTabColumnResponse } from "../backend/Api";
+import { GetJobResponse, GetJobsTabResponse, GetJobStepRefResponse, JobData, JobQuery, TabType, GetTemplateRefResponse, GetJobsTabLabelColumnResponse, GetLabelStateResponse, LabelOutcome, LabelState, GetJobsTabColumnResponse, JobStepOutcome } from "../backend/Api";
 import dashboard from "../backend/Dashboard";
 import { projectStore } from "../backend/ProjectStore";
 import templateCache from '../backend/TemplateCache';
@@ -46,7 +46,7 @@ export const JobSearchSimpleModal: React.FC<{ streamIdIn: string, changeListIn?:
       job: GetJobResponse
    };
 
-   const statuses = ["Waiting", "Running", "Complete", "Cancelled"]
+   const statuses = [ "Running", "Complete", "Succeeded", "Failed", "Waiting", "Cancelled"]
    const allStatuses = new Set<string>(statuses);
 
    const [searchState, setSearchState] = useState<SearchState>({ preflights: false, minCL: changeListIn, maxCL: changeListIn, templates: new Set(jobTab?.templates ?? []), statuses: allStatuses });
@@ -201,7 +201,7 @@ export const JobSearchSimpleModal: React.FC<{ streamIdIn: string, changeListIn?:
                id: jobIds.length ? jobIds : undefined,
                streamId: searchStreamId?.toLowerCase(),
                name: name,
-               filter: "id,labels,defaultLabel,name,change,createTime,streamId,preflightChange,startedByUserInfo,arguments,abortedByUserInfo,state",
+               filter: "id,labels,defaultLabel,name,change,createTime,streamId,preflightChange,startedByUserInfo,arguments,abortedByUserInfo,state,batches,cancellationReason",
                minChange: minCL,
                maxChange: maxCL,
                t: (!state.global && c_templates.length && !state.containsStep) ? c_templates.map(t => t.id) : undefined,
@@ -217,7 +217,14 @@ export const JobSearchSimpleModal: React.FC<{ streamIdIn: string, changeListIn?:
 
             if (c_statuses.length && c_statuses.length != allStatuses.size) {
                jobs = jobs.filter(j => {
-                  return (c_statuses.includes('Cancelled') && j.abortedByUserInfo?.id || c_statuses.includes(j.state) && !j.abortedByUserInfo?.id)
+                  const jobSucceeded = j.batches?.every(b => b.steps.every(s => s.outcome === JobStepOutcome.Success || s.outcome === JobStepOutcome.Warnings));
+
+                  return (
+                     c_statuses.includes('Cancelled') && j.abortedByUserInfo?.id // Cancelled
+                     || c_statuses.includes(j.state) && !j.abortedByUserInfo?.id // Running, Complete, Waiting
+                     || c_statuses.includes("Failed") && !jobSucceeded // Failed
+                     || c_statuses.includes("Succeeded") && jobSucceeded // Succeeded
+                  )
                })
             }
 
diff --git a/Engine/Source/Programs/Horde/HordeDashboard/src/components/JobViewAll.tsx b/Engine/Source/Programs/Horde/HordeDashboard/src/components/JobViewAll.tsx
--- a/Engine/Source/Programs/Horde/HordeDashboard/src/components/JobViewAll.tsx
+++ b/Engine/Source/Programs/Horde/HordeDashboard/src/components/JobViewAll.tsx
@@ -93,8 +93,7 @@ const JobViewAllInner: React.FC<{ filter: JobFilterSimple }> = observer(({ filte
       return <div />;
    }
 
-   let jobs = jobHandler.getFilteredJobs(jobFilter.status);
-
+   let jobs = jobHandler.getFilteredJobs(jobFilter.status, jobFilter.minCL, jobFilter.maxCL);
    jobs = jobs.filter(j => filterJob(j, filter.filterKeyword));
 
    let jobItems: JobItem[] = jobs.map(j => {
@@ -113,7 +112,6 @@ const JobViewAllInner: React.FC<{ filter: JobFilterSimple }> = observer(({ filte
 
    });
 
-
    let columns: IColumn[] = [
       { key: 'jobview_column1', name: 'Status', minWidth: 16, maxWidth: 16, onRenderHeader: () => null },
       { key: 'jobview_column2', name: 'Change', minWidth: 64, maxWidth: 64 },
diff --git a/Engine/Source/Programs/Horde/HordeDashboard/src/index.css b/Engine/Source/Programs/Horde/HordeDashboard/src/index.css
--- a/Engine/Source/Programs/Horde/HordeDashboard/src/index.css
+++ b/Engine/Source/Programs/Horde/HordeDashboard/src/index.css
@@ -61,8 +61,7 @@ code {
   user-select: none;
 }
 
-.cl-callout-button,
-.cl-callout-button:hover {
+.cl-callout-button {
   text-decoration: none;
   white-space: nowrap;
   user-select: none;
@@ -71,8 +70,11 @@ code {
   border-width: 0;
 }
 
-.cl-callout-button-user,
-.cl-callout-button-user:hover {
+.cl-callout-button:hover, .cl-callout-button-user:hover {
+  filter: brightness(1.2);
+}
+
+.cl-callout-button-user {
   text-decoration: none;
   white-space: nowrap;
   user-select: none;
