Horde Dashboard: cancellation reason on user home

diff --git a/Engine/Source/Programs/Horde/HordeDashboard/src/components/UserHome.tsx b/Engine/Source/Programs/Horde/HordeDashboard/src/components/UserHome.tsx
--- a/Engine/Source/Programs/Horde/HordeDashboard/src/components/UserHome.tsx
+++ b/Engine/Source/Programs/Horde/HordeDashboard/src/components/UserHome.tsx
@@ -411,7 +411,7 @@ class UserJobsHandler {
       const maxJobs = 100;
 
       try {
-         let filter = "id,streamId,name,change,preflightChange,templateId,templateHash,graphHash,startedByUserInfo,createTime,state,arguments,updateTime,labels,defaultLabel,batches,autoSubmit,autoSubmitChange,preflightDescription,abortedByUserInfo";
+         let filter = "id,streamId,name,change,preflightChange,templateId,templateHash,graphHash,startedByUserInfo,createTime,state,arguments,updateTime,labels,defaultLabel,batches,autoSubmit,autoSubmitChange,preflightDescription,abortedByUserInfo,cancellationReason";
 
          const query: JobQuery = {
             filter: filter,
