Horde: Add tags to the StreamConfig to control whether scheduled jobs are enabled

diff --git a/Engine/Source/Programs/Horde/HordeServer.Shared/Plugins/IPluginConfig.cs b/Engine/Source/Programs/Horde/HordeServer.Shared/Plugins/IPluginConfig.cs
--- a/Engine/Source/Programs/Horde/HordeServer.Shared/Plugins/IPluginConfig.cs
+++ b/Engine/Source/Programs/Horde/HordeServer.Shared/Plugins/IPluginConfig.cs
@@ -1,6 +1,7 @@
 // Copyright Epic Games, Inc. All Rights Reserved.
 
 using HordeServer.Acls;
+using Microsoft.Extensions.Logging;
 
 namespace HordeServer.Plugins
 {
@@ -19,7 +20,7 @@ namespace HordeServer.Plugins
 	/// <summary>
 	/// Options passed to <see cref="IPluginConfig.PostLoad(PluginConfigOptions)"/>
 	/// </summary>
-	public record class PluginConfigOptions(ConfigVersion Version, IEnumerable<IPluginConfig> Plugins, AclConfig ParentAcl, IServiceProvider? ServiceProvider = null);
+	public record class PluginConfigOptions(ConfigVersion Version, IEnumerable<IPluginConfig> Plugins, AclConfig ParentAcl, IServiceProvider? ServiceProvider = null, ILogger? Logger = null);
 
 	/// <summary>
 	/// Empty implementation of <see cref="IPluginConfig"/>
diff --git a/Engine/Source/Programs/Horde/HordeServer.Tests.Shared/ServerTestSetup.cs b/Engine/Source/Programs/Horde/HordeServer.Tests.Shared/ServerTestSetup.cs
--- a/Engine/Source/Programs/Horde/HordeServer.Tests.Shared/ServerTestSetup.cs
+++ b/Engine/Source/Programs/Horde/HordeServer.Tests.Shared/ServerTestSetup.cs
@@ -44,6 +44,7 @@ namespace HordeServer.Tests
 		public OpenTelemetry.Trace.Tracer Tracer => ServiceProvider.GetRequiredService<OpenTelemetry.Trace.Tracer>();
 		public Meter Meter => ServiceProvider.GetRequiredService<Meter>();
 		public ConfigService ConfigService => ServiceProvider.GetRequiredService<ConfigService>();
+		public ILogger<ConfigService> ConfigServiceLogger => ServiceProvider.GetRequiredService<ILogger<ConfigService>>();
 		public IOptionsMonitor<GlobalConfig> GlobalConfig => ServiceProvider.GetRequiredService<IOptionsMonitor<GlobalConfig>>();
 		public IOptionsSnapshot<GlobalConfig> GlobalConfigSnapshot => ServiceProvider.GetRequiredService<IOptionsSnapshot<GlobalConfig>>();
 
@@ -75,7 +76,7 @@ namespace HordeServer.Tests
 		protected async Task SetConfigAsync(GlobalConfig globalConfig)
 		{
 			await ConfigService.WaitForInitialConfigAsync();
-			globalConfig.PostLoad(ServerSettings, _pluginCollection.LoadedPlugins, DefaultAclModifiers, ServiceProvider);
+			globalConfig.PostLoad(ServerSettings, _pluginCollection.LoadedPlugins, DefaultAclModifiers, ServiceProvider, ConfigServiceLogger);
 			ConfigService.OverrideConfig(globalConfig);
 		}
 
@@ -84,7 +85,7 @@ namespace HordeServer.Tests
 			await ConfigService.WaitForInitialConfigAsync();
 			GlobalConfig globalConfig = GlobalConfig.CurrentValue;
 			action(globalConfig);
-			globalConfig.PostLoad(ServerSettings, _pluginCollection.LoadedPlugins, DefaultAclModifiers, ServiceProvider);
+			globalConfig.PostLoad(ServerSettings, _pluginCollection.LoadedPlugins, DefaultAclModifiers, ServiceProvider, ConfigServiceLogger);
 			ConfigService.OverrideConfig(globalConfig);
 		}
 
diff --git a/Engine/Source/Programs/Horde/HordeServer/Configuration/ConfigService.cs b/Engine/Source/Programs/Horde/HordeServer/Configuration/ConfigService.cs
--- a/Engine/Source/Programs/Horde/HordeServer/Configuration/ConfigService.cs
+++ b/Engine/Source/Programs/Horde/HordeServer/Configuration/ConfigService.cs
@@ -285,7 +285,7 @@ namespace HordeServer.Configuration
 			try
 			{
 				GlobalConfig globalConfig = await ReadGlobalConfigAsync(context, cancellationToken);
-				globalConfig.PostLoad(_serverSettings, _pluginCollection.LoadedPlugins, _defaultAclModifiers, _serviceProvider);
+				globalConfig.PostLoad(_serverSettings, _pluginCollection.LoadedPlugins, _defaultAclModifiers, _serviceProvider, _logger);
 
 				foreach (OverrideConfigFile file in overrideFiles.Values)
 				{
@@ -611,7 +611,7 @@ namespace HordeServer.Configuration
 				}
 
 				// Execute a PostLoad before returning so we can validate that everything is valid
-				globalConfig.PostLoad(_serverSettings, _pluginCollection.LoadedPlugins, _defaultAclModifiers, _serviceProvider);
+				globalConfig.PostLoad(_serverSettings, _pluginCollection.LoadedPlugins, _defaultAclModifiers, _serviceProvider, _logger);
 
 				return snapshot;
 			}
@@ -860,7 +860,7 @@ namespace HordeServer.Configuration
 			globalConfig.Revision = IoHash.Compute(data.Span).ToString();
 
 			// Run the postload callbacks on all the config objects
-			globalConfig.PostLoad(_serverSettings, _pluginCollection.LoadedPlugins, _defaultAclModifiers, _serviceProvider);
+			globalConfig.PostLoad(_serverSettings, _pluginCollection.LoadedPlugins, _defaultAclModifiers, _serviceProvider, _logger);
 			return globalConfig;
 		}
 	}
diff --git a/Engine/Source/Programs/Horde/HordeServer/Server/GlobalConfig.cs b/Engine/Source/Programs/Horde/HordeServer/Server/GlobalConfig.cs
--- a/Engine/Source/Programs/Horde/HordeServer/Server/GlobalConfig.cs
+++ b/Engine/Source/Programs/Horde/HordeServer/Server/GlobalConfig.cs
@@ -16,6 +16,7 @@ using HordeServer.Plugins;
 using HordeServer.Server.Notices;
 using HordeServer.ServiceAccounts;
 using HordeServer.Utilities;
+using Microsoft.Extensions.Logging;
 
 namespace HordeServer.Server
 {
@@ -104,7 +105,7 @@ namespace HordeServer.Server
 		/// <summary>
 		/// Called after the config file has been read
 		/// </summary>
-		public void PostLoad(ServerSettings serverSettings, IReadOnlyList<ILoadedPlugin> loadedPlugins, IEnumerable<IDefaultAclModifier> defaultAclModifiers, IServiceProvider? serviceProvider = null)
+		public void PostLoad(ServerSettings serverSettings, IReadOnlyList<ILoadedPlugin> loadedPlugins, IEnumerable<IDefaultAclModifier> defaultAclModifiers, IServiceProvider? serviceProvider = null, ILogger? logger = null)
 		{
 			ServerSettings = serverSettings;
 			AclAction[] aclActions = AclConfig.GetActions(
@@ -154,7 +155,7 @@ namespace HordeServer.Server
 				pluginConfigs = sortedPluginConfigs;
 			}
 
-			PluginConfigOptions pluginConfigOptions = new PluginConfigOptions(VersionEnum, Plugins.Values, Acl, serviceProvider);
+			PluginConfigOptions pluginConfigOptions = new PluginConfigOptions(VersionEnum, Plugins.Values, Acl, serviceProvider, logger);
 			foreach (IPluginConfig pluginConfig in pluginConfigs)
 			{
 				pluginConfig.PostLoad(pluginConfigOptions);
diff --git a/Engine/Source/Programs/Horde/Plugins/Build/HordeServer.Build.Tests/Jobs/Schedules/SchedulerTests.cs b/Engine/Source/Programs/Horde/Plugins/Build/HordeServer.Build.Tests/Jobs/Schedules/SchedulerTests.cs
--- a/Engine/Source/Programs/Horde/Plugins/Build/HordeServer.Build.Tests/Jobs/Schedules/SchedulerTests.cs
+++ b/Engine/Source/Programs/Horde/Plugins/Build/HordeServer.Build.Tests/Jobs/Schedules/SchedulerTests.cs
@@ -1,6 +1,8 @@
 // Copyright Epic Games, Inc. All Rights Reserved.
 
+using System.Text.Json;
 using EpicGames.Horde.Commits;
+using EpicGames.Horde.Common;
 using EpicGames.Horde.Jobs;
 using EpicGames.Horde.Jobs.Graphs;
 using EpicGames.Horde.Jobs.Schedules;
@@ -15,6 +17,7 @@ using HordeServer.Projects;
 using HordeServer.Server;
 using HordeServer.Streams;
 using HordeServer.Users;
+using HordeServer.Utilities;
 using Moq;
 
 namespace HordeServer.Tests.Jobs.Schedules
@@ -52,7 +55,7 @@ namespace HordeServer.Tests.Jobs.Schedules
 			PerforceService.AddChange(StreamId, 102, bob, "", new[] { "content.uasset" });
 		}
 
-		async Task<IStream> SetScheduleAsync(ScheduleConfig schedule)
+		async Task<IStream> SetScheduleAsync(ScheduleConfig schedule, List<StreamTag>? streamTags = null)
 		{
 			await ScheduleService.ResetAsync();
 
@@ -62,6 +65,11 @@ namespace HordeServer.Tests.Jobs.Schedules
 			streamConfig.Tabs.Add(new TabConfig { Title = "foo", Templates = new List<TemplateId> { TemplateId } });
 			streamConfig.Templates.Add(new TemplateRefConfig { Id = TemplateId, Name = "Test", Schedule = schedule });
 
+			if (streamTags != null)
+			{
+				streamConfig.StreamTags = streamTags;
+			}
+
 			await UpdateConfigAsync(x => x.Plugins.GetBuildConfig().Projects[0].Streams = new List<StreamConfig> { streamConfig });
 
 			return (await StreamCollection.GetAsync(streamConfig.Id))!;
@@ -736,5 +744,147 @@ namespace HordeServer.Tests.Jobs.Schedules
 			_initialJobIds.UnionWith(jobs.Select(x => x.Id));
 			return jobs.OrderBy(x => x.CommitId).ToList();
 		}
+
+		[TestMethod]
+		public async Task ScheduleWithConditionEnabledTestAsync()
+		{
+			List<StreamTag> tags = [
+				new() { Name = "A", Enabled = true }
+			];
+			ScheduleConfig schedule = new()
+			{
+				Condition = Condition.Parse("A == true")
+			};
+			await SetScheduleAsync(schedule, tags);
+			Assert.IsTrue(schedule.Enabled);
+		}
+
+		[TestMethod]
+		public async Task ScheduleWithConditionDisabledTestAsync()
+		{
+			List<StreamTag> tags = [
+				new() { Name = "A", Enabled = false }
+			];
+			ScheduleConfig schedule = new()
+			{
+				Condition = Condition.Parse("A == true")
+			};
+			await SetScheduleAsync(schedule, tags);
+			Assert.IsFalse(schedule.Enabled);
+		}
+
+		[TestMethod]
+		public async Task ScheduleWithConditionDisabledFromJsonTestAsync()
+		{
+			string json = """
+			{
+				"id": "ue5-main",
+				"name": "//UE5/Main",
+				"streamTags": [
+					{
+						"name": "A",
+						"enabled": false
+					}
+				],
+				"tabs": [
+					{
+						"title": "foo",
+						"templates": ["template1"]
+					}
+				],
+				"templates": [
+					{
+						"id": "template1",
+						"name": "Test",
+						"schedule": {
+							"condition": "A == true"
+						}
+					}
+				]
+			}
+			""";
+			StreamConfig streamConfig = JsonSerializer.Deserialize<StreamConfig>(json, JsonUtils.DefaultSerializerOptions)!;
+			await UpdateConfigAsync(x => x.Plugins.GetBuildConfig().Projects[0].Streams = [streamConfig]);
+			Assert.IsFalse(streamConfig.Templates[0].Schedule?.Enabled);
+		}
+
+		[TestMethod]
+		public async Task ScheduleBackwardsCompatibilityDisabledTestAsync()
+		{
+			ScheduleConfig schedule = new()
+			{
+				Condition = Condition.Parse("true"),
+				// For backwards compatibility the Enabled property remains in use
+				Enabled = false
+			};
+			await SetScheduleAsync(schedule);
+			Assert.IsFalse(schedule.Enabled);
+		}
+
+		[TestMethod]
+		[DataRow(true, DisplayName = "Schedule enabled by default with empty condition")]
+		[DataRow(false, DisplayName = "Schedule enabled by default with no condition")]
+		public async Task ScheduleEnabledByDefaultValuesTestAsync(bool withCondition)
+		{
+			ScheduleConfig schedule = new();
+			if (withCondition)
+			{
+				schedule.Condition = Condition.Parse(String.Empty);
+			}
+			await SetScheduleAsync(schedule);
+			Assert.IsTrue(schedule.Enabled);
+		}
+
+		[TestMethod]
+		[ExpectedException(typeof(ConditionException))]
+		public async Task ScheduleWithMissingTagsTestAsync()
+		{
+			ScheduleConfig schedule = new()
+			{
+				Condition = Condition.Parse("A == true")
+			};
+			await SetScheduleAsync(schedule);
+		}
+
+		[TestMethod]
+		[DataRow(true, true, DisplayName = "Multiple tags and scheduled is enabled")]
+		[DataRow(true, false, DisplayName = "Multiple tags and schedule is disabled")]
+		public async Task ScheduleWithMultipleTagsFromJsonTestAsync(bool isShipped, bool isCutOver)
+		{
+			string json = $$"""
+			{
+				"id": "game1",
+				"name": "//depot/game1",
+				"streamTags": [
+					{
+						"name": "isShipped",
+						"enabled": {{isShipped.ToString().ToLower()}}
+					},
+					{
+						"name": "isCutOver",
+						"enabled": {{isCutOver.ToString().ToLower()}}
+					}
+				],
+				"tabs": [
+					{
+						"title": "title",
+						"templates": ["template1"]
+					}
+				],
+				"templates": [
+					{
+						"id": "template1",
+						"name": "template1",
+						"schedule": {
+							"condition": "(isShipped == true) && (isCutOver == true)"
+						}
+					}
+				]
+			}
+			""";
+			StreamConfig streamConfig = JsonSerializer.Deserialize<StreamConfig>(json, JsonUtils.DefaultSerializerOptions)!;
+			await UpdateConfigAsync(x => x.Plugins.GetBuildConfig().Projects[0].Streams = [streamConfig]);
+			Assert.AreEqual(isShipped && isCutOver, streamConfig.Templates[0].Schedule?.Enabled);
+		}
 	}
 }
diff --git a/Engine/Source/Programs/Horde/Plugins/Build/HordeServer.Build/BuildConfig.cs b/Engine/Source/Programs/Horde/Plugins/Build/HordeServer.Build/BuildConfig.cs
--- a/Engine/Source/Programs/Horde/Plugins/Build/HordeServer.Build/BuildConfig.cs
+++ b/Engine/Source/Programs/Horde/Plugins/Build/HordeServer.Build/BuildConfig.cs
@@ -111,7 +111,7 @@ namespace HordeServer
 			foreach (ProjectConfig project in Projects)
 			{
 				_projectLookup.Add(project.Id, project);
-				project.PostLoad(project.Id, configOptions.ParentAcl, ArtifactTypes);
+				project.PostLoad(project.Id, configOptions.ParentAcl, ArtifactTypes, configOptions.Logger);
 
 				foreach (StreamConfig stream in project.Streams)
 				{
diff --git a/Engine/Source/Programs/Horde/Plugins/Build/HordeServer.Build/Jobs/Schedules/ScheduleConfig.cs b/Engine/Source/Programs/Horde/Plugins/Build/HordeServer.Build/Jobs/Schedules/ScheduleConfig.cs
--- a/Engine/Source/Programs/Horde/Plugins/Build/HordeServer.Build/Jobs/Schedules/ScheduleConfig.cs
+++ b/Engine/Source/Programs/Horde/Plugins/Build/HordeServer.Build/Jobs/Schedules/ScheduleConfig.cs
@@ -2,6 +2,7 @@
 
 using System.ComponentModel.DataAnnotations;
 using EpicGames.Horde.Commits;
+using EpicGames.Horde.Common;
 using EpicGames.Horde.Jobs.Schedules;
 using EpicGames.Horde.Jobs.Templates;
 using HordeServer.Streams;
@@ -18,6 +19,13 @@ namespace HordeServer.Jobs.Schedules
 		/// </summary>
 		public bool Enabled { get; set; } = true;
 
+		/// <summary>
+		/// Condition for whether the schedule should be enabled.
+		/// When a condition is provided the result of the expression will be logically AND with the <see cref="Enabled"/> property to decide
+		/// if the schedule should be enabled.
+		/// </summary>
+		public Condition? Condition { get; set; }
+
 		/// <summary>
 		/// Maximum number of builds that can be active at once
 		/// </summary>
diff --git a/Engine/Source/Programs/Horde/Plugins/Build/HordeServer.Build/Projects/ProjectConfig.cs b/Engine/Source/Programs/Horde/Plugins/Build/HordeServer.Build/Projects/ProjectConfig.cs
--- a/Engine/Source/Programs/Horde/Plugins/Build/HordeServer.Build/Projects/ProjectConfig.cs
+++ b/Engine/Source/Programs/Horde/Plugins/Build/HordeServer.Build/Projects/ProjectConfig.cs
@@ -17,6 +17,7 @@ using HordeServer.Devices;
 using HordeServer.Logs;
 using HordeServer.Streams;
 using HordeServer.Utilities;
+using Microsoft.Extensions.Logging;
 
 #pragma warning disable CA2227 // Change x to be read-only by removing the property setter
 
@@ -127,7 +128,8 @@ namespace HordeServer.Projects
 		/// <param name="id">Id of this project</param>
 		/// <param name="parentAcl">The owning global config object</param>
 		/// <param name="parentArtifactTypes">Set of artifact types</param>
-		public void PostLoad(ProjectId id, AclConfig parentAcl, IEnumerable<ArtifactTypeConfig> parentArtifactTypes)
+		/// <param name="logger">Logger for messages</param>
+		public void PostLoad(ProjectId id, AclConfig parentAcl, IEnumerable<ArtifactTypeConfig> parentArtifactTypes, ILogger? logger)
 		{
 			Id = id;
 			AclAction[] aclActions = AclConfig.GetActions([typeof(ArtifactAclAction), typeof(DeviceAclAction), typeof(LogAclAction), typeof(ProjectAclAction)]);
@@ -147,7 +149,7 @@ namespace HordeServer.Projects
 
 			foreach (StreamConfig stream in Streams)
 			{
-				stream.PostLoad(stream.Id, this, _artifactTypeLookup.Values);
+				stream.PostLoad(stream.Id, this, _artifactTypeLookup.Values, logger);
 			}
 		}
 	}
diff --git a/Engine/Source/Programs/Horde/Plugins/Build/HordeServer.Build/Streams/StreamConfig.cs b/Engine/Source/Programs/Horde/Plugins/Build/HordeServer.Build/Streams/StreamConfig.cs
--- a/Engine/Source/Programs/Horde/Plugins/Build/HordeServer.Build/Streams/StreamConfig.cs
+++ b/Engine/Source/Programs/Horde/Plugins/Build/HordeServer.Build/Streams/StreamConfig.cs
@@ -13,6 +13,7 @@ using EpicGames.Horde.Acls;
 using EpicGames.Horde.Agents.Pools;
 using EpicGames.Horde.Artifacts;
 using EpicGames.Horde.Commits;
+using EpicGames.Horde.Common;
 using EpicGames.Horde.Issues;
 using EpicGames.Horde.Jobs;
 using EpicGames.Horde.Jobs.Graphs;
@@ -34,6 +35,7 @@ using HordeServer.Notifications;
 using HordeServer.Projects;
 using HordeServer.Replicators;
 using HordeServer.Utilities;
+using Microsoft.Extensions.Logging;
 
 #pragma warning disable CA2227 // Change x to be read-only by removing the property setter
 
@@ -78,6 +80,22 @@ namespace HordeServer.Streams
 		Full,
 	}
 
+	/// <summary>
+	/// Tags for a stream
+	/// </summary>
+	public class StreamTag
+	{
+		/// <summary>
+		/// Name of the tag
+		/// </summary>
+		public string Name { get; set; } = String.Empty;
+
+		/// <summary>
+		/// Whether the tag should be enabled
+		/// </summary>
+		public bool Enabled { get; set; }
+	}
+
 	/// <summary>
 	/// Config for a stream
 	/// </summary>
@@ -257,6 +275,11 @@ namespace HordeServer.Streams
 		/// </summary>
 		public List<ArtifactTypeConfig> ArtifactTypes { get; set; } = new List<ArtifactTypeConfig>();
 
+		/// <summary>
+		/// Tags for the stream
+		/// </summary>
+		public List<StreamTag> StreamTags { get; set; } = [];
+
 		readonly Dictionary<ArtifactType, ArtifactTypeConfig> _artifactTypeLookup = new Dictionary<ArtifactType, ArtifactTypeConfig>();
 
 		/// <inheritdoc cref="AclConfig.Authorize(AclAction, ClaimsPrincipal)"/>
@@ -269,7 +292,8 @@ namespace HordeServer.Streams
 		/// <param name="id">The stream id</param>
 		/// <param name="projectConfig">Owning project</param>
 		/// <param name="parentArtifactTypes">Parent artifact types</param>
-		public void PostLoad(StreamId id, ProjectConfig projectConfig, IEnumerable<ArtifactTypeConfig> parentArtifactTypes)
+		/// <param name="logger">Logger for messages</param>
+		public void PostLoad(StreamId id, ProjectConfig projectConfig, IEnumerable<ArtifactTypeConfig> parentArtifactTypes, ILogger? logger)
 		{
 			Id = id;
 			ProjectConfig = projectConfig;
@@ -320,6 +344,12 @@ namespace HordeServer.Streams
 				}
 			}
 
+			Dictionary<string, string> scheduleConditionPropertyLookup = new();
+			foreach (StreamTag streamTag in StreamTags)
+			{
+				scheduleConditionPropertyLookup.TryAdd(streamTag.Name, streamTag.Enabled.ToString());
+			}
+
 			foreach (TemplateRefConfig template in Templates)
 			{
 				ScheduleConfig? schedule = template.Schedule;
@@ -348,6 +378,25 @@ namespace HordeServer.Streams
 							throw new InvalidOperationException($"Missing definition for commit tag '{commitTag}' referenced by {Id}:{template.Id}");
 						}
 					}
+
+					if (schedule.Condition != null)
+					{
+						bool condition = schedule.Condition.Evaluate(name =>
+						{
+							if (scheduleConditionPropertyLookup.TryGetValue(name, out string? enabled))
+							{
+								return [enabled];
+							}
+							throw new ConditionException($"Unknown property '{name}' in schedule condition '{schedule.Condition}'");
+						});
+
+						if (schedule.Enabled)
+						{
+							logger?.LogInformation("Schedule for template {TemplateId} is {Enabled} given the condition '{Condition}'", template.Id, condition ? "enabled" : "disabled", schedule.Condition);
+						}
+
+						schedule.Enabled &= condition;
+					}
 				}
 			}
 
diff --git a/Engine/Source/Programs/Shared/EpicGames.Horde/Common/Condition.cs b/Engine/Source/Programs/Shared/EpicGames.Horde/Common/Condition.cs
--- a/Engine/Source/Programs/Shared/EpicGames.Horde/Common/Condition.cs
+++ b/Engine/Source/Programs/Shared/EpicGames.Horde/Common/Condition.cs
@@ -519,7 +519,7 @@ namespace EpicGames.Horde.Common
 		}
 
 		/// <summary>
-		/// Evaluate the condition using the given callback to retreive property values
+		/// Evaluate the condition using the given callback to retrieve property values
 		/// </summary>
 		/// <param name="getPropertyValues"></param>
 		/// <returns></returns>
