Horde: Add missing capabilities data for Macs

diff --git a/Engine/Source/Programs/Horde/HordeAgent/Services/CapabilitiesService.cs b/Engine/Source/Programs/Horde/HordeAgent/Services/CapabilitiesService.cs
--- a/Engine/Source/Programs/Horde/HordeAgent/Services/CapabilitiesService.cs
+++ b/Engine/Source/Programs/Horde/HordeAgent/Services/CapabilitiesService.cs
@@ -230,7 +230,7 @@ namespace HordeAgent.Services
 				using (Process process = new Process())
 				{
 					process.StartInfo.FileName = "system_profiler";
-					process.StartInfo.Arguments = "SPHardwareDataType SPSoftwareDataType -xml";
+					process.StartInfo.Arguments = "SPHardwareDataType SPSoftwareDataType SPDisplaysDataType -xml";
 					process.StartInfo.CreateNoWindow = true;
 					process.StartInfo.RedirectStandardOutput = true;
 					process.StartInfo.RedirectStandardInput = false;
@@ -252,12 +252,23 @@ namespace HordeAgent.Services
 						capabilities.Properties.Add($"Model={model.InnerText}");
 					}
 
-					XmlNode? cpuTypeNode = hardwareNode.SelectSingleNode("key[. = 'cpu_type']/following-sibling::string");
-					XmlNode? cpuSpeedNode = hardwareNode.SelectSingleNode("key[. = 'current_processor_speed']/following-sibling::string");
-					XmlNode? cpuPackagesNode = hardwareNode.SelectSingleNode("key[. = 'packages']/following-sibling::integer");
-					if (cpuTypeNode != null && cpuSpeedNode != null && cpuPackagesNode != null)
+					if (RuntimeInformation.OSArchitecture == Architecture.Arm64)
+					{
+						XmlNode? chipTypeNode = hardwareNode.SelectSingleNode("key[. = 'chip_type']/following-sibling::string");
+						if (chipTypeNode != null)
+						{
+							capabilities.Properties.Add($"CPU={chipTypeNode.InnerText}");
+						}
+					}
+					else
 					{
-						capabilities.Properties.Add((cpuPackagesNode.InnerText != "1") ? $"CPU={cpuPackagesNode.InnerText} x {cpuTypeNode.InnerText} @ {cpuSpeedNode.InnerText}" : $"CPU={cpuTypeNode.InnerText} @ {cpuSpeedNode.InnerText}");
+						XmlNode? cpuTypeNode = hardwareNode.SelectSingleNode("key[. = 'cpu_type']/following-sibling::string");
+						XmlNode? cpuSpeedNode = hardwareNode.SelectSingleNode("key[. = 'current_processor_speed']/following-sibling::string");
+						XmlNode? cpuPackagesNode = hardwareNode.SelectSingleNode("key[. = 'packages']/following-sibling::integer");
+						if (cpuTypeNode != null && cpuSpeedNode != null && cpuPackagesNode != null)
+						{
+							capabilities.Properties.Add((cpuPackagesNode.InnerText != "1") ? $"CPU={cpuPackagesNode.InnerText} x {cpuTypeNode.InnerText} @ {cpuSpeedNode.InnerText}" : $"CPU={cpuTypeNode.InnerText} @ {cpuSpeedNode.InnerText}");
+						}
 					}
 
 					int numLogicalCores = Environment.ProcessorCount;
@@ -267,6 +278,7 @@ namespace HordeAgent.Services
 					}
 					numLogicalCores = (int)(numLogicalCores * _settings.CurrentValue.CpuMultiplier);
 					capabilities.Resources.Add(KnownResourceNames.LogicalCores, numLogicalCores);
+					capabilities.Properties.Add($"{KnownResourceNames.LogicalCores}={numLogicalCores}");
 
 					if (RuntimeInformation.OSArchitecture == Architecture.Arm64)
 					{
@@ -278,7 +290,9 @@ namespace HordeAgent.Services
 							{
 								if (Int32.TryParse(parts[1], out int performanceCores) && Int32.TryParse(parts[2], out int efficiencyCores))
 								{
-									capabilities.Properties.Add($"PhysicalCores={performanceCores}");
+									int physicalCores = performanceCores + efficiencyCores;
+									capabilities.Properties.Add($"PhysicalCores={physicalCores}");
+									capabilities.Properties.Add($"PerformanceCores={performanceCores}");
 									capabilities.Properties.Add($"EfficiencyCores={efficiencyCores}");
 								}
 							}
@@ -304,6 +318,15 @@ namespace HordeAgent.Services
 						}
 					}
 				}
+				XmlNode? displayNode = xml.SelectSingleNode("/plist/array/dict/key[. = '_dataType']/following-sibling::string[. = 'SPDisplaysDataType']/../key[. = '_items']/following-sibling::array/dict");
+				if (displayNode != null)
+				{
+					XmlNode? gpuTypeNode = displayNode.SelectSingleNode("key[. = 'sppci_model']/following-sibling::string");
+					if (gpuTypeNode != null)
+					{
+						capabilities.Properties.Add($"GPU={gpuTypeNode.InnerText}");
+					}
+				}
 
 				XmlNode? softwareNode = xml.SelectSingleNode("/plist/array/dict/key[. = '_dataType']/following-sibling::string[. = 'SPSoftwareDataType']/../key[. = '_items']/following-sibling::array/dict");
 				if (softwareNode != null)
@@ -320,8 +343,14 @@ namespace HordeAgent.Services
 						capabilities.Properties.Add($"OSKernelVersion={kernelVersionNode.InnerText}");
 					}
 				}
+
+				// Add session information
+				capabilities.Properties.Add($"User={Environment.UserName}");
+				capabilities.Properties.Add($"Domain={Environment.UserDomainName}");
+				capabilities.Properties.Add($"Interactive={Environment.UserInteractive}");
+				capabilities.Properties.Add($"Elevated={IsUserAdministrator()}");
 			}
-			
+
 			// Add EC2 properties if needed
 			if (_settings.CurrentValue.EnableAwsEc2Support)
 			{
