Horde: Display job search button in presubmit tab

diff --git a/Engine/Source/Programs/Horde/HordeDashboard/src/components/StreamView.tsx b/Engine/Source/Programs/Horde/HordeDashboard/src/components/StreamView.tsx
--- a/Engine/Source/Programs/Horde/HordeDashboard/src/components/StreamView.tsx
+++ b/Engine/Source/Programs/Horde/HordeDashboard/src/components/StreamView.tsx
@@ -480,7 +480,7 @@ const StreamViewInner: React.FC = observer(() => {
    const vw = Math.max(document.documentElement.clientWidth, window.innerWidth || 0);
 
    const newBuildTab = queryTab?.toLowerCase() === "summary" ? "all" : queryTab;
-   const isSwarmTab = queryTab?.toLowerCase() === "swarm" || queryTab?.toLowerCase() === "presubmit";
+   const isSwarmTab = queryTab?.toLowerCase() === "swarm";
    let isIncrementalTab = currentTab?.style === TabStyle.Compact;
 
    const windowWidth = windowSize.width;
