Horde: Remove legacy isSwarmTab variable and conditional displaying of job search button altogether

diff --git a/Engine/Source/Programs/Horde/HordeDashboard/src/components/StreamView.tsx b/Engine/Source/Programs/Horde/HordeDashboard/src/components/StreamView.tsx
--- a/Engine/Source/Programs/Horde/HordeDashboard/src/components/StreamView.tsx
+++ b/Engine/Source/Programs/Horde/HordeDashboard/src/components/StreamView.tsx
@@ -480,7 +480,6 @@ const StreamViewInner: React.FC = observer(() => {
    const vw = Math.max(document.documentElement.clientWidth, window.innerWidth || 0);
 
    const newBuildTab = queryTab?.toLowerCase() === "summary" ? "all" : queryTab;
-   const isSwarmTab = queryTab?.toLowerCase() === "swarm";
    let isIncrementalTab = currentTab?.style === TabStyle.Compact;
 
    const windowWidth = windowSize.width;
@@ -538,13 +537,13 @@ const StreamViewInner: React.FC = observer(() => {
                      <Stack grow />
                      <Stack horizontal verticalAlign="center" horizontalAlign={"end"} tokens={{ childrenGap: 8 }}>
                         <Stack horizontal tokens={{ childrenGap: 18 }}>
-                           {!isSwarmTab && <DefaultButton
+                           <DefaultButton
                               styles={{ root: { fontFamily: "Horde Open Sans SemiBold !important" } }}
                               text="Search"
                               split={!isSummary}
                               menuProps={!isSummary ? findJobsItems : undefined}
                               onClick={() => setFindJobsShown(true)}
-                           />}
+                           />
                            <PrimaryButton
                               styles={{ root: { fontFamily: "Horde Open Sans SemiBold !important" } }}
                               onClick={() => { setShown(true); }}>
