Horde: Truncate long template names in new build modal dropdown

diff --git a/Engine/Source/Programs/Horde/HordeDashboard/src/components/build/NewBuildV2.tsx b/Engine/Source/Programs/Horde/HordeDashboard/src/components/build/NewBuildV2.tsx
--- a/Engine/Source/Programs/Horde/HordeDashboard/src/components/build/NewBuildV2.tsx
+++ b/Engine/Source/Programs/Horde/HordeDashboard/src/components/build/NewBuildV2.tsx
@@ -1599,9 +1599,15 @@ const TemplateSelector: React.FC = () => {
       items: templateOptions,
    };
 
+   let templateName = options?.template?.name;
+   if (templateName && templateName?.length > 72) {
+      templateName = templateName.slice(0, 72);
+      templateName += "...";
+   }
+
    return <Stack>
       <Label>Template</Label>
-      <DefaultButton style={{ width: 280, textAlign: "left" }} disabled={options.readOnly} menuProps={templateMenuProps} text={options?.template?.name ?? "Error"} />
+      <DefaultButton style={{ width: 280, textAlign: "left" }} disabled={options.readOnly} menuProps={templateMenuProps} text={templateName ?? "Error"} />
    </Stack>
 }
 
