Horde: Updating schema documentation

#rnx

diff --git a/Engine/Source/Programs/Horde/Docs/Config/Schema/Streams.md b/Engine/Source/Programs/Horde/Docs/Config/Schema/Streams.md
--- a/Engine/Source/Programs/Horde/Docs/Config/Schema/Streams.md
+++ b/Engine/Source/Programs/Horde/Docs/Config/Schema/Streams.md
@@ -36,6 +36,7 @@ Name | Description
 `workflows` | [WorkflowConfig](#workflowconfig)`[]`<br>Workflows for dealing with new issues
 `tokens` | [TokenConfig](#tokenconfig)`[]`<br>Tokens to create for each job step
 `artifactTypes` | [ArtifactTypeConfig](#artifacttypeconfig)`[]`<br>Permissions for artifact types
+`streamTags` | [StreamTag](#streamtag)`[]`<br>Tags for the stream
 
 ## ConfigInclude
 
@@ -251,6 +252,7 @@ Parameters to create a new schedule
 Name | Description
 ---- | -----------
 `enabled` | `boolean`<br>Whether the schedule should be enabled
+`condition` | `string`<br>Condition for whether the schedule should be enabled. When a condition is provided the result of the expression will be logically AND with the  property to decide if the schedule should be enabled.
 `maxActive` | `integer`<br>Maximum number of builds that can be active at once
 `maxChanges` | `integer`<br>Maximum number of changes the schedule can fall behind head revision. If greater than zero, builds will be triggered for every submitted changelist until the backlog is this size.
 `requireSubmittedChange` | `boolean`<br>Whether the build requires a change to be submitted
@@ -515,3 +517,12 @@ Name | Description
 `keepCount` | `integer`<br>Number of artifacts to retain
 `keepDays` | `integer`<br>Number of days to retain artifacts of this type
 `namespaceId` | `string`<br>Storage namespace to use for this artifact types
+
+## StreamTag
+
+Tags for a stream
+
+Name | Description
+---- | -----------
+`name` | `string`<br>Name of the tag
+`enabled` | `boolean`<br>Whether the tag should be enabled
