Horde: Add GPU core count for Apple Silicon capabilities

#rnx

diff --git a/Engine/Source/Programs/Horde/HordeAgent/Services/CapabilitiesService.cs b/Engine/Source/Programs/Horde/HordeAgent/Services/CapabilitiesService.cs
--- a/Engine/Source/Programs/Horde/HordeAgent/Services/CapabilitiesService.cs
+++ b/Engine/Source/Programs/Horde/HordeAgent/Services/CapabilitiesService.cs
@@ -322,9 +322,12 @@ namespace HordeAgent.Services
 				if (displayNode != null)
 				{
 					XmlNode? gpuTypeNode = displayNode.SelectSingleNode("key[. = 'sppci_model']/following-sibling::string");
+					XmlNode? gpuCountNode = displayNode.SelectSingleNode("key[. = 'sppci_cores']/following-sibling::string");
 					if (gpuTypeNode != null)
 					{
-						capabilities.Properties.Add($"GPU={gpuTypeNode.InnerText}");
+						capabilities.Properties.Add(gpuCountNode != null
+							? $"GPU={gpuTypeNode.InnerText} {gpuCountNode.InnerText}-cores"
+							: $"GPU={gpuTypeNode.InnerText}");
 					}
 				}
 
