Horde: Add per-stream trace spans for ScheduleService.UpdateQueueAsync

diff --git a/Engine/Source/Programs/Horde/Plugins/Build/HordeServer.Build/Jobs/Schedules/ScheduleService.cs b/Engine/Source/Programs/Horde/Plugins/Build/HordeServer.Build/Jobs/Schedules/ScheduleService.cs
--- a/Engine/Source/Programs/Horde/Plugins/Build/HordeServer.Build/Jobs/Schedules/ScheduleService.cs
+++ b/Engine/Source/Programs/Horde/Plugins/Build/HordeServer.Build/Jobs/Schedules/ScheduleService.cs
@@ -228,12 +228,15 @@ namespace HordeServer.Jobs.Schedules
 		/// </summary>
 		public async Task UpdateQueueAsync(DateTime utcNow, CancellationToken cancellationToken)
 		{
-			using TelemetrySpan span = _tracer.StartActiveSpan($"{nameof(ScheduleService)}.{nameof(UpdateQueueAsync)}");
+			const string SpanName = $"{nameof(ScheduleService)}.{nameof(UpdateQueueAsync)}";
+			using TelemetrySpan span = _tracer.StartActiveSpan(SpanName);
 			List<SortedSetEntry<QueueItem>> queueItems = new List<SortedSetEntry<QueueItem>>();
 
 			BuildConfig buildConfig = _buildConfig.CurrentValue;
 			foreach (StreamConfig streamConfig in buildConfig.Streams)
 			{
+				using TelemetrySpan streamSpan = _tracer.StartActiveSpan(SpanName + ".Stream");
+				streamSpan.SetAttribute("horde.streamId", streamConfig.Id);
 				IStream? stream = await _streamCollection.GetAsync(streamConfig.Id, cancellationToken);
 				if (stream == null)
 				{
