Horde: Expose pool IDs through agent's HTTP management server

diff --git a/Engine/Source/Programs/Horde/HordeAgent.Tests/Services/AwsInstanceLifecycleServiceTests.cs b/Engine/Source/Programs/Horde/HordeAgent.Tests/Services/AwsInstanceLifecycleServiceTests.cs
--- a/Engine/Source/Programs/Horde/HordeAgent.Tests/Services/AwsInstanceLifecycleServiceTests.cs
+++ b/Engine/Source/Programs/Horde/HordeAgent.Tests/Services/AwsInstanceLifecycleServiceTests.cs
@@ -35,6 +35,7 @@ public class WorkerServiceStub : IWorkerService
 	}
 	
 	public bool IsConnected { get; set; }
+	public IReadOnlyList<string> PoolIds { get; set; } = [];
 }
 
 /// <summary>
diff --git a/Engine/Source/Programs/Horde/HordeAgent.Tests/Services/ManagementServiceTests.cs b/Engine/Source/Programs/Horde/HordeAgent.Tests/Services/ManagementServiceTests.cs
--- a/Engine/Source/Programs/Horde/HordeAgent.Tests/Services/ManagementServiceTests.cs
+++ b/Engine/Source/Programs/Horde/HordeAgent.Tests/Services/ManagementServiceTests.cs
@@ -1,6 +1,7 @@
 // Copyright Epic Games, Inc. All Rights Reserved.
 
 using System;
+using System.Collections.Generic;
 using System.IO;
 using System.Net;
 using System.Text.Json;
@@ -62,6 +63,15 @@ public sealed class ManagementServiceTests
 		}
 	}
 	
+	[TestMethod]
+	public void HttpRequest_Pools_Async()
+	{
+		_workerServiceStub.PoolIds = ["foo", "bar"];
+		(PoolsResponse? res, IHttpListenerResponse? httpRes) = RequestHealth<PoolsResponse>("http://localhost/pools");
+		Assert.AreEqual(200, httpRes.StatusCode);
+		CollectionAssert.AreEquivalent(new List<string> { "bar", "foo" }, res!.PoolIds);
+	}
+	
 	private (T?, IHttpListenerResponse) RequestAdmin<T>(string path)
 	{
 		return Request<T>(path, (ms, req, res) => ms.TestAdminRoute(req, res));
diff --git a/Engine/Source/Programs/Horde/HordeAgent/Services/ManagementService.cs b/Engine/Source/Programs/Horde/HordeAgent/Services/ManagementService.cs
--- a/Engine/Source/Programs/Horde/HordeAgent/Services/ManagementService.cs
+++ b/Engine/Source/Programs/Horde/HordeAgent/Services/ManagementService.cs
@@ -57,6 +57,11 @@ internal class HealthCheckResponse(bool isHealthy, bool isConnected)
 	public bool IsConnected { get; set; } = isConnected;
 }
 
+internal class PoolsResponse(List<string> poolIds)
+{
+	public List<string> PoolIds { get; set; } = poolIds;
+}
+
 /// <summary>
 /// Handles agent's built-in HTTP server used for admin tasks and health check
 /// </summary>
@@ -104,6 +109,7 @@ internal class ManagementService : IHostedService, IDisposable
 				res.StatusCode = isHealthy ? 200 : 503;
 				return new HealthCheckResponse(isHealthy, workerService.IsConnected);
 			}),
+			["/pools"] = HandleJson<PoolsResponse>((_, res) => new PoolsResponse(workerService.PoolIds.ToList())),
 		};
 		
 		string SettingsEndpointToPrefix(string endpoint) => $"http://{endpoint}/";
diff --git a/Engine/Source/Programs/Horde/HordeAgent/Services/WorkerService.cs b/Engine/Source/Programs/Horde/HordeAgent/Services/WorkerService.cs
--- a/Engine/Source/Programs/Horde/HordeAgent/Services/WorkerService.cs
+++ b/Engine/Source/Programs/Horde/HordeAgent/Services/WorkerService.cs
@@ -29,6 +29,11 @@ namespace HordeAgent.Services
 		/// Whether agent is connected to Horde server
 		/// </summary>
 		public bool IsConnected { get; }
+		
+		/// <summary>
+		/// List of pool IDs this agent (session) currently is a member of
+		/// </summary>
+		public IReadOnlyList<string> PoolIds { get; }
 	}
 	
 	/// <summary>
@@ -41,7 +46,8 @@ namespace HordeAgent.Services
 		readonly StatusService _statusService;
 		readonly IServiceProvider _serviceProvider;
 		LeaseManager? _currentLeaseManager;
-
+		
+		/// <inheritdoc/>
 		public IReadOnlyList<string> PoolIds => _currentLeaseManager?.PoolIds ?? Array.Empty<string>();
 
 		static readonly TimeSpan[] s_sessionBackOffTime =
