Horde: Add filtering to Build Health Card in Summary tab
Make issues and unpromoted observable to ensure HealthPanelIssues only re-renders when they change.

diff --git a/Engine/Source/Programs/Horde/HordeDashboard/src/components/StreamSummary.tsx b/Engine/Source/Programs/Horde/HordeDashboard/src/components/StreamSummary.tsx
--- a/Engine/Source/Programs/Horde/HordeDashboard/src/components/StreamSummary.tsx
+++ b/Engine/Source/Programs/Horde/HordeDashboard/src/components/StreamSummary.tsx
@@ -1,10 +1,10 @@
 // Copyright Epic Games, Inc. All Rights Reserved.
 
-import { DefaultButton, DetailsList, DetailsListLayoutMode, DetailsRow, IColumn, IDetailsGroupRenderProps, IDetailsListProps, mergeStyleSets, Pivot, PivotItem, ScrollablePane, ScrollbarVisibility, SelectionMode, Spinner, SpinnerSize, Stack, Text } from '@fluentui/react';
+import { Callout, DefaultButton, DetailsList, DetailsListLayoutMode, DetailsRow, Dropdown, IColumn, IconButton, IDetailsGroupRenderProps, IDetailsListProps, IDropdownOption, mergeStyleSets, Pivot, PivotItem, ScrollablePane, ScrollbarVisibility, SelectionMode, Slider, Spinner, SpinnerSize, Stack, Text, Toggle, TooltipHost } from '@fluentui/react';
 import { action, makeObservable, observable } from 'mobx';
 import { observer } from 'mobx-react-lite';
 import moment from 'moment-timezone';
-import React, { useEffect, useState } from 'react';
+import React, { useEffect, useMemo, useState } from 'react';
 import { Link, useLocation, useNavigate, useParams } from 'react-router-dom';
 import backend, { useBackend } from '../backend';
 import { FindIssueResponse, GetExternalIssueResponse, JobsTabData, TabType } from '../backend/Api';
@@ -113,8 +113,7 @@ class SummaryHandler {
             return;
          }
 
-         this.issues = values[0].filter(p => p.promoted || jiraIssues.has(p.id));
-         this.unpromoted = values[0].filter(p => !p.promoted && !jiraIssues.has(p.id));
+         this.setIssues(values, jiraIssues);
 
          this.initialLoad = false;
          this.updated();
@@ -270,9 +269,17 @@ class SummaryHandler {
 
    initialLoad = true;
 
+   @observable
    issues: FindIssueResponse[] = [];
+   @observable
    unpromoted: FindIssueResponse[] = [];
 
+   @action
+   private setIssues(values: [FindIssueResponse[]], jiraIssues: Set<number>) {
+      this.issues = values[0].filter(p => p.promoted || jiraIssues.has(p.id));
+      this.unpromoted = values[0].filter(p => !p.promoted && !jiraIssues.has(p.id));
+   }
+
    collapsedIssueGroups: Map<string, boolean> = new Map();
 
    streamId?: string;
@@ -408,6 +415,14 @@ const sortIssueGroups = (itemsIn: HealthItem[], projectStore: ProjectStore): [He
 
 const HealthPanelIssues: React.FC<{ desktopAlerts?: boolean }> = observer(({ desktopAlerts }) => {
 
+   type FilterState = {
+      visible: boolean,
+      unassigned?: boolean,
+      timeLower?: number,
+      timeUpper?: number,
+      timeUnit?: IDropdownOption
+   }
+
    const query = useQuery();
    const location = useLocation();
    const { projectStore } = useBackend();
@@ -415,6 +430,7 @@ const HealthPanelIssues: React.FC<{ desktopAlerts?: boolean }> = observer(({ des
    const [currentPivot, setCurrentPivot] = useState(query.get("workflow") === null ? "" : query.get("workflow")!);
    const { hordeClasses, detailClasses } = getHordeStyling();
    const navigate = useNavigate();
+   const [filterState, setFilterState] = useState<FilterState>({ visible: false, unassigned: false, timeLower: 0, timeUpper: 0 });
 
    // subscribe
    if (handler.update) { }
@@ -441,6 +457,42 @@ const HealthPanelIssues: React.FC<{ desktopAlerts?: boolean }> = observer(({ des
       }
    }
 
+   // apply user-specified filters
+   if (filterState.unassigned) { 
+      issues = issues?.filter(issue => {
+         return !!!issue.owner?.name;
+      });
+   }
+
+   if (filterState.timeLower !== undefined && filterState.timeUpper !== undefined && filterState.timeUpper !== 0) {
+      // zero time gap
+      if (filterState.timeLower === filterState.timeUpper ) {
+         issues = [];
+      }
+      else {
+         let conversionFactor: number;
+         switch (filterState.timeUnit?.key) {
+            case 'days':
+               conversionFactor = 24 * 60 * 60 * 1000;
+               break;
+            case 'minutes': 
+               conversionFactor = 60 * 1000;
+               break;
+            default:
+               conversionFactor = 60 * 60 * 1000;
+         }
+
+         let minTime = (Date.now() - (filterState.timeUpper * conversionFactor));
+         let maxTime = (Date.now() - (filterState.timeLower * conversionFactor));
+         
+         issues = issues?.filter(issue => {
+            let issueTime = new Date(issue.createdAt).getTime();
+            return (issueTime >= minTime && issueTime <= maxTime);
+         });
+      }
+   }
+
+
    const [items, groups] = sortIssueGroups(issues.map(i => { return { issue: i } }), projectStore);
 
    // 508
@@ -674,6 +726,32 @@ const HealthPanelIssues: React.FC<{ desktopAlerts?: boolean }> = observer(({ des
       return null;      
    }
 
+   function toggleUnassigned(_event: React.MouseEvent<HTMLElement>, checked?: boolean) {
+      setFilterState(prev => ({ ...prev, unassigned: checked ?? false }));
+   }
+
+   function toggleFiltersVisisble(_event: React.MouseEvent<HTMLElement>) {
+      setFilterState(prev => ({ ...prev, visible: !filterState.visible }));
+   }
+
+   const updateTimeRange = (_event: unknown, range: [number, number]) => {
+      setFilterState(prev => ({ ...prev, timeLower: range[0], timeUpper: range[1] }));
+   };
+
+   const timeOptions: IDropdownOption[] = useMemo( () => [
+      { key: 'minutes', text: 'minutes' },
+      { key: 'hours', text: 'hours' },
+      { key: 'days', text: 'days' }
+   ], [filterState.timeUnit]);
+
+  const updateTimeUnit = (_event: React.FormEvent<HTMLDivElement>, item: IDropdownOption): void => {
+      setFilterState(prev => ({ ...prev, timeUnit: item }));
+   };
+
+  function clearFilters(_event: React.MouseEvent<HTMLElement>) {
+      setFilterState({ visible: false });
+   }
+
    return <Stack style={{ width: 1324, marginLeft: 4 }}>
       <IssueModalV2 issueId={query.get("issue")} streamId={handler.streamId!} popHistoryOnClose={issueHistory} />
       <Stack styles={{ root: { paddingLeft: 4, paddingRight: 0, paddingTop: 4, paddingBottom: 4 } }}>
@@ -702,8 +780,33 @@ const HealthPanelIssues: React.FC<{ desktopAlerts?: boolean }> = observer(({ des
                </Pivot>
             </Stack>
             <Stack grow />
-            <Stack style={{ paddingRight: 8 }}>
-               {(currentPivot === "$promoted" || currentPivot === "$current") && !handler.initialLoad && <DefaultButton text={`Export ${currentPivot === "$promoted" ? "Promoted" : "Current"}`} style={{ fontSize: 10, padding: 8, height: "24px", minWidth: "48px" }} onClick={() => handler.exportIssues(projectStore, items, groups, currentPivot === "$promoted")} />}
+            <Stack horizontal style={{ paddingLeft: 10, paddingRight: 12}}>
+               {!handler.initialLoad && (<TooltipHost content='Filters'>
+                     <IconButton id='filtersButton' iconProps={{iconName: 'FilterList'}} onClick={toggleFiltersVisisble} />
+                  </TooltipHost>
+               )}
+               {filterState.visible && (<Callout gapSpace={0} target={'#filtersButton'} onDismiss={toggleFiltersVisisble} setInitialFocus style={{ margin: 14 }} >
+                  <Stack tokens={{ childrenGap: 28 }}>
+                     <Text style={{ fontSize: 18, fontWeight: 'bold' }}>Filters</Text>
+                     <Toggle label='Unassigned only' inlineLabel onChange={toggleUnassigned} checked={filterState.unassigned}/>
+                     <Stack tokens={{ childrenGap: 5 }}>
+                     <Text style={{ fontWeight: 'bold' }}>Show issues from past</Text>
+                        <Stack horizontal verticalAlign='center' tokens={{ childrenGap: 4 }}>
+                           <Text style={{ fontSize: 14, minWidth: 20 }}>{filterState.timeLower ?? 0}</Text>
+                           <Slider ranged min={0} max={72} value={filterState.timeUpper} lowerValue={filterState.timeLower} showValue={false} onChange={updateTimeRange} styles={{ root : { width: 300 }, slideBox: { paddingLeft: 2 } }} />
+                           <Text style={{ fontSize: 14, minWidth: 20 }}>{filterState.timeUpper && filterState.timeUpper !== 0 ? filterState.timeUpper : ' ' }</Text>
+                           <Dropdown options={timeOptions} defaultSelectedKey={'hours'} selectedKey={filterState.timeUnit ? filterState.timeUnit.key : undefined} onChange={updateTimeUnit} styles={{ root : { width: 75, paddingBottom: 3 } }}/>
+                        </Stack>
+                     </Stack>
+                     <Stack horizontal horizontalAlign='end'>
+                        <DefaultButton text='Clear' onClick={clearFilters} style={{ fontSize: 13, height: 28, minWidth: 0, width: 62, padding: 0, marginTop: 10 }}/>
+                     </Stack>
+                  </Stack>
+               </Callout>
+               )}
+            </Stack>
+            <Stack >
+               {(currentPivot === "$promoted" || currentPivot === "$current") && !handler.initialLoad && <DefaultButton text={`Export ${currentPivot === "$promoted" ? "Promoted" : "Current"}`} style={{ fontSize: 10, padding: 12, height: "24px", minWidth: "48px" }} onClick={() => handler.exportIssues(projectStore, items, groups, currentPivot === "$promoted")} />}
             </Stack>
          </Stack>
 
diff --git a/Engine/Source/Programs/Horde/HordeDashboard/src/styles/Styles.tsx b/Engine/Source/Programs/Horde/HordeDashboard/src/styles/Styles.tsx
--- a/Engine/Source/Programs/Horde/HordeDashboard/src/styles/Styles.tsx
+++ b/Engine/Source/Programs/Horde/HordeDashboard/src/styles/Styles.tsx
@@ -104,6 +104,7 @@ registerIcons({
       "Repeat": IconContents.REPEAT,
       "FlowReview": IconContents.FLOW_REVIEW,
       "FastForward": IconContents.FAST_FORWARD,
+      "FilterList": IconContents.FILTER_LIST,
       "Play": IconContents.PLAY,
       "Maximize" : IconContents.MAXIMIZE,
       "Star": IconContents.STAR,
