Renamed the telemetry entry for the Arguments field in the Job event. This is to fix a problem where the format was being processed incorrectly causing entries to appear as null. It feels unecessary to deprecate the old name as these fields are not in wide use yet and it's better to keep things clean and only write out data once rather than have duplicate data.

diff --git a/Engine/Source/Programs/Horde/Plugins/Build/HordeServer.Build/Jobs/JobCollection.cs b/Engine/Source/Programs/Horde/Plugins/Build/HordeServer.Build/Jobs/JobCollection.cs
--- a/Engine/Source/Programs/Horde/Plugins/Build/HordeServer.Build/Jobs/JobCollection.cs
+++ b/Engine/Source/Programs/Horde/Plugins/Build/HordeServer.Build/Jobs/JobCollection.cs
@@ -923,7 +923,7 @@ namespace HordeServer.Jobs
 					EventName = "State.Job",
 					Id = newJob.Id,
 					StreamId = newJob.StreamId,
-					Arguments = newJob.Arguments,
+					Args = newJob.Arguments,
 					AutoSubmit = newJob.AutoSubmit,
 					Change = newJob.CommitId.Name,
 					CodeChange = newJob.CodeCommitId?.Name,
