Added stream information to localization issues created by Horde

diff --git a/Engine/Source/Programs/Horde/Plugins/Build/HordeServer.Build.Tests/Issues/IssueHandlerTests.cs b/Engine/Source/Programs/Horde/Plugins/Build/HordeServer.Build.Tests/Issues/IssueHandlerTests.cs
--- a/Engine/Source/Programs/Horde/Plugins/Build/HordeServer.Build.Tests/Issues/IssueHandlerTests.cs
+++ b/Engine/Source/Programs/Horde/Plugins/Build/HordeServer.Build.Tests/Issues/IssueHandlerTests.cs
@@ -76,7 +76,7 @@ namespace HordeServer.Tests.Issues
 				"{ \"time\":\"2025-01-01T01:20:30.999Z\",\"level\":\"Warning\",\"message\":\"LogLocTextHelper: Warning: Text conflict for namespace '' and key 'ABCDEFGHIJKLMNOPQRSTUVWXYZABCDEH'. First entry is /SamePlugin/Scripts/SubFolder/AssetName.AssetName.Sub.Message.ContextualMessages(0).ContextualMessages.Messages(0).Messages.Message:'Text in first asset.'. Conflicting entry is /SamePlugin/Scripts/SubFolder/CopiedAsset.CopiedAsset.Sub.Message.ContextualMessages(0).ContextualMessages.Messages(0).Messages.Message:'Text in second asset.'.\",\"properties\":{ \"_channel\":{ \"$type\":\"Channel\",\"$text\":\"LogLocTextHelper\"},\"_severity\":{ \"$type\":\"Severity\",\"$text\":\"Warning\"},\"from\":\"\",\"location\":\"/SamePlugin/Scripts/SubFolder/AssetName.AssetName.Sub.Message.ContextualMessages(0).ContextualMessages.Messages(0).Messages.Message\",\"locNamespace\":\"\",\"locKey\":\"ABCDEFGHIJKLMNOPQRSTUVWXYZABCDEH\",\"locKeyMetaData\":\"\",\"text\":\"Text in first asset.\",\"textMetaData\":\"\",\"conflictLocation\":\"/SamePlugin/Scripts/SubFolder/CopiedAsset.CopiedAsset.Sub.Message.ContextualMessages(0).ContextualMessages.Messages(0).Messages.Message\",\"conflictText\":\"Text in second asset.\",\"conflictTextMetaData\":\"\",\"id\":304},\"format\":\"{_channel}: {_severity}: Text conflict{from} for namespace '{locNamespace}' and key '{locKey}'{locKeyMetaData}. First entry is {location}:'{text}'{textMetaData}. Conflicting entry is {conflictLocation}:'{conflictText}'{conflictTextMetaData}.\"}", // This "location" created an exception in the code at one point
 			};
 
-			IssueHandlerContext context = new IssueHandlerContext(new StreamId(), new TemplateId(), "Localize ProjectX", new Dictionary<string, string>());
+			IssueHandlerContext context = new IssueHandlerContext(new StreamId("Release-XX.XX"), new TemplateId(), "Localize ProjectX", new Dictionary<string, string>());
 			LocalizationIssueHandler localizationIssueHandler = new LocalizationIssueHandler(context);
 			List<IssueEventGroup> groupedIssues = TestSingleIssueHandler(localizationIssueHandler, jsonStructuredLogs);
 
@@ -90,8 +90,8 @@ namespace HordeServer.Tests.Issues
 			Assert.IsTrue(groupedIssues[0].Events.Count == 1);
 			Assert.IsTrue(groupedIssues[0].Events[0].EventId == KnownLogEvents.Engine_Localization);
 			Assert.IsTrue(groupedIssues[0].Events[0].Severity == Microsoft.Extensions.Logging.LogLevel.Warning);
-			Assert.IsTrue(groupedIssues[0].Metadata.Count == 1);
-			Assert.IsTrue(groupedIssues[0].SummaryTemplate == "{Meta:Node}: Localization {Severity} in {Files}");
+			Assert.IsTrue(groupedIssues[0].Metadata.Count == 2);
+			Assert.IsTrue(groupedIssues[0].SummaryTemplate == "{Meta:Node}: Localization {Severity} in {Meta:Stream}");
 			Assert.IsTrue(groupedIssues[0].Type == "Localization");
 
 			Assert.IsTrue(groupedIssues[1].Keys.Count == 3);
@@ -102,8 +102,8 @@ namespace HordeServer.Tests.Issues
 			Assert.IsTrue(groupedIssues[1].Events.Count == 1);
 			Assert.IsTrue(groupedIssues[1].Events[0].EventId == KnownLogEvents.Engine_Localization);
 			Assert.IsTrue(groupedIssues[1].Events[0].Severity == Microsoft.Extensions.Logging.LogLevel.Warning);
-			Assert.IsTrue(groupedIssues[1].Metadata.Count == 1);
-			Assert.IsTrue(groupedIssues[1].SummaryTemplate == "{Meta:Node}: Localization {Severity} in {Files}");
+			Assert.IsTrue(groupedIssues[1].Metadata.Count == 2);
+			Assert.IsTrue(groupedIssues[1].SummaryTemplate == "{Meta:Node}: Localization {Severity} in {Meta:Stream}");
 			Assert.IsTrue(groupedIssues[1].Type == "Localization");
 
 			Assert.IsTrue(groupedIssues[2].Keys.Count == 2);
@@ -113,8 +113,8 @@ namespace HordeServer.Tests.Issues
 			Assert.IsTrue(groupedIssues[2].Events.Count == 1);
 			Assert.IsTrue(groupedIssues[2].Events[0].EventId == KnownLogEvents.Engine_Localization);
 			Assert.IsTrue(groupedIssues[2].Events[0].Severity == Microsoft.Extensions.Logging.LogLevel.Warning);
-			Assert.IsTrue(groupedIssues[2].Metadata.Count == 1);
-			Assert.IsTrue(groupedIssues[2].SummaryTemplate == "{Meta:Node}: Localization {Severity} in {Files}");
+			Assert.IsTrue(groupedIssues[2].Metadata.Count == 2);
+			Assert.IsTrue(groupedIssues[2].SummaryTemplate == "{Meta:Node}: Localization {Severity} in {Meta:Stream}");
 			Assert.IsTrue(groupedIssues[2].Type == "Localization");
 
 			Assert.IsTrue(groupedIssues[3].Keys.Count == 2);
@@ -124,8 +124,8 @@ namespace HordeServer.Tests.Issues
 			Assert.IsTrue(groupedIssues[3].Events.Count == 1);
 			Assert.IsTrue(groupedIssues[3].Events[0].EventId == KnownLogEvents.Engine_Localization);
 			Assert.IsTrue(groupedIssues[3].Events[0].Severity == Microsoft.Extensions.Logging.LogLevel.Warning);
-			Assert.IsTrue(groupedIssues[3].Metadata.Count == 1);
-			Assert.IsTrue(groupedIssues[3].SummaryTemplate == "{Meta:Node}: Localization {Severity} in {Files}");
+			Assert.IsTrue(groupedIssues[3].Metadata.Count == 2);
+			Assert.IsTrue(groupedIssues[3].SummaryTemplate == "{Meta:Node}: Localization {Severity} in {Meta:Stream}");
 			Assert.IsTrue(groupedIssues[3].Type == "Localization");
 
 			Assert.IsTrue(groupedIssues[4].Keys.Count == 3);
@@ -136,8 +136,8 @@ namespace HordeServer.Tests.Issues
 			Assert.IsTrue(groupedIssues[4].Events.Count == 1);
 			Assert.IsTrue(groupedIssues[4].Events[0].EventId == KnownLogEvents.Engine_Localization);
 			Assert.IsTrue(groupedIssues[4].Events[0].Severity == Microsoft.Extensions.Logging.LogLevel.Warning);
-			Assert.IsTrue(groupedIssues[4].Metadata.Count == 1);
-			Assert.IsTrue(groupedIssues[4].SummaryTemplate == "{Meta:Node}: Localization {Severity} in {Files}");
+			Assert.IsTrue(groupedIssues[4].Metadata.Count == 2);
+			Assert.IsTrue(groupedIssues[4].SummaryTemplate == "{Meta:Node}: Localization {Severity} in {Meta:Stream}");
 			Assert.IsTrue(groupedIssues[4].Type == "Localization");
 
 			Assert.IsTrue(groupedIssues[5].Keys.Count == 2);
@@ -147,8 +147,8 @@ namespace HordeServer.Tests.Issues
 			Assert.IsTrue(groupedIssues[5].Events.Count == 1);
 			Assert.IsTrue(groupedIssues[5].Events[0].EventId == KnownLogEvents.Engine_Localization);
 			Assert.IsTrue(groupedIssues[5].Events[0].Severity == Microsoft.Extensions.Logging.LogLevel.Warning);
-			Assert.IsTrue(groupedIssues[5].Metadata.Count == 1);
-			Assert.IsTrue(groupedIssues[5].SummaryTemplate == "{Meta:Node}: Localization {Severity} in {Files}");
+			Assert.IsTrue(groupedIssues[5].Metadata.Count == 2);
+			Assert.IsTrue(groupedIssues[5].SummaryTemplate == "{Meta:Node}: Localization {Severity} in {Meta:Stream}");
 			Assert.IsTrue(groupedIssues[5].Type == "Localization");
 
 			Assert.IsTrue(groupedIssues[6].Keys.Count == 2);
@@ -158,8 +158,8 @@ namespace HordeServer.Tests.Issues
 			Assert.IsTrue(groupedIssues[6].Events.Count == 1);
 			Assert.IsTrue(groupedIssues[6].Events[0].EventId == KnownLogEvents.Engine_Localization);
 			Assert.IsTrue(groupedIssues[6].Events[0].Severity == Microsoft.Extensions.Logging.LogLevel.Error);
-			Assert.IsTrue(groupedIssues[6].Metadata.Count == 1);
-			Assert.IsTrue(groupedIssues[6].SummaryTemplate == "{Meta:Node}: Localization {Severity} in {Files}");
+			Assert.IsTrue(groupedIssues[6].Metadata.Count == 2);
+			Assert.IsTrue(groupedIssues[6].SummaryTemplate == "{Meta:Node}: Localization {Severity} in {Meta:Stream}");
 			Assert.IsTrue(groupedIssues[6].Type == "Localization");
 
 			Assert.IsTrue(groupedIssues[7].Keys.Count == 2);
@@ -171,8 +171,8 @@ namespace HordeServer.Tests.Issues
 			Assert.IsTrue(groupedIssues[7].Events[0].Severity == Microsoft.Extensions.Logging.LogLevel.Error);
 			Assert.IsTrue(groupedIssues[7].Events[1].EventId == KnownLogEvents.Engine_Localization);
 			Assert.IsTrue(groupedIssues[7].Events[1].Severity == Microsoft.Extensions.Logging.LogLevel.Error);
-			Assert.IsTrue(groupedIssues[7].Metadata.Count == 0);
-			Assert.IsTrue(groupedIssues[7].SummaryTemplate == "Localization {Severity}");
+			Assert.IsTrue(groupedIssues[7].Metadata.Count == 2);
+			Assert.IsTrue(groupedIssues[7].SummaryTemplate == "{Meta:Node}: Localization {Severity} in {Meta:Stream}");
 			Assert.IsTrue(groupedIssues[7].Type == "Localization");
 
 			Assert.IsTrue(groupedIssues[8].Keys.Count == 3);
@@ -183,8 +183,8 @@ namespace HordeServer.Tests.Issues
 			Assert.IsTrue(groupedIssues[8].Events.Count == 1);
 			Assert.IsTrue(groupedIssues[8].Events[0].EventId == KnownLogEvents.Engine_Localization);
 			Assert.IsTrue(groupedIssues[8].Events[0].Severity == Microsoft.Extensions.Logging.LogLevel.Warning);
-			Assert.IsTrue(groupedIssues[8].Metadata.Count == 1);
-			Assert.IsTrue(groupedIssues[8].SummaryTemplate == "{Meta:Node}: Localization {Severity} in {Files}");
+			Assert.IsTrue(groupedIssues[8].Metadata.Count == 2);
+			Assert.IsTrue(groupedIssues[8].SummaryTemplate == "{Meta:Node}: Localization {Severity} in {Meta:Stream}");
 			Assert.IsTrue(groupedIssues[8].Type == "Localization");
 		}
 	}
diff --git a/Engine/Source/Programs/Shared/EpicGames.Horde/Issues/Handlers/LocalizationIssueHandler.cs b/Engine/Source/Programs/Shared/EpicGames.Horde/Issues/Handlers/LocalizationIssueHandler.cs
--- a/Engine/Source/Programs/Shared/EpicGames.Horde/Issues/Handlers/LocalizationIssueHandler.cs
+++ b/Engine/Source/Programs/Shared/EpicGames.Horde/Issues/Handlers/LocalizationIssueHandler.cs
@@ -36,8 +36,8 @@ namespace EpicGames.Horde.Issues.Handlers
 		}
 
 		readonly IssueHandlerContext _context;
+		readonly IssueEventGroup _issuesWithoutClearSuspects;
 		readonly List<IssueEventGroup> _issues = new List<IssueEventGroup>();
-		readonly IssueEventGroup _issuesWithoutClearFiles = new IssueEventGroup("Localization", "Localization {Severity}", IssueChangeFilter.All);
 
 		/// <inheritdoc/>
 		public override int Priority => 10;
@@ -48,6 +48,7 @@ namespace EpicGames.Horde.Issues.Handlers
 		public LocalizationIssueHandler(IssueHandlerContext context)
 		{
 			_context = context;
+			_issuesWithoutClearSuspects = CreateNewIssueGroup();
 		}
 
 		/// <summary>
@@ -70,6 +71,25 @@ namespace EpicGames.Horde.Issues.Handlers
 			return eventId == KnownLogEvents.ExitCode;
 		}
 
+		private IssueEventGroup CreateNewIssueGroup()
+		{
+			string localizationIssueGroup = "Localization";
+			string localizationSummaryTemplate = (!System.String.IsNullOrEmpty(_context.NodeName) ? "{Meta:Node}: " : "") + "Localization {Severity}" + (!System.String.IsNullOrEmpty(_context.StreamId.ToString()) ? " in {Meta:Stream}" : "");
+
+			IssueEventGroup newIssueGroup = new IssueEventGroup(localizationIssueGroup, localizationSummaryTemplate, IssueChangeFilter.All);
+
+			if (!System.String.IsNullOrEmpty(_context.NodeName))
+			{
+				newIssueGroup.Metadata.Add("Node", _context.NodeName);
+			}
+			if (!System.String.IsNullOrEmpty(_context.StreamId.ToString()))
+			{
+				newIssueGroup.Metadata.Add("Stream", _context.StreamId.ToString());
+			}
+
+			return newIssueGroup;
+		}
+
 		/// <summary>
 		/// Generate all the IssueKeys with the extracted LocalizationData
 		/// </summary>
@@ -235,20 +255,14 @@ namespace EpicGames.Horde.Issues.Handlers
 					if(data.HasAnyFile())
 					{
 						// Create a new issue group (might be grouped later based on keys)
-						string localizationIssueGroup = "Localization";
-						string localizationSummaryTemplate = (!System.String.IsNullOrEmpty(_context.NodeName) ? "{Meta:Node}: " : "") + "Localization {Severity} in {Files}";
-						issueGroup = new IssueEventGroup(localizationIssueGroup, localizationSummaryTemplate, IssueChangeFilter.All);
-						if (!System.String.IsNullOrEmpty(_context.NodeName))
-						{
-							issueGroup.Metadata.Add("Node", _context.NodeName);
-						}
+						issueGroup = CreateNewIssueGroup();
 					}
 					else
 					{
 						// If we have no idea on the file an issue is associated with, then we won't be able to narrow down a suspect.
 						// If we can't narrow down a suspect. There is no gain to split the issues amongst different issueGroup.
 						// Actually, there is only a risk of creating thousands of issue group if localization is in a weird/broken state.
-						issueGroup = _issuesWithoutClearFiles;
+						issueGroup = _issuesWithoutClearSuspects;
 					}
 
 					// Add current event to the chosen issueGroup
