API encounters ".NET type System.Text.Json.JsonElement cannot be mapped to a BsonValue" on swagger API invocations.

diff --git a/Engine/Source/Programs/Horde/Plugins/Build/HordeServer.Build/Jobs/TestData/TestDataController.cs b/Engine/Source/Programs/Horde/Plugins/Build/HordeServer.Build/Jobs/TestData/TestDataController.cs
--- a/Engine/Source/Programs/Horde/Plugins/Build/HordeServer.Build/Jobs/TestData/TestDataController.cs
+++ b/Engine/Source/Programs/Horde/Plugins/Build/HordeServer.Build/Jobs/TestData/TestDataController.cs
@@ -1,5 +1,7 @@
 // Copyright Epic Games, Inc. All Rights Reserved.
 
+
+using System.Text.Json;
 using EpicGames.Core;
 using EpicGames.Horde.Commits;
 using EpicGames.Horde.Jobs;
@@ -375,7 +377,9 @@ namespace HordeServer.Jobs.TestData
 				return NotFound();
 			}
 
-			IReadOnlyList<ITestData> testData = await _testDataCollection.AddAsync(job, jobStep, new (string key, BsonDocument value)[] { (request.Key, new BsonDocument(request.Data)) });
+			string dataJson = JsonSerializer.Serialize(request.Data); 
+			BsonDocument dataBson = BsonSerializer.Deserialize<BsonDocument>(dataJson);
+			IReadOnlyList<ITestData> testData = await _testDataCollection.AddAsync(job, jobStep, new (string key, BsonDocument value)[] { (request.Key, dataBson) });
 			return new CreateTestDataResponse(testData[0].Id.ToString());
 		}
 
