Horde Dashboard: Clean up logic in LogRender.tsx

diff --git a/Engine/Source/Programs/Horde/HordeDashboard/src/components/LogRender.tsx b/Engine/Source/Programs/Horde/HordeDashboard/src/components/LogRender.tsx
--- a/Engine/Source/Programs/Horde/HordeDashboard/src/components/LogRender.tsx
+++ b/Engine/Source/Programs/Horde/HordeDashboard/src/components/LogRender.tsx
@@ -29,6 +29,16 @@ export type LogItem = {
    issue?: IssueData;
 }
 
+// Provides the start of hrefs for each tag type with simple link construction behavior
+const simpleTagLinks: Partial<Record<TagType, string>> = {
+   [TagType.JobId]: "/job/",
+   [TagType.LogId]: "/log/",
+   [TagType.ProjectId]: "/project/",
+   [TagType.StreamId]: "/stream/",
+   [TagType.PoolId]: "/pools?poolid=",
+   [TagType.DeviceId]: "/audit/device/"
+}
+
 const getLineNumber = (line: LogLine): number | undefined => {
 
    const properties = line.properties;
@@ -138,158 +148,164 @@ const renderTags = (navigate: NavigateFunction, line: LogLine, lineNumber: numbe
 
          type = record.type ?? record["$type"] ?? "";
          text = record.text ?? record["$text"] ?? "";
-
-         if (type === "SourceFile") {
-            tagType = TagType.SourceFile;
-         }
-
-         if (type === "Link") {
-            tagType = TagType.Link;
-         }
-
+         
          if (type === "ErrorCode" && text.startsWith("C")) {
             tagType = TagType.MSDNCode;
          }
 
-         if (type === "LeaseId") {
-            tagType = TagType.LeaseId;
-         }
-
-         if (type === "AgentId") {
-            tagType = TagType.AgentId;
-         }
-
-         if (type === "PoolId") {
-            tagType = TagType.PoolId;
-         }
-
-         if (type === "StreamId") {
-            tagType = TagType.StreamId;
-         }
-
-         if (type === "ProjectId") {
-            tagType = TagType.ProjectId;
-         }
-
-         if (type === "JobId") {
-            tagType = TagType.JobId;
-         }
-
-         if (type === "LogId") {
-            tagType = TagType.LogId;
-         }
-
-         if (type === "ArtifactId") {
-            tagType = TagType.ArtifactId;
-         }
-
-         if (type === "DeviceId") {
-            tagType = TagType.DeviceId
+         const enumCast = type as keyof typeof TagType;
+         if(enumCast in TagType) {
+            tagType = TagType[enumCast];
          }
+      }
 
+      let highlight = (
+         <Highlight
+            search={search ? search : ""} 
+            className={logStyle.logLine} 
+         >
+            {text}
+         </Highlight>
+      )
+
+      if(tagType === TagType.None || !record) {
+         return (
+            // Defined separately from highlight because of needed key prop
+            <Highlight
+               key={key}
+               search={search ? search : ""} 
+               className={logStyle.logLine} 
+            >
+               {text}
+            </Highlight>
+         )
       }
 
-      if (tagType === TagType.None || !record) {
-
-         return <Highlight key={key} search={search ? search : ""} className={logStyle.logLine}>{text}</Highlight>;
-
-      } else if (tagType === TagType.JobId) {
-         const url = `/job/${text}`
-         return <a key={key} href={url} onClick={(ev) => { ev.stopPropagation(); ev.preventDefault(); navigate(url) }}><Highlight search={search ? search : ""} className={logStyle.logLine}>{text}</Highlight></a>;
-      } else if (tagType === TagType.LogId) {
-         const url = `/log/${text}`
-         return <a key={key} href={url} onClick={(ev) => { ev.stopPropagation(); ev.preventDefault(); navigate(url) }}><Highlight search={search ? search : ""} className={logStyle.logLine}>{text}</Highlight></a>;
-      } else if (tagType === TagType.ProjectId) {
-         const url = `/project/${text}`
-         return <a key={key} href={url} onClick={(ev) => { ev.stopPropagation(); ev.preventDefault(); navigate(url) }}><Highlight search={search ? search : ""} className={logStyle.logLine}>{text}</Highlight></a>;
-      } else if (tagType === TagType.StreamId) {
-         const url = `/stream/${text}`
-         return <a key={key} href={url} onClick={(ev) => { ev.stopPropagation(); ev.preventDefault(); navigate(url) }}><Highlight search={search ? search : ""} className={logStyle.logLine}>{text}</Highlight></a>;
-      } else if (tagType === TagType.PoolId) {
-         const url = `/pools?poolid=${text}`
-         return <a key={key} href={url} onClick={(ev) => { ev.stopPropagation(); ev.preventDefault(); navigate(url) }}><Highlight search={search ? search : ""} className={logStyle.logLine}>{text}</Highlight></a>;
-      } else if (tagType === TagType.DeviceId) {
-         const url = `/audit/device/${text}`
-         return <a key={key} href={url} onClick={(ev) => { ev.stopPropagation(); ev.preventDefault(); navigate(url) }}><Highlight search={search ? search : ""} className={logStyle.logLine}>{text}</Highlight></a>;
-      } else if (tagType === TagType.LeaseId) {
-
-         const navigateToLeaseLog = async (toplevel: boolean) => {
-
-            const logData = await backend.getLease(text);
-            const url = `/log/${logData?.logId}`;
-            if (!toplevel) {
-               navigate(url)
-            } else {
-               window.open(url, "_blank");
+      let linkProps: React.AnchorHTMLAttributes<HTMLAnchorElement> | null = null;
+
+      // Update link props object based on tag type, then return a link with those props
+      // or a span
+      switch(tagType) {
+         // Cases with similar simple tag behaviors
+         case TagType.JobId:
+         case TagType.LogId:
+         case TagType.ProjectId:
+         case TagType.StreamId:
+         case TagType.PoolId:
+         case TagType.DeviceId:
+            const url = simpleTagLinks[tagType] + text; // simpleTagLinks created for concision here
+            linkProps = {
+               href: url,
+               onClick: (ev) => { ev.stopPropagation(); ev.preventDefault(); navigate(url) }
             }
-         }
-
-         return <a key={key} href="/"
-
-            onAuxClick={(ev) => {
-               ev.preventDefault();
-               ev.stopPropagation()
-               navigateToLeaseLog(true)
-            }}
-
-            onClick={(ev) => {
-               ev.stopPropagation();
-               ev.preventDefault();
-               navigateToLeaseLog(!!ev?.ctrlKey || !!ev.metaKey)
-            }}>
-            <Highlight search={search ? search : ""} className={logStyle.logLine}>{text}</Highlight>
-         </a>;
-
-      } else if (tagType === TagType.MSDNCode) {
-
-         return <a key={key} target="_blank" rel="noopener noreferrer" href={`https://msdn.microsoft.com/query/dev16.query?appId=Dev16IDEF1&l=EN-US&k=k(${text.toLowerCase()})&rd=true`} onClick={(ev) => ev.stopPropagation()}><Highlight search={search ? search : ""} className={logStyle.logLine}>{text}</Highlight></a>;
-
-      } else if (tagType === TagType.AgentId) {
-         const search = new URLSearchParams(window.location.search);
-         search.set("agentId", encodeURIComponent(text));
-         const url = `${window.location.pathname}?` + search.toString();
-         return <a key={key} href="/" onClick={async (ev) => { ev.stopPropagation(); ev.preventDefault(); navigate(url, { replace: true }) }}><Highlight search={search ? search : ""} className={logStyle.logLine}>{text}</Highlight></a>;
-
-      } else if (tagType === TagType.SourceFile) {
-
-         let depotPath = record.depotPath;
-
-         if (!depotPath) {
-            return <Highlight key={key} search={search ? search : ""} className={logStyle.logLine}>{text}</Highlight>;
-         }
-
-         if (depotPath.indexOf("@") !== -1) {
-            depotPath = depotPath.slice(0, depotPath.indexOf("@"));
-         }
-
-         depotPath = encodeURIComponent(depotPath);
-
-         // @todo: handle line, it isn't in the source file meta data, and matching isn't obvious (file vs note, etc)
-         let tagLine = "";
-         if (tagLine) {
-            depotPath += `&line=${tagLine}`;
-         }
-
-         return <a key={key} href={`ugs://timelapse?depotPath=${(depotPath)}`} onClick={(ev) => ev.stopPropagation()}><Highlight search={search ? search : ""} className={logStyle.logLine}>{text}</Highlight></a>;
-      } else if (tagType === TagType.Link) {
-
-         return <a key={key} rel="noreferrer" href={record.target} onClick={(ev) => ev.stopPropagation()}><Highlight search={search ? search : ""} className={logStyle.logLine}>{text}</Highlight></a>;
-      } else if (tagType === TagType.ArtifactId) {
-         const artifactType = properties.ArtifactType ?? "_none";
-
-         const search = new URLSearchParams(window.location.search);
-         search.set("artifactContext", encodeURIComponent(artifactType as string));
-         if (properties.ArtifactId) {
-            if (properties.ArtifactId["$text"]) {
-               search.set("artifactId", encodeURIComponent(properties.ArtifactId["$text"] as string));
+            break;
+
+         case TagType.MSDNCode:
+            linkProps = {
+               href: `https://msdn.microsoft.com/query/dev16.query?appId=Dev16IDEF1&l=EN-US&k=k(${text.toLowerCase()})&rd=true`,
+               onClick: (ev) => ev.stopPropagation(),
+               target: "_blank",
+               rel: "noopener noreferrer"
+            };
+            break;
+
+         case TagType.Link:
+            linkProps = {
+               href: record.target,
+               onClick: (ev) => ev.stopPropagation(),
+               rel: "noreferrer"
             }
-         }
-
-         const url = `${window.location.pathname}?` + search.toString();
-         return <a key={key} href="/" onClick={(ev) => { ev.stopPropagation(); ev.preventDefault(); navigate(url, { replace: true }) }}><Highlight search={search ? search : ""} className={logStyle.logLine}>{text}</Highlight></a>;
+            break;
+
+         case TagType.AgentId:
+            const search = new URLSearchParams(window.location.search);
+            search.set("agentId", encodeURIComponent(text));
+            // Update highlight because of new search in block scope
+            highlight = (
+               <Highlight 
+                  search={search ? search : ""} 
+                  className={logStyle.logLine}
+               >
+                  {text}
+               </Highlight>
+            )
+            const url2 = `${window.location.pathname}?` + search.toString();
+            linkProps = {
+               href: url2,
+               onClick: async (ev) => { ev.stopPropagation(); ev.preventDefault(); navigate(url2, { replace: true }) }
+            }
+            break;
+
+         case TagType.ArtifactId:
+            const artifactType = properties.ArtifactType ?? "_none";
+            const search2 = new URLSearchParams(window.location.search);
+            search2.set("artifactContext", encodeURIComponent(artifactType as string));
+            if (properties.ArtifactId) {
+               if (properties.ArtifactId["$text"]) {
+                  search2.set("artifactId", encodeURIComponent(properties.ArtifactId["$text"] as string));
+               }
+            }
+            // Update highlight because of new search in block scope
+            highlight = (
+               <Highlight 
+                  search={search2 ? search2 : ""} 
+                  className={logStyle.logLine}
+               >
+                  {text}
+               </Highlight>
+            )
+            const url3 = `${window.location.pathname}?${search2.toString()}`;
+            linkProps = {
+               href: url3,
+               onClick: async (ev) => { ev.stopPropagation(); ev.preventDefault(); navigate(url3, { replace: true }) }
+            };
+            break;
+
+         case TagType.SourceFile:
+            let depotPath = record.depotPath;
+            if (!depotPath) break;
+            if (depotPath.indexOf("@") !== -1) {
+               depotPath = depotPath.slice(0, depotPath.indexOf("@"));
+            }
+            depotPath = encodeURIComponent(depotPath);
+            // @todo: handle line, it isn't in the source file meta data, and matching isn't obvious (file vs note, etc)
+            let tagLine = "";
+            if (tagLine) {
+               depotPath += `&line=${tagLine}`;
+            }
+            linkProps = {
+               href: `ugs://timelapse?depotPath=${(depotPath)}`,
+               onClick: (ev) => ev.stopPropagation()
+            }
+            break;
+
+         case TagType.LeaseId:
+            const navigateToLeaseLog = async (toplevel: boolean) => {
+               const logData = await backend.getLease(text);
+               const url = `/log/${logData?.logId}`;
+               if (!toplevel) {
+                  navigate(url)
+               } else {
+                  window.open(url, "_blank");
+               }
+            }
+            linkProps = {
+               href: "/",
+               onClick: (ev) => {
+                  ev.stopPropagation();
+                  ev.preventDefault();
+                  navigateToLeaseLog(!!ev?.ctrlKey || !!ev.metaKey)
+               },
+               onAuxClick: (ev) => {
+                  ev.preventDefault();
+                  ev.stopPropagation()
+                  navigateToLeaseLog(true)
+               }
+            }
+            break;
       }
 
-      return <span key={key} />;
+      return linkProps ? <a key={key} {...linkProps}>{highlight}</a> : <span key={key}/>
    }).filter(t => !!t);
 
    let remaining = line.format;
