Horde: Update release notes

#skipci
#submitool bypass

diff --git a/Engine/Source/Programs/Horde/Docs/ReleaseNotes.md b/Engine/Source/Programs/Horde/Docs/ReleaseNotes.md
--- a/Engine/Source/Programs/Horde/Docs/ReleaseNotes.md
+++ b/Engine/Source/Programs/Horde/Docs/ReleaseNotes.md
@@ -1,5 +1,26 @@
 # Release Notes
 
+## 2025-07-17
+
+* Horde Dashboard: Clean up logic in LogRender.tsx (44106602)
+* API encounters ".NET type System.Text.Json.JsonElement cannot be mapped to a BsonValue" on swagger API invocations. (44078324)
+* Review change for LocalizationIssueHandler (44009728)
+* Added stream information to localization issues created by Horde (44002634)
+* Renamed the telemetry entry for the Arguments field in the Job event. This is to fix a problem where the format was being processed incorrectly causing entries to appear as null. It feels unecessary to deprecate the old name as these fields are not in wide use yet and it's better to keep things clean and only write out data once rather than have duplicate data. (43869880)
+* EpicGames.Perforce: Log if untracked files are removed (43864355)
+* Add filtering to Build Health Card in Summary tab
+  Make issues and unpromoted observable to ensure HealthPanelIssues only re-renders when they change. (43771480)
+* Skip client IP allow listing for allocated compute agents (43771211)
+* Expose pool IDs through agent's HTTP management server (43734454)
+* Add per-stream trace spans for ScheduleService.UpdateQueueAsync (43729468)
+* Truncate long template names in new build modal dropdown (43712812)
+* Remove legacy isSwarmTab variable and conditional displaying of job search button altogether (43704159)
+* Display job search button in presubmit tab (43702416)
+* Add missing capabilities data for Macs (43693319)
+* Add tags to the StreamConfig to control whether scheduled jobs are enabled (43687954)
+* Horde Dashboard: cancellation reason on user home (43655326)
+* Horde Dashboard: job status feedback (43651240)
+
 ## 2025-06-19
 
 * Avoid issue service exceptions during suspect finding (43629946)
