Horde:  Build Health - Add json payload to error logging when creating Jiras from Horde
#rnx

diff --git a/Engine/Source/Programs/Horde/Plugins/Build/HordeServer.Build/Issues/External/JiraService.cs b/Engine/Source/Programs/Horde/Plugins/Build/HordeServer.Build/Issues/External/JiraService.cs
--- a/Engine/Source/Programs/Horde/Plugins/Build/HordeServer.Build/Issues/External/JiraService.cs
+++ b/Engine/Source/Programs/Horde/Plugins/Build/HordeServer.Build/Issues/External/JiraService.cs
@@ -440,7 +440,16 @@ namespace HordeServer.Issues.External
 
 			HttpResponseMessage response = await _retryPolicy.ExecuteAsync(() => _client.PostAsync(uri, new StringContent(bodyJson, Encoding.UTF8, "application/json")));
 
-			response.EnsureSuccessStatusCode();
+			try
+			{
+				response.EnsureSuccessStatusCode();
+			}
+			catch 
+			{
+				_logger.LogWarning("JiraService create issue: {Uri} returned {Code} Payload: {BodyJson}", uri, response.StatusCode, bodyJson);
+				throw;
+			}
+			
 
 			string responseBody = await response.Content.ReadAsStringAsync(cancellationToken);
 
