Horde Dashboard: Put HelpModal back into top level nav

diff --git a/Engine/Source/Programs/Horde/HordeDashboard/src/components/HelpModal.tsx b/Engine/Source/Programs/Horde/HordeDashboard/src/components/HelpModal.tsx
deleted file mode 100644
diff --git a/Engine/Source/Programs/Horde/HordeDashboard/src/components/SupportModal.tsx b/Engine/Source/Programs/Horde/HordeDashboard/src/components/SupportModal.tsx
new file mode 100644
--- /dev/null
+++ b/Engine/Source/Programs/Horde/HordeDashboard/src/components/SupportModal.tsx
@@ -0,0 +1,47 @@
+// Copyright Epic Games, Inc. All Rights Reserved.
+import { IconButton, Modal, PrimaryButton, Stack, Text, Label } from '@fluentui/react';
+import React from 'react';
+import dashboard from '../backend/Dashboard';
+import { getHordeStyling } from '../styles/Styles';
+
+
+export const SupportModal: React.FC<{ show: boolean, onClose: () => void }> = ({ show, onClose }) => {
+
+    const { hordeClasses } = getHordeStyling();
+
+    const close = () => {
+        onClose();
+    };
+
+    const helpEmail = dashboard.helpEmail;
+    const helpSlack = dashboard.helpSlack;
+
+    return (
+        <Modal className={hordeClasses.modal} isOpen={show} styles={{ main: { padding: 16, width: 540, display: "relative" } }} onDismiss={() => { if (show) close() }}>
+            <IconButton
+                style={{position: "absolute", right: 10, top: 10}}
+                iconProps={{ iconName: 'Cancel' }}
+                ariaLabel="Close popup modal"
+                onClick={() => { close(); }}
+            />
+            
+            <Text block style={{ width: "100%", fontSize: 18, fontWeight: 400, fontFamily: "Horde Open Sans SemiBold", marginBottom: 16 }}>Horde Support</Text>
+
+            <Text block style={{ fontSize: 15, marginBottom: 12 }}>For support, feedback, and suggestions:</Text>
+            
+            <Stack tokens={{childrenGap: 8}}>
+                {helpEmail && <Text block style={{ fontSize: 15, marginLeft: 8 }}>
+                    Email:
+                    <Text style={{ marginLeft: 12, fontSize: 15, fontWeight: 400, fontFamily: "Horde Open Sans SemiBold" }}>{helpEmail}</Text>
+                </Text>}
+                {helpSlack && <Text block style={{ fontSize: 15, marginLeft: 8 }}>
+                    Slack:
+                    <Text style={{ marginLeft: 12, fontSize: 15, fontWeight: 400, fontFamily: "Horde Open Sans SemiBold" }}>{helpSlack}</Text>
+                </Text>}
+            </Stack>
+
+            <PrimaryButton style={{ position: "absolute", right: 16, bottom: 16 }} text="Ok" disabled={false} onClick={() => { close(); }} />
+        </Modal>
+    );
+
+};
diff --git a/Engine/Source/Programs/Horde/HordeDashboard/src/components/TopNav.tsx b/Engine/Source/Programs/Horde/HordeDashboard/src/components/TopNav.tsx
--- a/Engine/Source/Programs/Horde/HordeDashboard/src/components/TopNav.tsx
+++ b/Engine/Source/Programs/Horde/HordeDashboard/src/components/TopNav.tsx
@@ -14,6 +14,7 @@ import { getHordeTheme } from '../styles/theme';
 import { getHordeStyling } from '../styles/Styles';
 import { ArtifactQueryState, FindArtifactsModal } from './artifacts/ArtifactsSearch';
 import { getHordePlugins, MountType } from 'hordePlugins';
+import { SupportModal } from './SupportModal';
 
 const getStyles = () => {
 
@@ -414,6 +415,7 @@ export const TopNav: React.FC<{ suppressServer?: boolean }> = observer(({ suppre
 
    const [showMenu, setShowMenu] = useState(false);
    const [showVersion, setShowVersion] = useState(false);
+   const [showSupport, setShowSupport] = useState(false);
    const [showPreviewChanges, setShowPreviewChanges] = useState(false);
    const navigate = useNavigate();
    const { projectStore } = useBackend();
@@ -694,6 +696,14 @@ export const TopNav: React.FC<{ suppressServer?: boolean }> = observer(({ suppre
          onClick: () => { setShowVersion(true) }
       });
 
+      if(dashboard.helpEmail || dashboard.helpSlack){
+         hordeItems.push({
+            key: "support_contacts",
+            text: "Support",
+            onClick: () => { setShowSupport(true) }
+         })
+      }
+
       const style = { ...menuStyles } as Partial<IContextualMenuStyles>;
 
       const cbItem: ICommandBarItemProps = {
@@ -788,6 +798,7 @@ export const TopNav: React.FC<{ suppressServer?: boolean }> = observer(({ suppre
 
          }} />}
          {showVersion && <VersionModal show={true} onClose={() => { setShowVersion(false) }} />}
+         {showSupport && <SupportModal show={true} onClose={() => {setShowSupport(false) }}/>}
          {showPreviewChanges && <PreviewChangesModal onClose={() => { setShowPreviewChanges(false) }} />}
          <Stack tokens={{ maxWidth: 1440, childrenGap: 0 }} disableShrink={true} styles={{ root: { backgroundColor: hordeTheme.horde.topNavBackground, margin: "auto", width: "100%" } }}>
             <Stack horizontal verticalAlign='center' styles={{ root: { height: "60px" } }} >
diff --git a/Engine/Source/Programs/Horde/HordeDashboard/src/components/error/ErrorBar.tsx b/Engine/Source/Programs/Horde/HordeDashboard/src/components/error/ErrorBar.tsx
--- a/Engine/Source/Programs/Horde/HordeDashboard/src/components/error/ErrorBar.tsx
+++ b/Engine/Source/Programs/Horde/HordeDashboard/src/components/error/ErrorBar.tsx
@@ -11,14 +11,24 @@ export const ErrorBar: React.FC = observer(() => {
 
     const helpEmail = dashboard.helpEmail;
     const helpSlack = dashboard.helpSlack;
-    const anyHelp = !!helpEmail || !!helpSlack;
+    const anyHelp = !!helpSlack || !!helpEmail;
 
     const getConsoleErrorMessage = (): string => {
         if(!error) return "";
 
         let msg = `${error.title}: ${errorBarStore.message}`
         msg += error.url ? ` - Occurred while trying to access ${error.url}` : "";
-        msg += anyHelp ? ` - If the issue persists, please contact us for assistance at ${helpEmail} OR the ${helpSlack} Slack channel.` : "";
+
+        // Clear formatting for console output
+        if(anyHelp) {
+            msg += ` - If the issue persists, please contact us for assistance:\n\n`;
+
+            msg += helpEmail ? `Email: ${helpEmail}\n` : "";
+            msg += helpSlack ? `Slack: ${helpSlack}\n` : "";
+
+            msg += "\n";
+        }
+
         return msg;
     }
 
