Add phase artifacts widget to allow user to filter phase related artifacts

#rnx

diff --git a/Engine/Source/Programs/Horde/HordeDashboard/plugins/test-automation/testAutomationCommon.tsx b/Engine/Source/Programs/Horde/HordeDashboard/plugins/test-automation/testAutomationCommon.tsx
--- a/Engine/Source/Programs/Horde/HordeDashboard/plugins/test-automation/testAutomationCommon.tsx
+++ b/Engine/Source/Programs/Horde/HordeDashboard/plugins/test-automation/testAutomationCommon.tsx
@@ -1,10 +1,11 @@
-import { mergeStyles, mergeStyleSets, Text, ComboBox, IComboBoxOption, SelectableOptionMenuItemType, IContextualMenuProps, IContextualMenuItem, IconButton, IContextualMenuListProps, Stack, Checkbox, FontIcon } from "@fluentui/react";
+import { mergeStyles, mergeStyleSets, Text, ComboBox, IComboBoxOption, SelectableOptionMenuItemType, IContextualMenuProps, IContextualMenuItem, IconButton, IContextualMenuListProps, Stack, Checkbox, FontIcon, Modal, DefaultButton, ScrollablePane, SelectionZone, DetailsList, Selection, DetailsListLayoutMode, SelectionMode, IColumn, PrimaryButton } from "@fluentui/react";
 import { MetadataRef, PhaseSessionResult, TestDataHandler, TestNameRef, TestPhaseStatus, TestSessionDetails, TestSessionResult } from "./testData";
 import dashboard, { StatusColor } from "horde/backend/Dashboard";
-import { useEffect, useState } from "react";
+import { useCallback, useEffect, useMemo, useState } from "react";
 import { TestOutcome, TestPhaseOutcome } from "./api";
 import { PhaseHistoryGraph } from "./phaseHistoryGraph";
 import { projectStore } from "horde/backend/ProjectStore";
+import { getHordeStyling } from "horde/styles/Styles";
 
 export type PhaseSessionCallback = (session: PhaseSessionResult) => void;
 
@@ -44,15 +45,46 @@ export const styles = {
         height: 15,
         paddingTop: 1
     },
-}
-
-const barStyles = mergeStyleSets({
-
     stripes: {
         backgroundImage: 'repeating-linear-gradient(-45deg, rgba(255, 255, 255, .2) 25%, transparent 25%, transparent 50%, rgba(255, 255, 255, .2) 50%, rgba(255, 255, 255, .2) 75%, transparent 75%, transparent)',
-    }
+    },
+}
 
-});
+let _styles: any;
+
+const getStyles = () => {
+
+   const border = `1px solid ${dashboard.darktheme ? "#2D2B29" : "#EDEBE9"}`
+
+   const styles = _styles ?? mergeStyleSets({
+      list: {
+         selectors: {
+            'a': {
+               height: "unset !important",
+            },
+            '.ms-List-cell': {
+
+               borderTop: border,
+               borderRight: border,
+               borderLeft: border
+            },
+            '.ms-List-cell:nth-last-child(-n + 1)': {
+               borderBottom: border
+            },
+            ".ms-DetailsRow #artifactview": {
+               opacity: 0
+            },
+            ".ms-DetailsRow:hover #artifactview": {
+               opacity: 1
+            },
+         }
+      }
+   });
+
+   _styles = styles;
+
+   return styles;
+}
 
 export const MultiOptionChooser: React.FC<{options: IComboBoxOption[], initialSelection: string[], updateKeys: (selectedKeys: string[]) => void, placeholder?: string, style?: any, disabled?: boolean }> = ({ options, initialSelection, updateKeys, placeholder, style, disabled }) => {
     const [selection, setSelection] = useState<string[]>(initialSelection);
@@ -114,13 +146,13 @@ export const StatusBar = (stack: StatusBarStack[], width: number, height: number
         <div className={mergeStyles({ backgroundColor: basecolor, width: width, height: height, verticalAlign: 'middle', display: "flex" }, style)} title={mainTitle}>
             {stack.map((item) => <span key={item.title!}
                 onClick={item.onClick}
-                className={item.stripes ? barStyles.stripes : undefined}
                 style={{
-                width: `${item.value}%`, height: '100%',
-                backgroundColor: item.color,
-                display: 'block',
-                cursor: item.onClick ? 'pointer' : 'inherit',
-                backgroundSize: `${height * 2}px ${height * 2}px`
+                    width: `${item.value}%`, height: '100%',
+                    backgroundColor: item.color,
+                    display: 'block',
+                    cursor: item.onClick ? 'pointer' : 'inherit',
+                    backgroundSize: `${height * 2}px ${height * 2}px`,
+                    backgroundImage: item.stripes ? styles.stripes.backgroundImage : undefined
                 }} />)}
         </div>
     );
@@ -553,3 +585,177 @@ export const CheckListOption: React.FC<{options: ICheckListOption[], onChange: (
                 }}
             />
 }
+
+type ArtifactHref = {
+    name: string;
+    reference: string;
+    href: string;
+}
+
+export const PhaseArtifactsModal: React.FC<{ phase: TestPhaseStatus, artifactPaths: string[], onClose?: () => void }> = ({ phase, artifactPaths, onClose }) => {
+    const [selectVer, setSelectVer] = useState(0);
+    const [items, setItems] = useState<ArtifactHref[]>([]);
+
+    useEffect(() => {
+        if (artifactPaths.length === 0) return;
+        const hrefs: (Promise<string> | undefined)[] = [];
+        for (const artifact of artifactPaths) {
+            hrefs.push(phase.artifacts?.getLink(artifact));
+        }
+        // get common prefix end position
+        let pos = 0;
+        const sortedPaths = artifactPaths.concat().sort();
+        const shortest = sortedPaths[0];
+        const longest = sortedPaths[sortedPaths.length - 1];
+        if (shortest === longest) {
+            pos = shortest.lastIndexOf('/') + 1;
+        } else {
+            while (pos < shortest.length && shortest.charAt(pos) === longest.charAt(pos) ) pos++;
+            pos = shortest.lastIndexOf('/', pos) + 1;
+        }
+        
+        Promise.all(hrefs).then(results => {
+            const items: ArtifactHref[] = [];
+            artifactPaths.forEach((artifact, index) => {
+                const result = results[index];
+                if (!result) return;
+                items.push({name: artifact.substring(pos), reference: artifact, href: result});
+            });
+            setItems(items);
+        });
+    }, [phase, artifactPaths]);
+
+    const { hordeClasses } = getHordeStyling();
+    const styles = getStyles();
+
+    const selector = useMemo(() => new Selection({onSelectionChanged: () => setSelectVer(Math.random)}), [phase, artifactPaths]);
+
+    const downloadText = useMemo(() => {
+        let text = "Download";
+        if (selector.getSelectedCount() === 0) {
+            text += ` (${artifactPaths.length} files)`;
+        } else if (selector.getSelectedCount() === 1) {
+            text += ` (1 file)`;
+        } else if (selector.getSelectedCount() > 1) {
+            text += ` (${selector.getSelectedCount()} files)`;
+        }
+        return text;
+    }, [selectVer]);
+
+    const downloadZip = useCallback(() => {
+        if (artifactPaths.length === 0) return;
+
+        let items = selector.getSelection() as ArtifactHref[];
+        if (items.length === 0) items = selector.getItems() as ArtifactHref[];
+
+        // download a single file
+        if (items.length === 1) {
+            phase.artifacts?.download(items[0].reference);
+            return;
+        }
+
+        phase.artifacts?.downloadZip(items.map(item => item.reference));
+
+    }, [phase, artifactPaths]);
+
+    const columns: IColumn[] = useMemo(() => [
+        { key: 'column1', name: 'Name', minWidth: 794 - 32, maxWidth: 794 - 32, isResizable: false, isPadded: false },
+        { key: 'column2', name: 'View_Download', minWidth: 64 + 32, maxWidth: 64 + 32, isResizable: false, isPadded: false }
+    ], []);
+
+    const renderItem = useCallback((item: ArtifactHref, index?: number, column?: IColumn) => {
+        if (!column)  return null;
+
+        if (column.name === "View_Download") {
+
+            const href = item.href;
+
+            return <Stack data-selection-disabled verticalAlign="center" verticalFill horizontal horizontalAlign="end" style={{ paddingTop: 0, paddingBottom: 0 }}>
+                        <IconButton id="artifactview" href={`${href}&inline=true`} target="_blank" style={{ paddingTop: 1, color: "#106EBE" }} iconProps={{ iconName: "Eye", styles: { root: { fontSize: "14px" } } }} />
+                        <IconButton id="artifactview" href={href} target="_blank" style={{ paddingTop: 1, color: "#106EBE" }} iconProps={{ iconName: "CloudDownload", styles: { root: { fontSize: "14px" } } }} />
+                    </Stack>
+        }
+
+        if (column.name === "Name") {
+
+            const path = (item.name as string).split("/");
+
+            const pathElements = path.map((t, index) => {
+                const last = index === (path.length - 1);
+                let color = last ? (dashboard.darktheme ? "#FFFFFF" : "#605E5C") : undefined;
+                const font = last ? undefined : "Horde Open Sans Light";
+                const sep = last ? undefined : "/"
+                return <Text key={`p-${item.name}-${index}`} styles={{ root: { fontFamily: font } }} style={{ color: color }}>{t}{sep}</Text>
+            })
+
+            return <Stack verticalFill verticalAlign="center" style={{ cursor: "pointer" }}>
+                        <Stack horizontal tokens={{ childrenGap: 8 }}>
+                            <Stack>
+                                <FontIcon style={{ paddingTop: 1, fontSize: 16 }} iconName="Document" />
+                            </Stack>
+                            <Stack horizontal>
+                                {pathElements}
+                            </Stack>
+                        </Stack>
+                    </Stack>
+        }
+
+        return null;
+    }, []);
+
+    return <Stack>
+            <Modal isOpen={true} topOffsetFixed={true} styles={{ main: { padding: 8, width: 1180, height: 820, hasBeenOpened: false, top: "80px", position: "absolute" } }} onDismiss={onClose} className={hordeClasses.modal}>
+                <Stack className="horde-no-darktheme" styles={{ root: { paddingTop: 10, paddingRight: 12 } }}>
+                    <Stack style={{ paddingLeft: 24, paddingRight: 24 }}>
+                        <Stack tokens={{ childrenGap: 12 }} style={{ height: 800 }}>
+                            <Stack horizontal verticalAlign="start">
+                                <Stack style={{ paddingTop: 3 }}>
+                                    <Text variant="mediumPlus" styles={{ root: { fontFamily: "Horde Open Sans SemiBold" } }}>Artifacts</Text>
+                                </Stack>
+                                <Stack grow />
+                                <Stack horizontalAlign="end">
+                                    <IconButton
+                                        iconProps={{ iconName: 'Cancel' }}
+                                        onClick={onClose}
+                                    />
+                                </Stack>
+                            </Stack>
+                            <Stack styles={{ root: { paddingLeft: 4, paddingRight: 0, paddingTop: 8, paddingBottom: 4 } }}>
+                                <Stack tokens={{ childrenGap: 12 }}>
+                                    <Stack horizontal verticalAlign="center" tokens={{ childrenGap: 18 }} style={{ paddingBottom: 12 }}>
+                                        <Text style={{fontSize: 13, fontWeight: 600, whiteSpace: 'nowrap', overflow: 'hidden', textOverflow: 'ellipsis', cursor: "text"}}>{phase.name}</Text>
+                                        <Stack grow />
+                                        <PrimaryButton text={downloadText} onClick={downloadZip}/>
+                                    </Stack >
+                                    <Stack style={{ height: 492 + 160, position: "relative" }}>
+                                        <ScrollablePane style={{ height: 492 + 160 }}>
+                                            <SelectionZone selection={selector}>
+                                                <DetailsList
+                                                    styles={{ root: { overflowX: "hidden" } }}
+                                                    className={styles.list}
+                                                    isHeaderVisible={false}
+                                                    compact={true}
+                                                    items={items}
+                                                    columns={columns}
+                                                    layoutMode={DetailsListLayoutMode.fixedColumns}
+                                                    selectionMode={SelectionMode.multiple}
+                                                    enableUpdateAnimations={false}
+                                                    selection={selector}
+                                                    selectionPreservedOnEmptyClick={true}
+                                                    onShouldVirtualize={() => false}
+                                                    onItemInvoked={(item: ArtifactHref) => {
+                                                        window.open(`${item.href}&inline=true`, "_blank");
+                                                    }}
+                                                    onRenderItemColumn={renderItem}
+                                                />
+                                            </SelectionZone>
+                                        </ScrollablePane>
+                                    </Stack>
+                                </Stack >
+                            </Stack>
+                        </Stack>
+                    </Stack>
+                </Stack>
+            </Modal>
+        </Stack>
+}
diff --git a/Engine/Source/Programs/Horde/HordeDashboard/plugins/test-automation/testAutomationEvents.tsx b/Engine/Source/Programs/Horde/HordeDashboard/plugins/test-automation/testAutomationEvents.tsx
--- a/Engine/Source/Programs/Horde/HordeDashboard/plugins/test-automation/testAutomationEvents.tsx
+++ b/Engine/Source/Programs/Horde/HordeDashboard/plugins/test-automation/testAutomationEvents.tsx
@@ -1,5 +1,5 @@
 import { Checkbox, DefaultButton, DirectionalHint, FontIcon, HoverCard, HoverCardType, Image, IPlainCardProps, mergeStyles, mergeStyleSets, Spinner, SpinnerSize, Stack, Text } from "@fluentui/react";
-import { EventEntry, EventLevel, ImageCompare, LogPropertyType } from "./testEventsModel";
+import { Artifacts, EventEntry, EventLevel, ImageCompare, LogPropertyType } from "./testEventsModel";
 import { TestPhaseOutcome } from "./api";
 import { memo, useCallback, useEffect, useState } from "react";
 import { ArtifactFactory, DeviceRef, TestPhaseStatus } from "./testData";
@@ -7,7 +7,7 @@ import dashboard, { StatusColor } from "horde/backend/Dashboard";
 import { getShortNiceTime, msecToElapsed } from "horde/base/utilities/timeUtils";
 import { DashboardPreference } from "horde/backend/Api";
 import { isBright } from "horde/base/utilities/colors";
-import { styles } from "./testAutomationCommon";
+import { PhaseArtifactsModal, styles } from "./testAutomationCommon";
 
 let eventStyles: any = undefined;
 export const getEventStyling = () => {
@@ -101,7 +101,7 @@ const Event: React.FC<{ entry: EventEntry, phase: TestPhaseStatus }> = memo(({en
 const EventTypeChooser = (property: string | number | Record<string, string | number | undefined>, phase: TestPhaseStatus): JSX.Element | undefined => {
     switch (property["$type"]) {
         case LogPropertyType.ImageCompare:
-            return <ImageComparer imageCompare={property as ImageCompare} phase={phase} />
+            return <ImageComparer key={`ic-${property["$text"]}`} imageCompare={property as ImageCompare} phase={phase} />
     }
 
     return;
@@ -142,15 +142,20 @@ const DeviceItem = (device: DeviceRef): JSX.Element => {
             </HoverCard>
 }
 
-export const PhaseEventsPane: React.FC<{phase: TestPhaseStatus, onEventsLoaded?: () => void}> = ({phase, onEventsLoaded}) => {
+export const PhaseEventsPane: React.FC<{phase: TestPhaseStatus, onEventsLoaded?: (events: EventEntry[]) => void}> = ({phase, onEventsLoaded}) => {
     const [events, setEvents] = useState<EventEntry[] | undefined>(undefined);
     const [filterError, setFilterError] = useState<boolean>(true);
     const [filterWarning, setFilterWarning] = useState<boolean>(true);
     const [filterInfo, setFilterInfo] = useState<boolean>(false);
+    const [artifacts, setArtifacts] = useState<string[] | undefined>(undefined);
+    const [artifactModal, setArtifactModal] = useState<boolean>(false);
 
     useEffect(() => {
         // reset events and filtering
-        if (events !== undefined) setEvents(undefined);
+        if (events !== undefined) {
+            setEvents(undefined);
+            setArtifacts(undefined);
+        }
         if (!filterInfo && phase.outcome !== TestPhaseOutcome.Failed) {
             setFilterInfo(true);
         }
@@ -161,7 +166,44 @@ export const PhaseEventsPane: React.FC<{phase: TestPhaseStatus, onEventsLoaded?:
     }, [phase]);
 
     useEffect(() => {
-        if (events) onEventsLoaded?.();
+        if (events) onEventsLoaded?.(events);
+    }, [events]);
+
+    const collectPhaseArtifacts = useCallback((): string[] | undefined => {
+        if (!events) return;
+        // collect artifact references
+        const artifacts: string[] = [];
+        // phase log
+        const logPath = phase.devices.length && phase.devices[0].logPath;
+        if (logPath) artifacts.push(logPath);
+        // parse event properties
+        events.forEach(entry => {
+            if (!!entry.properties) {
+                Object.values(entry.properties).forEach(property => {
+                    if (!property) return;
+                    switch (property["$type"]) {
+                        case LogPropertyType.ImageCompare:
+                            const imageCompare = property as ImageCompare;
+                            if (imageCompare.approved) artifacts.push(imageCompare.approved);
+                            if (imageCompare.approved_metadata) artifacts.push(imageCompare.approved_metadata);
+                            if (imageCompare.unapproved) artifacts.push(imageCompare.unapproved);
+                            if (imageCompare.unapproved_metadata) artifacts.push(imageCompare.unapproved_metadata);
+                            if (imageCompare.difference) artifacts.push(imageCompare.difference);
+                            if (imageCompare.difference_report) artifacts.push(imageCompare.difference_report);
+                            break;
+
+                        case LogPropertyType.Artifacts:
+                            const artifactsProp = property as Artifacts
+                            Object.entries(artifactsProp).forEach(([key, value]) => {
+                                if (!value || key.startsWith('$')) return;
+                                artifacts.push(value);
+                            });
+                            break;
+                    }
+                });
+            }
+        });
+        return artifacts;
     }, [events]);
 
     const isEntryNeedDisplay = useCallback((event: EventEntry): boolean => {
@@ -186,6 +228,9 @@ export const PhaseEventsPane: React.FC<{phase: TestPhaseStatus, onEventsLoaded?:
                         <Spinner size={SpinnerSize.medium} />
                     </Stack>
                 }
+                {!! events && artifactModal && artifacts &&
+                    <PhaseArtifactsModal phase={phase} artifactPaths={artifacts} onClose={() => setArtifactModal(false)}/>
+                }
                 {!!events &&
                     <Stack style={{overflow: 'hidden'}} tokens={{childrenGap: 3}} grow>
                         <Stack style={{borderWidth: '0px 0px 1px 0px', borderColor: borderColor, borderStyle: 'solid', paddingBottom: 2}} horizontal tokens={{childrenGap: 8}} disableShrink>
@@ -204,9 +249,15 @@ export const PhaseEventsPane: React.FC<{phase: TestPhaseStatus, onEventsLoaded?:
                                 <Checkbox styles={styles.defaultCheckbox} label="Info" checked={filterInfo} onChange={(_: any, checked?: boolean) => setFilterInfo(!!checked)} disabled={!events.length}/>
                             </Stack>
                             <Stack horizontalAlign="end">
-                                <DefaultButton disabled={true} style={{ fontSize: 11, padding: 0, height: 18, color: 'white' }} title="Unavailable"
+                                <DefaultButton disabled={artifacts && artifacts.length === 0} style={{ fontSize: 11, padding: 0, height: 18, color: 'white', backgroundColor: artifacts && artifacts.length === 0 ? undefined : styles.defaultButton.backgroundColor }}
                                     onClick={(ev) => {
                                         ev.stopPropagation();
+                                        let collectedArtifacts = artifacts;
+                                        if (!collectedArtifacts) {
+                                            collectedArtifacts = collectPhaseArtifacts();
+                                            setArtifacts(collectedArtifacts);
+                                        }
+                                        if (collectedArtifacts && collectedArtifacts.length > 0) setArtifactModal(true);
                                     }}>Artifacts</DefaultButton>
                             </Stack>
                         </Stack>
diff --git a/Engine/Source/Programs/Horde/HordeDashboard/plugins/test-automation/testAutomationPhaseView.tsx b/Engine/Source/Programs/Horde/HordeDashboard/plugins/test-automation/testAutomationPhaseView.tsx
--- a/Engine/Source/Programs/Horde/HordeDashboard/plugins/test-automation/testAutomationPhaseView.tsx
+++ b/Engine/Source/Programs/Horde/HordeDashboard/plugins/test-automation/testAutomationPhaseView.tsx
@@ -15,6 +15,7 @@ import { JobArtifactsModal } from "horde/components/artifacts/ArtifactsModal";
 import { getJobParams, JobContext } from "./phaseCompare";
 import { getEventStyling, PhaseEventsPane } from "./testAutomationEvents";
 import { comparisonContext, PhaseComparisonPane } from "./testAutomationPhaseCompare";
+import { EventEntry } from "./testEventsModel";
 
 const PhaseHistoryPane: React.FC<{ test: TestNameRef, meta?: MetadataRef, phaseKey: string, streamId?: string, sessionId?: string, handler: TestDataHandler, onClickCompare: () => void }> = memo(({test, meta, phaseKey, streamId, sessionId, handler, onClickCompare}) => {
 
@@ -48,7 +49,7 @@ const PhaseHistoryPane: React.FC<{ test: TestNameRef, meta?: MetadataRef, phaseK
             </Stack>
 });
 
-const TestPhaseExpandPane: React.FC<{ phase: TestPhaseStatus, onEventsLoaded?: () => void, handler: TestDataHandler }> = memo(({ phase, onEventsLoaded, handler }) => {
+const TestPhaseExpandPane: React.FC<{ phase: TestPhaseStatus, onEventsLoaded?: (events: EventEntry[]) => void, handler: TestDataHandler }> = memo(({ phase, onEventsLoaded, handler }) => {
     const [sessionId, setSessionId] = useState<string | undefined>();
     const [comparison, setComparison] = useState<boolean>(comparisonContext.isEnable);
 
diff --git a/Engine/Source/Programs/Horde/HordeDashboard/plugins/test-automation/testData.ts b/Engine/Source/Programs/Horde/HordeDashboard/plugins/test-automation/testData.ts
--- a/Engine/Source/Programs/Horde/HordeDashboard/plugins/test-automation/testData.ts
+++ b/Engine/Source/Programs/Horde/HordeDashboard/plugins/test-automation/testData.ts
@@ -1548,6 +1548,16 @@ export class ArtifactFactory {
 
     }
 
+    async download(referencePath: string): Promise<void> {
+        await this.getJobStepSavedArtifacts();
+        if (!!this.artifactId) backend.downloadArtifactV2(this.artifactId, referencePath);
+    }
+
+    async downloadZip(referencePaths: string[]) : Promise<void> {
+        await this.getJobStepSavedArtifacts();
+        if (!!this.artifactId) backend.downloadArtifactZipV2(this.artifactId, {filter: referencePaths})
+    }
+
     private jobId: string;
     private stepId: string;
     private artifactV2Response?: GetArtifactResponse;
diff --git a/Engine/Source/Programs/Horde/HordeDashboard/plugins/test-automation/testEventsModel.ts b/Engine/Source/Programs/Horde/HordeDashboard/plugins/test-automation/testEventsModel.ts
--- a/Engine/Source/Programs/Horde/HordeDashboard/plugins/test-automation/testEventsModel.ts
+++ b/Engine/Source/Programs/Horde/HordeDashboard/plugins/test-automation/testEventsModel.ts
@@ -30,7 +30,8 @@ export enum EventLevel {
 
 /// Known property types
 export enum LogPropertyType {
-    ImageCompare = "Image Compare"
+    ImageCompare = "Image Compare",
+    Artifacts = "Artifacts"
 }
 
 /// Image Compare detailed properties
@@ -58,3 +59,8 @@ export type ImageCompare = {
     //** The path to the json report associated with the image from the artifacts */
     difference_report?: string;
 }
+
+/// Artifacts detailed properties
+export type Artifacts = {
+    [Key in string]: string | undefined;
+}
