fix grammar error

#rnx

diff --git a/Engine/Source/Programs/Horde/HordeDashboard/plugins/test-automation/testAutomationCommon.tsx b/Engine/Source/Programs/Horde/HordeDashboard/plugins/test-automation/testAutomationCommon.tsx
--- a/Engine/Source/Programs/Horde/HordeDashboard/plugins/test-automation/testAutomationCommon.tsx
+++ b/Engine/Source/Programs/Horde/HordeDashboard/plugins/test-automation/testAutomationCommon.tsx
@@ -632,12 +632,11 @@ export const PhaseArtifactsModal: React.FC<{ phase: TestPhaseStatus, artifactPat
 
     const downloadText = useMemo(() => {
         let text = "Download";
-        if (selector.getSelectedCount() === 0) {
-            text += ` (${artifactPaths.length} files)`;
-        } else if (selector.getSelectedCount() === 1) {
+        const count = selector.getSelectedCount() || artifactPaths.length;
+        if (count === 1) {
             text += ` (1 file)`;
-        } else if (selector.getSelectedCount() > 1) {
-            text += ` (${selector.getSelectedCount()} files)`;
+        } else if (count > 1) {
+            text += ` (${count} files)`;
         }
         return text;
     }, [selectVer]);
