Horde Dashboard: Better handling of advanced additional arguments in NewBuildV2

diff --git a/Engine/Source/Programs/Horde/HordeDashboard/src/components/build/NewBuildV2.tsx b/Engine/Source/Programs/Horde/HordeDashboard/src/components/build/NewBuildV2.tsx
--- a/Engine/Source/Programs/Horde/HordeDashboard/src/components/build/NewBuildV2.tsx
+++ b/Engine/Source/Programs/Horde/HordeDashboard/src/components/build/NewBuildV2.tsx
@@ -1948,7 +1948,7 @@ const BuildModal: React.FC<{ setUseLegacyDialog: (value: boolean) => void }> = o
       let additionalArgs: string[] = [];
       if (options!.advAdditionalArgs) {
          const argRegex = /"(\\"|[^"])*?"|[^ ]+/g;
-         options.advAdditionalArgs.trim().replaceAll("\n", " ").match(argRegex)?.forEach(arg => additionalArgs.push(arg.replace(/"/g, "")));
+         options.advAdditionalArgs.trim().replaceAll("\n", " ").match(argRegex)?.forEach(arg => additionalArgs.push(arg.replace(/"|\\"|'|\\'|&quot;/g, "")));
       }
 
       // trim all parameters
@@ -1991,6 +1991,7 @@ const BuildModal: React.FC<{ setUseLegacyDialog: (value: boolean) => void }> = o
          if (!!import.meta.env.VITE_HORDE_DEBUG_NEW_JOB) {
 
             console.log("Debug Job Submit");
+            console.log(data);
             console.log(JSON.stringify(data));
             return;
          }
