Horde: Update code coverage to use JetBrains.dotCover.CommandLineTools

The previuous code coverage tool JetBrains.dotCover.GlobalTool has been deprecated.

diff --git a/Engine/Source/Programs/Horde/HordeServer/Dockerfile b/Engine/Source/Programs/Horde/HordeServer/Dockerfile
--- a/Engine/Source/Programs/Horde/HordeServer/Dockerfile
+++ b/Engine/Source/Programs/Horde/HordeServer/Dockerfile
@@ -18,8 +18,8 @@ COPY Source/Programs/Horde/nuget.config ./Source/Programs/Horde/nuget.config
 RUN DEBIAN_FRONTEND=noninteractive \
 	apt-get -qq update && \
 	apt-get -qq install curl zip && \
-	dotnet tool install --global JetBrains.dotCover.GlobalTool --version 2022.2.3 && \
-	ln -s /root/.dotnet/tools/dotnet-dotcover /usr/local/bin/dotnet-dotcover
+	dotnet tool install --global JetBrains.dotCover.CommandLineTools && \
+	ln -s /root/.dotnet/tools/dotCover /usr/local/bin/dotnet-dotcover
 
 COPY --from=redis /usr/local/bin/redis-server /usr/local/bin/redis-server
 
diff --git a/Engine/Source/Programs/Horde/Scripts/test.sh b/Engine/Source/Programs/Horde/Scripts/test.sh
--- a/Engine/Source/Programs/Horde/Scripts/test.sh
+++ b/Engine/Source/Programs/Horde/Scripts/test.sh
@@ -29,7 +29,7 @@ for csProj in "${testProjects[@]}"; do
 	filename="${csProj##*/}"
 	args=("test")
 	if [ "$code_coverage" = "true" ]; then
-		args=(dotcover test --dcOutput=/tmp/${filename}.dcvr --dcFilters="+:EpicGames*;+:Horde*;-:*.Tests")
+		args=(dotcover cover-dotnet --output=/tmp/${filename}.dcvr --filters="+:EpicGames*;+:Horde*;-:*.Tests" -- test)
 	fi
 	dotnet "${args[@]}" "$csProj" --blame-hang-timeout 5m --blame-hang-dump-type mini --logger 'console;verbosity=normal' || exit 1
 done
