Horde: A node that is modified and skipped during a build graph update does not fail the whole build

If there is a modified node with a skipped state during a build graph update the whole graph would previously be marked as failed, instead allow the build to continue by marking dependent nodes as skipped/failed but still allow other independent subtrees to be run.

diff --git a/Engine/Source/Programs/Horde/Plugins/Build/HordeServer.Build.Tests/Jobs/JobCollectionTests.cs b/Engine/Source/Programs/Horde/Plugins/Build/HordeServer.Build.Tests/Jobs/JobCollectionTests.cs
--- a/Engine/Source/Programs/Horde/Plugins/Build/HordeServer.Build.Tests/Jobs/JobCollectionTests.cs
+++ b/Engine/Source/Programs/Horde/Plugins/Build/HordeServer.Build.Tests/Jobs/JobCollectionTests.cs
@@ -246,6 +246,95 @@ namespace HordeServer.Tests.Jobs
 			Assert.AreEqual(JobStepState.Ready, job.Batches[5].Steps[0].State); // Gather
 		}
 
+		[TestMethod]
+		public async Task UpdateGraphWithModifiedSkippedStepAsync()
+		{
+			Mock<ITemplate> templateMock = new(MockBehavior.Strict);
+			templateMock.SetupGet(x => x.InitialAgentType).Returns((string?)null);
+
+			IGraph baseGraph = await GraphCollection.AddAsync(templateMock.Object);
+
+			List<NewGroup> groups =
+			[
+				new NewGroup("Test", [new NewNode(name: "Compile Item 1", outputs: ["#CompileOutput1"])]),
+				new NewGroup("Test", [new NewNode(name: "Compile Item 2", outputs: ["#CompileOutput1"])]),
+				new NewGroup("Test", [new NewNode(name: "Gather", inputDependencies: ["Compile Item 1", "Compile Item 2"])])
+			];
+
+			IGraph graph1 = await GraphCollection.AppendAsync(baseGraph, groups);
+
+			groups.Insert(2, new NewGroup("Test", [new NewNode(name: "Generate Metadata", inputDependencies: ["Compile Item 2"])]));
+			groups.Insert(3, new NewGroup("Test", [new NewNode(name: "Process Content", outputs: ["#ProcessContent"], inputDependencies: ["Compile Item 2"])]));
+			groups.Insert(4, new NewGroup("Test", [new NewNode(name: "Publish Item 2", outputs: ["#PublishItem2"], inputDependencies: ["Process Content", "Generate Metadata", "Compile Item 2"])]));
+			groups.Insert(5, new NewGroup("Test", [new NewNode(name: "Store Symbols Item 1", outputs: ["#StoreSymbolsItem1"], inputDependencies: ["Compile Item 1"])]));
+			groups.Insert(6, new NewGroup("Test", [new NewNode(name: "Publish Item 1", outputs: ["#PublishItem1"], inputDependencies: ["Store Symbols Item 1"])]));
+			// Modify step Gather, now at position 7 
+			groups[7].Nodes[0].Inputs = ["#PublishItem1", "#PublishItem2"];
+			groups[7].Nodes[0].InputDependencies = ["Publish Item 1", "Publish Item 2"];
+
+			IGraph graph2 = await GraphCollection.AppendAsync(baseGraph, groups);
+
+			CreateJobOptions options = new();
+			options.Arguments.Add("-Target=Gather");
+
+			IJob job = await JobCollection.AddAsync(JobIdUtils.GenerateNewId(), new StreamId("ue4-main"), new TemplateId("test-build"), ContentHash.SHA1("hello"), baseGraph, "Test job", CommitIdWithOrder.FromPerforceChange(123), CommitIdWithOrder.FromPerforceChange(123), options);
+
+			Assert.AreEqual(1, job.Batches.Count);
+
+			job = await RunBatchAsync(job, 0);
+			job = await RunStepAsync(job, 0, 0, JobStepOutcome.Success); // Setup Build
+
+			job = Deref(await job.TryUpdateGraphAsync(graph1));
+
+			Assert.AreEqual(4, job.Batches.Count);
+
+			job = await RunBatchAsync(job, 1);
+			job = await RunStepAsync(job, 1, 0, JobStepOutcome.Success); // Compile Item 1
+
+			// Fail this step to cause dependent steps to be marked as skipped and ignored when update graph is called
+			job = await RunBatchAsync(job, 2);
+			job = await RunStepAsync(job, 2, 0, JobStepOutcome.Failure); // Compile Item 2
+
+			Assert.AreEqual(JobStepState.Skipped, job.Batches[3].Steps[0].State); // Gather
+
+			job = Deref(await job.TryUpdateGraphAsync(graph2));
+
+			Assert.AreEqual(9, job.Batches.Count);
+
+			// Check the state of the new dependent steps of the failed step are marked as skipped
+			Assert.AreEqual(JobStepState.Completed, job.Batches[0].Steps[0].State); // Setup Build
+			Assert.AreEqual(JobStepState.Completed, job.Batches[1].Steps[0].State); // Compile Item 1
+			Assert.AreEqual(JobStepState.Completed, job.Batches[2].Steps[0].State); // Compile Item 2
+			Assert.AreEqual(JobStepState.Skipped,   job.Batches[3].Steps[0].State); // Generate Metadata
+			Assert.AreEqual(JobStepState.Skipped,   job.Batches[4].Steps[0].State); // Process Content
+			Assert.AreEqual(JobStepState.Skipped,   job.Batches[5].Steps[0].State); // Publish Item 2
+			Assert.AreEqual(JobStepState.Ready,     job.Batches[6].Steps[0].State); // Store Symbols Item 1
+			Assert.AreEqual(JobStepState.Waiting,   job.Batches[7].Steps[0].State); // Publish Item 1
+			Assert.AreEqual(JobStepState.Skipped,   job.Batches[8].Steps[0].State); // Gather
+
+			// Check the outcome of the new dependent steps of the failed step are marked as failed
+			Assert.AreEqual(JobStepOutcome.Success, job.Batches[0].Steps[0].Outcome); // Setup Build
+			Assert.AreEqual(JobStepOutcome.Success, job.Batches[1].Steps[0].Outcome); // Compile Item 1
+			Assert.AreEqual(JobStepOutcome.Failure, job.Batches[2].Steps[0].Outcome); // Compile Item 2
+			Assert.AreEqual(JobStepOutcome.Failure, job.Batches[3].Steps[0].Outcome); // Generate Metadata
+			Assert.AreEqual(JobStepOutcome.Failure, job.Batches[4].Steps[0].Outcome); // Process Content
+			Assert.AreEqual(JobStepOutcome.Failure, job.Batches[5].Steps[0].Outcome); // Publish Item 2
+			Assert.AreEqual(JobStepOutcome.Success, job.Batches[6].Steps[0].Outcome); // Store Symbols Item 1
+			Assert.AreEqual(JobStepOutcome.Success, job.Batches[7].Steps[0].Outcome); // Publish Item 1
+			Assert.AreEqual(JobStepOutcome.Failure, job.Batches[8].Steps[0].Outcome); // Gather
+
+			// Check the remaining steps can complete OK
+			job = await RunBatchAsync(job, 6);
+			job = await RunStepAsync(job, 6, 0, JobStepOutcome.Success); // Store Symbols Item 1
+			job = await RunBatchAsync(job, 7);
+			job = await RunStepAsync(job, 7, 0, JobStepOutcome.Success); // Publish Item 1
+
+			Assert.AreEqual(JobStepState.Completed, job.Batches[6].Steps[0].State); // Store Symbols Item 1
+			Assert.AreEqual(JobStepOutcome.Success, job.Batches[6].Steps[0].Outcome);
+			Assert.AreEqual(JobStepState.Completed, job.Batches[7].Steps[0].State); // Publish Item 1
+			Assert.AreEqual(JobStepOutcome.Success, job.Batches[7].Steps[0].Outcome);
+		}
+
 		[TestMethod]
 		public async Task UpdateGraphOnFailedBatchAsync()
 		{
diff --git a/Engine/Source/Programs/Horde/Plugins/Build/HordeServer.Build/Jobs/JobCollection.cs b/Engine/Source/Programs/Horde/Plugins/Build/HordeServer.Build/Jobs/JobCollection.cs
--- a/Engine/Source/Programs/Horde/Plugins/Build/HordeServer.Build/Jobs/JobCollection.cs
+++ b/Engine/Source/Programs/Horde/Plugins/Build/HordeServer.Build/Jobs/JobCollection.cs
@@ -1845,7 +1845,11 @@ namespace HordeServer.Jobs
 
 					if (!JobStepExtensions.IsPendingState(step.State) && !NodesMatch(oldGraph, oldNode, newGraph, newNode))
 					{
-						throw new InvalidOperationException($"Definition for node '{oldNode.Name}' has changed.");
+						// Modified skipped steps are ignored to allow other non-dependent subtrees to still be run
+						if (step.State != JobStepState.Skipped)
+						{
+							throw new InvalidOperationException($"Definition for node '{oldNode.Name}' has changed.");
+						}
 					}
 
 					step.NodeIdx = newNodeIdx;
