Horde: Add example of using Azure Key Vault to Secrets plugin

The Azure Key Vault Secrets Provider is provided as experimental and an example of how to implement a secret provider from another cloud provider.

diff --git a/Engine/Source/Programs/Horde/Docs/Config/Secrets.md b/Engine/Source/Programs/Horde/Docs/Config/Secrets.md
--- a/Engine/Source/Programs/Horde/Docs/Config/Secrets.md
+++ b/Engine/Source/Programs/Horde/Docs/Config/Secrets.md
@@ -93,7 +93,7 @@ features for this provider.
 
 ### HashiCorp Vault
 
-For HashiCorp Vault to set the endpoint and credentials, will use a `providerConfig` property instead of `name` property to provide the data.
+To set the endpoint and credentials for HashiCorp Vault use the `providerConfig` property instead of the `provider` property.
 For example to fetch three secrets from HashiCorp Vault using different authentication methods:
 
    * Use an AWS ARN to pass to AWS Security Token Service AssumeRole
@@ -182,6 +182,41 @@ For example to fetch three secrets from HashiCorp Vault using different authenti
     }
  ```
 
+### Azure
+
+> The Azure Key Vault Secrets Provider is provided as experimental and an example of how to implement a secret provider from another cloud provider.
+
+To set the URI of the Azure Key Vault use the `providerConfig` property instead of the `provider` property.
+For authenication to Azure the implementation uses [DefaultAzureCredential](https://learn.microsoft.com/en-us/dotnet/api/azure.identity.defaultazurecredential?view=azure-dotnet).
+
+   ```json
+    "secrets": [
+        {
+            "id": "secret-from-key-vault",
+ 
+            "sources": [
+                {
+                    "providerConfig": "azure-key-vault",
+                    "path": "name-of-the-secret"
+                }
+            ]
+        }
+    ],
+    "plugins": {
+        "secrets": {
+            "providerConfigs": [
+                {
+                    "name": "azure-key-vault",
+                    "provider": "AzureKeyVault",
+                    "azureKeyVault": {
+                        "vaultUri": "https://vault-name.vault.azure.net/"
+                    }
+                }
+            ]
+        }
+    }
+   ```
+
 ## Using secrets
 
 The most common use case for secrets is during build automation pipelines. In this scenario, the Horde Server
diff --git a/Engine/Source/Programs/Horde/Plugins/Secrets/HordeServer.Secrets.Tests/Secrets/SecretProviderTests.cs b/Engine/Source/Programs/Horde/Plugins/Secrets/HordeServer.Secrets.Tests/Secrets/SecretProviderTests.cs
--- a/Engine/Source/Programs/Horde/Plugins/Secrets/HordeServer.Secrets.Tests/Secrets/SecretProviderTests.cs
+++ b/Engine/Source/Programs/Horde/Plugins/Secrets/HordeServer.Secrets.Tests/Secrets/SecretProviderTests.cs
@@ -49,7 +49,8 @@ namespace HordeServer.Tests.Secrets
 		{
 			SecretCollectionInternal secretCollectionInternal = new([
 				new InMemoryProvider(),
-				new HcpVaultSecretProvider(new StubHttpClientFactory(), NullLogger<HcpVaultSecretProvider>.Instance)
+				new HcpVaultSecretProvider(new StubHttpClientFactory(), NullLogger<HcpVaultSecretProvider>.Instance),
+				new AzureKeyVaultSecretProvider()
 			]);
 			string json = """
 				{
@@ -126,6 +127,16 @@ namespace HordeServer.Tests.Secrets
 								}
 							]
 						},
+						{
+							"id": "secretUsingAzureKeyVault",
+							"sources": [
+								{
+									"providerConfig": "azure-key-vault",
+									"path": "test",
+									"key": "secret"
+								}
+							]
+						},
 						{
 							"id": "emptySecret",
 							"sources": [
@@ -178,6 +189,13 @@ namespace HordeServer.Tests.Secrets
 								"awsArnRole": "arn:aws:iam::1234567890:role/vault-example-role",
 								"role": "vault-example"
 							}
+						},
+						{
+							"name": "azure-key-vault",
+							"provider": "AzureKeyVault",
+							"azureKeyVault": {
+								"vaultUri": "https://example.vault.azure.net/"
+							}
 						}
 					]
 				}
@@ -262,6 +280,20 @@ namespace HordeServer.Tests.Secrets
 			Assert.AreEqual(secret?.Data["test"], "test");
 		}
 
+		[Ignore("Requires access to a Key Vault in Azure and an authorized Entra ID")]
+		[TestMethod]
+		public async Task GetSecretFromAzureKeyVaultAsync()
+		{
+			// To run this test from Visual Studio
+			// 1. add your Azure Entra account to Visual Studio
+			// 2. select that account in Azure Authentication settings
+			// 3. update Vault URI in the SecretProviderConfig options for the SecretId used in this test 
+			// 4. grant permissions for the Entra account to access the key vault and secret
+			SecretCollection secretCollection = CreateSecretCollection(); 
+			ISecret? secret = await secretCollection.GetAsync(new SecretId("SecretUsingAzureKeyVault"), CancellationToken.None);
+			Assert.AreEqual(secret?.Data["secret"], "test");
+		}
+
 		[TestMethod]
 		public async Task GetSecretFromMinimalJsonAsync()
 		{
diff --git a/Engine/Source/Programs/Horde/Plugins/Secrets/HordeServer.Secrets/HordeServer.Secrets.csproj b/Engine/Source/Programs/Horde/Plugins/Secrets/HordeServer.Secrets/HordeServer.Secrets.csproj
--- a/Engine/Source/Programs/Horde/Plugins/Secrets/HordeServer.Secrets/HordeServer.Secrets.csproj
+++ b/Engine/Source/Programs/Horde/Plugins/Secrets/HordeServer.Secrets/HordeServer.Secrets.csproj
@@ -23,6 +23,8 @@
 
   <ItemGroup>
     <PackageReference Include="AWSSDK.SimpleSystemsManagement" Version="3.7.304.24" />
+    <PackageReference Include="Azure.Identity" Version="1.14.2" />
+    <PackageReference Include="Azure.Security.KeyVault.Secrets" Version="4.8.0" />
   </ItemGroup>
 
   <ItemGroup>
diff --git a/Engine/Source/Programs/Horde/Plugins/Secrets/HordeServer.Secrets/Secrets/Providers/AzureKeyVaultConfig.cs b/Engine/Source/Programs/Horde/Plugins/Secrets/HordeServer.Secrets/Secrets/Providers/AzureKeyVaultConfig.cs
new file mode 100644
--- /dev/null
+++ b/Engine/Source/Programs/Horde/Plugins/Secrets/HordeServer.Secrets/Secrets/Providers/AzureKeyVaultConfig.cs
@@ -0,0 +1,14 @@
+// Copyright Epic Games, Inc. All Rights Reserved.
+
+namespace HordeServer.Secrets.Providers;
+
+/// <summary>
+/// Configuration for Azure Key Vault secrets.
+/// </summary>
+public class AzureKeyVaultConfig
+{
+	/// <summary>
+	/// The Uri to the Azure Key Vault.
+	/// </summary>
+	public string? VaultUri { get; set; }
+}
\ No newline at end of file
diff --git a/Engine/Source/Programs/Horde/Plugins/Secrets/HordeServer.Secrets/Secrets/Providers/AzureKeyVaultSecretProvider.cs b/Engine/Source/Programs/Horde/Plugins/Secrets/HordeServer.Secrets/Secrets/Providers/AzureKeyVaultSecretProvider.cs
new file mode 100644
--- /dev/null
+++ b/Engine/Source/Programs/Horde/Plugins/Secrets/HordeServer.Secrets/Secrets/Providers/AzureKeyVaultSecretProvider.cs
@@ -0,0 +1,33 @@
+// Copyright Epic Games, Inc. All Rights Reserved.
+
+using Azure;
+using Azure.Identity;
+using Azure.Security.KeyVault.Secrets;
+
+namespace HordeServer.Secrets.Providers;
+
+/// <summary>
+/// Fetches secrets from Azure Key Vault.
+/// </summary>
+public class AzureKeyVaultSecretProvider : ISecretProvider
+{
+	/// <inheritdoc/>
+	public string Name => "AzureKeyVault";
+
+	/// <inheritdoc/>
+	public async Task<string> GetSecretAsync(string path, SecretProviderConfig? config, CancellationToken cancellationToken)
+	{
+		if (config?.AzureKeyVault == null)
+		{
+			throw new InvalidOperationException($"Unable to fetch secret {path} from Azure Key Vault (No Options)");
+		}
+		if (config.AzureKeyVault.VaultUri == null)
+		{
+			throw new InvalidOperationException($"Unable to fetch secret {path} from Azure Key Vault (No Vault URI)");
+		}
+		DefaultAzureCredential credential = new();
+		SecretClient secretClient = new(new Uri(config.AzureKeyVault.VaultUri), credential);
+		Response<KeyVaultSecret>? secret = await secretClient.GetSecretAsync(path, cancellationToken: cancellationToken);
+		return secret.Value.Value;
+	}
+}
\ No newline at end of file
diff --git a/Engine/Source/Programs/Horde/Plugins/Secrets/HordeServer.Secrets/Secrets/SecretProviderConfig.cs b/Engine/Source/Programs/Horde/Plugins/Secrets/HordeServer.Secrets/Secrets/SecretProviderConfig.cs
--- a/Engine/Source/Programs/Horde/Plugins/Secrets/HordeServer.Secrets/Secrets/SecretProviderConfig.cs
+++ b/Engine/Source/Programs/Horde/Plugins/Secrets/HordeServer.Secrets/Secrets/SecretProviderConfig.cs
@@ -23,5 +23,10 @@ namespace HordeServer.Secrets
 		/// Configuration for HCP Vault
 		/// </summary>
 		public HcpVaultConfig? HcpVault { get; set; }
+
+		/// <summary>
+		/// Configuration for Azure Key Vault
+		/// </summary>
+		public AzureKeyVaultConfig? AzureKeyVault { get; set; }
 	}
 }
diff --git a/Engine/Source/Programs/Horde/Plugins/Secrets/HordeServer.Secrets/SecretsPlugin.cs b/Engine/Source/Programs/Horde/Plugins/Secrets/HordeServer.Secrets/SecretsPlugin.cs
--- a/Engine/Source/Programs/Horde/Plugins/Secrets/HordeServer.Secrets/SecretsPlugin.cs
+++ b/Engine/Source/Programs/Horde/Plugins/Secrets/HordeServer.Secrets/SecretsPlugin.cs
@@ -48,6 +48,7 @@ namespace HordeServer
 			}
 
 			services.AddSingleton<ISecretProvider, HcpVaultSecretProvider>();
+			services.AddSingleton<ISecretProvider, AzureKeyVaultSecretProvider>();
 		}
 	}
 
