Horde: Reduce verbosity of suspect addition/removal in audit log.
Start using debug level logging in issue audit log.

diff --git a/Engine/Source/Programs/Horde/Plugins/Build/HordeServer.Build/Issues/IssueCollection.cs b/Engine/Source/Programs/Horde/Plugins/Build/HordeServer.Build/Issues/IssueCollection.cs
--- a/Engine/Source/Programs/Horde/Plugins/Build/HordeServer.Build/Issues/IssueCollection.cs
+++ b/Engine/Source/Programs/Horde/Plugins/Build/HordeServer.Build/Issues/IssueCollection.cs
@@ -1021,15 +1021,23 @@ namespace HordeServer.Issues
 		{
 			HashSet<(UserId, CommitId)> oldSuspects = new HashSet<(UserId, CommitId)>(oldIssueSuspects.Select(x => (x.AuthorId, (CommitId)x.CommitId)));
 			HashSet<(UserId, CommitId)> newSuspects = new HashSet<(UserId, CommitId)>(newIssueSuspects.Select(x => (x.AuthorId, (CommitId)x.CommitId)));
+			int addedSuspects = 0;
+			int removedSuspects = 0;
+
 			foreach ((UserId userId, CommitId commit) in newSuspects.Where(x => !oldSuspects.Contains(x)))
 			{
-				issueLogger.LogInformation("Added suspect {UserName} ({UserId}) for commit {Commit}", await GetUserNameAsync(userId, cancellationToken), userId, commit);
+				issueLogger.LogDebug("Added suspect {UserName} ({UserId}) for commit {Commit}", await GetUserNameAsync(userId, cancellationToken), userId, commit);
+				addedSuspects++;
 			}
 			foreach ((UserId userId, CommitId commit) in oldSuspects.Where(x => !newSuspects.Contains(x)))
 			{
-				issueLogger.LogInformation("Removed suspect {UserName} ({UserId}) for commit {Commit}", await GetUserNameAsync(userId, cancellationToken), userId, commit);
+				issueLogger.LogDebug("Removed suspect {UserName} ({UserId}) for commit {Commit}", await GetUserNameAsync(userId, cancellationToken), userId, commit);
+				removedSuspects++;
 			}
 
+			issueLogger.LogInformation("Added {AddedSuspects} suspect(s). See Debug level logging for details.", addedSuspects);
+			issueLogger.LogInformation("Removed {RemovedSuspects} suspect(s). See Debug level logging for details.", removedSuspects);
+
 			HashSet<UserId> oldDeclinedBy = new HashSet<UserId>(oldIssueSuspects.Where(x => x.DeclinedAt != null).Select(x => x.AuthorId));
 			HashSet<UserId> newDeclinedBy = new HashSet<UserId>(newIssueSuspects.Where(x => x.DeclinedAt != null).Select(x => x.AuthorId));
 			foreach (UserId addDeclinedBy in newDeclinedBy.Where(x => !oldDeclinedBy.Contains(x)))
