Horde Dashboard: Usage of ErrorModal vs ErrorBar

diff --git a/Engine/Source/Programs/Horde/HordeDashboard/src/backend/Fetch.ts b/Engine/Source/Programs/Horde/HordeDashboard/src/backend/Fetch.ts
--- a/Engine/Source/Programs/Horde/HordeDashboard/src/backend/Fetch.ts
+++ b/Engine/Source/Programs/Horde/HordeDashboard/src/backend/Fetch.ts
@@ -37,13 +37,13 @@ function handleError(
 
     switch(reportType){
         case ErrorReportType.Dialog:
-            errorDialogStore.set(e, force ?? false, retryCallback);
+            errorDialogStore.set(e, !!force, retryCallback);
             break;
         case ErrorReportType.Bar:
-            errorBarStore.set(e, force ?? false, retryCallback);
+            errorBarStore.set(e, !!force, retryCallback);
             break;
         default:
-            errorDialogStore.set(e, force ?? false, retryCallback); // Make default Bar after dialog investigation
+            errorBarStore.set(e, !!force, retryCallback);
             break;
     }
 }
@@ -73,7 +73,7 @@ export class Fetch {
                                 mode: "GET",
                                 url: url,
                                 message: text
-                            }, true);            
+                            }, true, ErrorReportType.Dialog);            
                         });
                         return null;
                     }
@@ -88,7 +88,7 @@ export class Fetch {
                                 mode: "GET",
                                 url: url,
                                 message: o?.message ?? "Malformed json on response object"
-                            }, true);            
+                            }, true, ErrorReportType.Dialog);            
                         }).catch(reason => {
                             console.error("Unable to parse response json on 404: ", reason)
                         });
@@ -104,7 +104,7 @@ export class Fetch {
                             mode: "GET",
                             url,
                             message: "You don't have access because you are not logged in. Try logging in again."
-                        }, true, ErrorReportType.Bar)
+                        }, true)
                         return null;
                     }
 
@@ -117,7 +117,7 @@ export class Fetch {
                             mode: "GET",
                             url,
                             message: "You don't have access to this page or some of its features"
-                        }, true, ErrorReportType.Bar)
+                        }, true)
                         return null;
                     }
 
@@ -378,7 +378,7 @@ export class Fetch {
                     url: url,
                     title: "Internal Server Error",
                     message: text                
-                }, true);
+                }, true, ErrorReportType.Dialog);
             });
             return reject("Internal Server Error");
         }
@@ -399,7 +399,7 @@ export class Fetch {
                     response: response,
                     url: url,
                     title: "Access Denied"
-                });
+                }, true, ErrorReportType.Dialog);
 
             } else {
 
