Horde Dashboard: Remove PrintException Component

diff --git a/Engine/Source/Programs/Horde/HordeDashboard/src/components/LogView.tsx b/Engine/Source/Programs/Horde/HordeDashboard/src/components/LogView.tsx
--- a/Engine/Source/Programs/Horde/HordeDashboard/src/components/LogView.tsx
+++ b/Engine/Source/Programs/Horde/HordeDashboard/src/components/LogView.tsx
@@ -26,7 +26,6 @@ import { useQuery } from './JobDetailCommon';
 import { LogItem, renderLine } from './LogRender';
 import { JobLogSource, LogSource } from './LogSource';
 import { getLogStyles, logMetricNormal, logMetricSmall } from "./LogStyle";
-import { PrintException } from './PrintException';
 import { StepRefStatusIcon } from './StatusIcon';
 import { TopNav } from './TopNav';
 import { isBright } from 'horde/base/utilities/colors';
@@ -755,7 +754,9 @@ export const LogList: React.FC<{ logId: string }> = observer(({ logId }) => {
       const error = `Error getting job data, please check that you are logged in and that the link is valid.\n\n${handler.logSource.fatalError}`;
       return <Stack horizontal style={{ paddingTop: 48 }}>
          <div key={`windowsize_logview1_${windowSize.width}_${windowSize.height}`} style={{ width: vw / 2 - 720, flexShrink: 0, backgroundColor: hordeTheme.horde.contentBackground }} />
-         <Stack horizontalAlign="center" style={{ width: 1440 }}><PrintException message={error} /></Stack>
+         <Stack horizontalAlign="center" style={{ width: 1440 }}>
+            <div style={{whiteSpace:"pre"}}> {error} </div>
+         </Stack>
       </Stack>
    }
 
diff --git a/Engine/Source/Programs/Horde/HordeDashboard/src/components/PrintException.tsx b/Engine/Source/Programs/Horde/HordeDashboard/src/components/PrintException.tsx
deleted file mode 100644
