Horde: Update release notes

#rnx

diff --git a/Engine/Source/Programs/Horde/Docs/ReleaseNotes.md b/Engine/Source/Programs/Horde/Docs/ReleaseNotes.md
--- a/Engine/Source/Programs/Horde/Docs/ReleaseNotes.md
+++ b/Engine/Source/Programs/Horde/Docs/ReleaseNotes.md
@@ -1,5 +1,16 @@
 # Release Notes
 
+## 2025-07-24
+
+* Horde Dashboard: Remove PrintException Component (44281453)
+* Horde Dashboard: Usage of ErrorModal vs ErrorBar (44281424)
+* Reduce verbosity of suspect addition/removal in audit log. Start using debug level logging in issue audit log. (44280236)
+* Add example of using Azure Key Vault to Secrets plugin. The Azure Key Vault Secrets Provider is provided as experimental and an example of how to implement a secret provider from another cloud provider. (44274851)
+* A node that is modified and skipped during a build graph update does not fail the whole build. If there is a modified node with a skipped state during a build graph update the whole graph would previously be marked as failed, instead allow the build to continue by marking dependent nodes as skipped/failed but still allow other independent subtrees to be run. (44273401)
+* Update code coverage to use JetBrains.dotCover.CommandLineTools. The previuous code coverage tool JetBrains.dotCover.GlobalTool has been deprecated. (44273037)
+* Horde Dashboard: Better handling of advanced additional arguments in NewBuildV2 (44207109)
+* Horde Dashboard: Put HelpModal back into top level nav (44142211)
+
 ## 2025-07-17
 
 * Horde Dashboard: Clean up logic in LogRender.tsx (44106602)
