Horde Dashboard: Remove unused component files and exports

diff --git a/Engine/Source/Programs/Horde/HordeDashboard/src/backend/IssueStore.ts b/Engine/Source/Programs/Horde/HordeDashboard/src/backend/IssueStore.ts
--- a/Engine/Source/Programs/Horde/HordeDashboard/src/backend/IssueStore.ts
+++ b/Engine/Source/Programs/Horde/HordeDashboard/src/backend/IssueStore.ts
@@ -7,7 +7,7 @@ import backend from '.';
 const updateMS = 5000;
 
 // This class is intended for high level, cross job stream issue views, for job specific issues see JobDetails
-export class IssueStore {
+class IssueStore {
 
     constructor() {
         makeObservable(this);
diff --git a/Engine/Source/Programs/Horde/HordeDashboard/src/backend/TemplateCache.ts b/Engine/Source/Programs/Horde/HordeDashboard/src/backend/TemplateCache.ts
--- a/Engine/Source/Programs/Horde/HordeDashboard/src/backend/TemplateCache.ts
+++ b/Engine/Source/Programs/Horde/HordeDashboard/src/backend/TemplateCache.ts
@@ -2,7 +2,7 @@
 
 import { GetTemplateRefResponse, StreamData } from './Api';
 
-export class TemplateCache {
+class TemplateCache {
 
 	getStreamTemplates(stream: StreamData): Promise<GetTemplateRefResponse[]> {
 
diff --git a/Engine/Source/Programs/Horde/HordeDashboard/src/base/components/MarkdownTable/MarkdownCell.tsx b/Engine/Source/Programs/Horde/HordeDashboard/src/base/components/MarkdownTable/MarkdownCell.tsx
--- a/Engine/Source/Programs/Horde/HordeDashboard/src/base/components/MarkdownTable/MarkdownCell.tsx
+++ b/Engine/Source/Programs/Horde/HordeDashboard/src/base/components/MarkdownTable/MarkdownCell.tsx
@@ -6,7 +6,7 @@ import { PropsWithChildren } from 'react';
 
 const getClassNames = classNamesFunction<IMarkdownTableStyleProps, IMarkdownTableStyles>();
 
-export class MarkdownCellBase extends React.PureComponent<PropsWithChildren<IMarkdownTableCellProps>> {
+class MarkdownCellBase extends React.PureComponent<PropsWithChildren<IMarkdownTableCellProps>> {
   public render(): JSX.Element {
     const { as = 'td', children, styles, theme } = this.props;
 
diff --git a/Engine/Source/Programs/Horde/HordeDashboard/src/base/components/MarkdownTable/MarkdownTBody.tsx b/Engine/Source/Programs/Horde/HordeDashboard/src/base/components/MarkdownTable/MarkdownTBody.tsx
--- a/Engine/Source/Programs/Horde/HordeDashboard/src/base/components/MarkdownTable/MarkdownTBody.tsx
+++ b/Engine/Source/Programs/Horde/HordeDashboard/src/base/components/MarkdownTable/MarkdownTBody.tsx
@@ -6,7 +6,7 @@ import { PropsWithChildren } from 'react';
 
 const getClassNames = classNamesFunction<IMarkdownTableStyleProps, IMarkdownTableStyles>();
 
-export class MarkdownTBodyBase extends React.PureComponent<PropsWithChildren<IMarkdownTableProps>> {
+class MarkdownTBodyBase extends React.PureComponent<PropsWithChildren<IMarkdownTableProps>> {
   public render(): JSX.Element {
     const { children, styles, theme } = this.props;
 
diff --git a/Engine/Source/Programs/Horde/HordeDashboard/src/base/components/MarkdownTable/MarkdownTHead.tsx b/Engine/Source/Programs/Horde/HordeDashboard/src/base/components/MarkdownTable/MarkdownTHead.tsx
--- a/Engine/Source/Programs/Horde/HordeDashboard/src/base/components/MarkdownTable/MarkdownTHead.tsx
+++ b/Engine/Source/Programs/Horde/HordeDashboard/src/base/components/MarkdownTable/MarkdownTHead.tsx
@@ -6,7 +6,7 @@ import { PropsWithChildren } from 'react';
 
 const getClassNames = classNamesFunction<IMarkdownTableStyleProps, IMarkdownTableStyles>();
 
-export class MarkdownTHeadBase extends React.PureComponent<PropsWithChildren<IMarkdownTableProps>> {
+class MarkdownTHeadBase extends React.PureComponent<PropsWithChildren<IMarkdownTableProps>> {
   public render(): JSX.Element {
     const { children, styles, theme } = this.props;
 
diff --git a/Engine/Source/Programs/Horde/HordeDashboard/src/base/components/MarkdownTable/MarkdownTable.tsx b/Engine/Source/Programs/Horde/HordeDashboard/src/base/components/MarkdownTable/MarkdownTable.tsx
--- a/Engine/Source/Programs/Horde/HordeDashboard/src/base/components/MarkdownTable/MarkdownTable.tsx
+++ b/Engine/Source/Programs/Horde/HordeDashboard/src/base/components/MarkdownTable/MarkdownTable.tsx
@@ -6,7 +6,7 @@ import { PropsWithChildren } from 'react';
 
 const getClassNames = classNamesFunction<IMarkdownTableStyleProps, IMarkdownTableStyles>();
 
-export class MarkdownTableBase extends React.PureComponent<PropsWithChildren<IMarkdownTableProps>> {
+class MarkdownTableBase extends React.PureComponent<PropsWithChildren<IMarkdownTableProps>> {
   public render(): JSX.Element {
     const { children, className, styles, theme } = this.props;
 
diff --git a/Engine/Source/Programs/Horde/HordeDashboard/src/base/components/MarkdownTable/MarkdownTr.tsx b/Engine/Source/Programs/Horde/HordeDashboard/src/base/components/MarkdownTable/MarkdownTr.tsx
--- a/Engine/Source/Programs/Horde/HordeDashboard/src/base/components/MarkdownTable/MarkdownTr.tsx
+++ b/Engine/Source/Programs/Horde/HordeDashboard/src/base/components/MarkdownTable/MarkdownTr.tsx
@@ -6,7 +6,7 @@ import { PropsWithChildren } from 'react';
 
 const getClassNames = classNamesFunction<IMarkdownTableStyleProps, IMarkdownTableStyles>();
 
-export class MarkdownTrBase extends React.PureComponent<PropsWithChildren<IMarkdownTableProps>> {
+class MarkdownTrBase extends React.PureComponent<PropsWithChildren<IMarkdownTableProps>> {
   public render(): JSX.Element {
     const { children, styles, theme } = this.props;
 
diff --git a/Engine/Source/Programs/Horde/HordeDashboard/src/components/AuditLog.tsx b/Engine/Source/Programs/Horde/HordeDashboard/src/components/AuditLog.tsx
--- a/Engine/Source/Programs/Horde/HordeDashboard/src/components/AuditLog.tsx
+++ b/Engine/Source/Programs/Horde/HordeDashboard/src/components/AuditLog.tsx
@@ -402,7 +402,7 @@ const timeSelections: TimeSelection[] = [
 
 const selection = new Selection({ selectionMode: SelectionMode.multiple });
 
-export const AuditLogPanel: React.FC<{ agentId?: string, deviceId?: string, jobId?: string, issueId?: string, templateAudit?: TemplateAudit }> = observer(({ agentId, deviceId, jobId, issueId, templateAudit }) => {
+const AuditLogPanel: React.FC<{ agentId?: string, deviceId?: string, jobId?: string, issueId?: string, templateAudit?: TemplateAudit }> = observer(({ agentId, deviceId, jobId, issueId, templateAudit }) => {
 
    const windowSize = useWindowSize();
    const [showDatePicker, setShowDatePicker] = useState(false);
diff --git a/Engine/Source/Programs/Horde/HordeDashboard/src/components/AutoSubmit.tsx b/Engine/Source/Programs/Horde/HordeDashboard/src/components/AutoSubmit.tsx
deleted file mode 100644
diff --git a/Engine/Source/Programs/Horde/HordeDashboard/src/components/AutomationViewSummary.tsx b/Engine/Source/Programs/Horde/HordeDashboard/src/components/AutomationViewSummary.tsx
--- a/Engine/Source/Programs/Horde/HordeDashboard/src/components/AutomationViewSummary.tsx
+++ b/Engine/Source/Programs/Horde/HordeDashboard/src/components/AutomationViewSummary.tsx
@@ -865,7 +865,7 @@ const SuiteSummaryButton: React.FC<{ suite: GetTestSuiteResponse, handler: TestD
    </Stack>
 }
 
-export const SuiteSummary: React.FC<{ handler: TestDataHandler }> = ({ handler }) => {
+const SuiteSummary: React.FC<{ handler: TestDataHandler }> = ({ handler }) => {
 
    const { hordeClasses } = getHordeStyling();
 
diff --git a/Engine/Source/Programs/Horde/HordeDashboard/src/components/ChangeCallout.tsx b/Engine/Source/Programs/Horde/HordeDashboard/src/components/ChangeCallout.tsx
deleted file mode 100644
diff --git a/Engine/Source/Programs/Horde/HordeDashboard/src/components/ChangeSummary.tsx b/Engine/Source/Programs/Horde/HordeDashboard/src/components/ChangeSummary.tsx
deleted file mode 100644
diff --git a/Engine/Source/Programs/Horde/HordeDashboard/src/components/HomeView.tsx b/Engine/Source/Programs/Horde/HordeDashboard/src/components/HomeView.tsx
deleted file mode 100644
diff --git a/Engine/Source/Programs/Horde/HordeDashboard/src/components/IssueViewV2.tsx b/Engine/Source/Programs/Horde/HordeDashboard/src/components/IssueViewV2.tsx
--- a/Engine/Source/Programs/Horde/HordeDashboard/src/components/IssueViewV2.tsx
+++ b/Engine/Source/Programs/Horde/HordeDashboard/src/components/IssueViewV2.tsx
@@ -2990,7 +2990,7 @@ const EditIssueModal: React.FC<{ onClose: () => void }> = ({ onClose }) => {
 
 }
 
-export const IssueQuarantineModal: React.FC<{ onClose: () => void }> = ({ onClose }) => {
+const IssueQuarantineModal: React.FC<{ onClose: () => void }> = ({ onClose }) => {
 
    const issue = details.issue!;
    const issueId = issue.id;
@@ -3072,7 +3072,7 @@ export const IssueQuarantineModal: React.FC<{ onClose: () => void }> = ({ onClos
 };
 
 
-export const IssueForceCloseModal: React.FC<{ onClose: () => void }> = ({ onClose }) => {
+const IssueForceCloseModal: React.FC<{ onClose: () => void }> = ({ onClose }) => {
 
    const issue = details.issue!;
    const issueId = issue.id;
diff --git a/Engine/Source/Programs/Horde/HordeDashboard/src/components/JobFilterBar.tsx b/Engine/Source/Programs/Horde/HordeDashboard/src/components/JobFilterBar.tsx
--- a/Engine/Source/Programs/Horde/HordeDashboard/src/components/JobFilterBar.tsx
+++ b/Engine/Source/Programs/Horde/HordeDashboard/src/components/JobFilterBar.tsx
@@ -22,7 +22,7 @@ const statusItems = ["Running", "Complete", "Succeeded", "Failed", "Waiting", "C
    };
 });
 
-export class JobFilter {
+class JobFilter {
 
    constructor() {
       makeObservable(this);
diff --git a/Engine/Source/Programs/Horde/HordeDashboard/src/components/LogView.tsx b/Engine/Source/Programs/Horde/HordeDashboard/src/components/LogView.tsx
--- a/Engine/Source/Programs/Horde/HordeDashboard/src/components/LogView.tsx
+++ b/Engine/Source/Programs/Horde/HordeDashboard/src/components/LogView.tsx
@@ -638,7 +638,7 @@ const LogLineIndicator: React.FC<{ lineIndex: number }> = ({ lineIndex }) => {
 }
 
 
-export const LogList: React.FC<{ logId: string }> = observer(({ logId }) => {
+const LogList: React.FC<{ logId: string }> = observer(({ logId }) => {
 
    const windowSize = useWindowSize();
    const query = useQuery();
diff --git a/Engine/Source/Programs/Horde/HordeDashboard/src/components/NewBuild.tsx b/Engine/Source/Programs/Horde/HordeDashboard/src/components/NewBuild.tsx
--- a/Engine/Source/Programs/Horde/HordeDashboard/src/components/NewBuild.tsx
+++ b/Engine/Source/Programs/Horde/HordeDashboard/src/components/NewBuild.tsx
@@ -1936,7 +1936,7 @@ export const NewBuild: React.FC<{ streamId: string; show: boolean; onClose: (new
 });
 
 
-export const ValidationErrorModal: React.FC<{ errors: ValidationError[], show: boolean; onClose: () => void }> = ({ errors, show, onClose }) => {
+const ValidationErrorModal: React.FC<{ errors: ValidationError[], show: boolean; onClose: () => void }> = ({ errors, show, onClose }) => {
 
    const { hordeClasses } = getHordeStyling();
 
diff --git a/Engine/Source/Programs/Horde/HordeDashboard/src/components/PoolsView.tsx b/Engine/Source/Programs/Horde/HordeDashboard/src/components/PoolsView.tsx
--- a/Engine/Source/Programs/Horde/HordeDashboard/src/components/PoolsView.tsx
+++ b/Engine/Source/Programs/Horde/HordeDashboard/src/components/PoolsView.tsx
@@ -106,7 +106,7 @@ const stripeStyles = mergeStyleSets({
 
 });
 
-export const StatusBar = (stack: StatusBarStack[], width: number, height: number, basecolor?: string, style?: any): JSX.Element => {
+const StatusBar = (stack: StatusBarStack[], width: number, height: number, basecolor?: string, style?: any): JSX.Element => {
 
    stack = stack.filter(s => s.value > 0);
 
@@ -472,7 +472,7 @@ const PoolPicker: React.FC = observer(() => {
    </Stack>
 })
 
-export const PoolPivot: React.FC = () => {
+const PoolPivot: React.FC = () => {
 
    const { hordeClasses, modeColors } = getHordeStyling();
 
diff --git a/Engine/Source/Programs/Horde/HordeDashboard/src/components/Search.tsx b/Engine/Source/Programs/Horde/HordeDashboard/src/components/Search.tsx
deleted file mode 100644
diff --git a/Engine/Source/Programs/Horde/HordeDashboard/src/components/ServerSettings.tsx b/Engine/Source/Programs/Horde/HordeDashboard/src/components/ServerSettings.tsx
deleted file mode 100644
diff --git a/Engine/Source/Programs/Horde/HordeDashboard/src/components/StatusIcon.tsx b/Engine/Source/Programs/Horde/HordeDashboard/src/components/StatusIcon.tsx
--- a/Engine/Source/Programs/Horde/HordeDashboard/src/components/StatusIcon.tsx
+++ b/Engine/Source/Programs/Horde/HordeDashboard/src/components/StatusIcon.tsx
@@ -14,7 +14,7 @@ const StatusIcon: React.FC<{ iconName: string, style?: CSSProperties }> = ({ ico
    return <Stack className="horde-no-darktheme"><FontIcon style={{ ...style }} iconName={iconName} /></Stack>
 }
 
-export const getLabelIcon = (state: LabelState | undefined, outcome: LabelOutcome | undefined): { icon: string, color: string } => {
+const getLabelIcon = (state: LabelState | undefined, outcome: LabelOutcome | undefined): { icon: string, color: string } => {
 
    const colors = dashboard.getStatusColors();
 
diff --git a/Engine/Source/Programs/Horde/HordeDashboard/src/components/StreamView.tsx b/Engine/Source/Programs/Horde/HordeDashboard/src/components/StreamView.tsx
--- a/Engine/Source/Programs/Horde/HordeDashboard/src/components/StreamView.tsx
+++ b/Engine/Source/Programs/Horde/HordeDashboard/src/components/StreamView.tsx
@@ -24,41 +24,6 @@ import { StreamSummary } from './StreamSummary';
 import { TopNav } from './TopNav';
 import { NewBuildV2 } from './build/NewBuildV2';
 
-export const SummaryPage: React.FC = () => {
-
-   return (<Stack tokens={{ padding: 16 }} ><Text variant="mediumPlus">{`Summary Page`}</Text></Stack>);
-};
-
-export const GridView: React.FC = () => {
-
-   return (<Stack tokens={{ padding: 16 }} ><Text variant="mediumPlus">{`Grid View`}</Text></Stack>);
-};
-
-
-export const customClasses = mergeStyleSets({
-   filterBar: {
-      selectors: {
-         '.ms-Button': {
-            paddingLeft: 16,
-            paddingRight: 16,
-            minWidth: 80,
-            height: 32,
-            borderRadius: 2
-         },
-         '.ms-Button-label': {
-            fontFamily: "Horde Open Sans SemiBold !important"
-         }
-      }
-   },
-   raised: {
-      backgroundColor: "#ffffff",
-      boxShadow: "0 0 0 0 rgba(0,0,0,0.132), 0 0 0.9px 0 rgba(0,0,0,0.108)",
-      padding: "25px 12px 25px 30px",
-      margin: 0
-   }
-
-});
-
 type StreamIncemental = {
    streamId: string;
    template: GetTemplateRefResponse;
diff --git a/Engine/Source/Programs/Horde/HordeDashboard/src/components/SubmitQueue.tsx b/Engine/Source/Programs/Horde/HordeDashboard/src/components/SubmitQueue.tsx
deleted file mode 100644
diff --git a/Engine/Source/Programs/Horde/HordeDashboard/src/components/TestReportPanel.tsx b/Engine/Source/Programs/Horde/HordeDashboard/src/components/TestReportPanel.tsx
--- a/Engine/Source/Programs/Horde/HordeDashboard/src/components/TestReportPanel.tsx
+++ b/Engine/Source/Programs/Horde/HordeDashboard/src/components/TestReportPanel.tsx
@@ -1,94 +1,16 @@
 // Copyright Epic Games, Inc. All Rights Reserved.
 
 import React from 'react';
-import { mergeStyleSets } from '@fluentui/react/lib/Styling';
-import { Stack, Text } from '@fluentui/react';
-import { Link } from 'react-router-dom';
-import { TestData } from '../backend/Api';
+import { Stack } from '@fluentui/react';
 import { PluginMount } from  '../legacyPlugins/pluginTypes';
 import hordePlugins from  '../legacyPlugins';
 import { ComponentMount } from './TestReportView';
-import { getHordeStyling } from '../styles/Styles';
-import { getHordeTheme } from '../styles/theme';
-
-
-let _styles: any;
-const getStyles = () => {
-
-   const theme = getHordeTheme();
-
-   const styles = _styles ?? mergeStyleSets({
-       item: {
-           padding: 8,
-           borderBottom: '1px solid '+theme.palette.neutralLighter,
-           selectors: {
-               ':hover': {background: theme.palette.neutralLight}
-           }
-       }
-   });
-
-   _styles = styles;
-   return styles;
-   
-}
-
 
 type ComponentItem = {
     hasComponent: boolean;
     linkComponent?: React.FC<any>;
 }
 
-type TestDataItem = ComponentItem & {
-    name: string;
-    item: TestData;
-}
-
-const testReportComponentTypes : Map<string, ComponentItem> = new Map();
-
-export const TestReportPanel: React.FC<{ testdata: TestData[] }> = ({ testdata }) => {
-
-   const { hordeClasses } = getHordeStyling();
-   const styles = getStyles();
-
-    const testdataItems : TestDataItem[] = [];
-    testdata.forEach((test) => {
-        let components : ComponentItem = {hasComponent: false};
-        const splitKey = test.key.split('::', 2);
-        const type = splitKey[0];
-        const name = splitKey.length > 1 ? splitKey[1] : type;
-        if (testReportComponentTypes.has(type)) {
-            components = testReportComponentTypes.get(type)!;
-        }
-        else
-        {
-            let pluginComponents = hordePlugins.getComponents(PluginMount.TestReportPanel, type);
-            components.hasComponent = pluginComponents.length > 0
-            testReportComponentTypes.set(type, components);
-            if (components.hasComponent) {
-                pluginComponents = hordePlugins.getComponents(PluginMount.TestReportLink, type);
-                if (pluginComponents.length > 0) {
-                    components.linkComponent = pluginComponents[0].component;
-                }
-            }
-        }
-        components.hasComponent && testdataItems.push({...components, item: test, name: name})
-    });
-
-    return (<Stack styles={{ root: { paddingTop: 18, paddingRight: 12 } }}>
-        <Stack className={hordeClasses.raised}>
-            <Stack tokens={{ childrenGap: 12 }}>
-                <Text variant="mediumPlus" styles={{ root: { fontFamily: "Horde Open Sans SemiBold" } }}>Test Report</Text>
-                <Stack styles={{ root: { paddingLeft: 4, paddingRight: 0, paddingTop: 8, paddingBottom: 4 } }}>
-                    {testdataItems.map((test) => <Stack key={test.item.id} className={styles.item} horizontal wrap tokens={{ childrenGap: 30 }}>
-                        <Link className={"view-log-link"} to={`/testreport/${test.item.id}`}>{test.name}</Link>
-                        <ComponentMount component={test.linkComponent} id={test.item.id} fullName={test.item.key}/>
-                    </Stack>)}
-                </Stack>
-            </Stack>
-        </Stack>
-    </Stack>);
-};
-
 export const BuildHealthTestReportPanel: React.FC<{ streamId: string }> = ({ streamId }) => {
     const items: ComponentItem[] = hordePlugins.getComponents(PluginMount.BuildHealthSummary).map((item) => {
         return {hasComponent: true, linkComponent: item.component};
diff --git a/Engine/Source/Programs/Horde/HordeDashboard/src/components/ToolView.tsx b/Engine/Source/Programs/Horde/HordeDashboard/src/components/ToolView.tsx
--- a/Engine/Source/Programs/Horde/HordeDashboard/src/components/ToolView.tsx
+++ b/Engine/Source/Programs/Horde/HordeDashboard/src/components/ToolView.tsx
@@ -317,7 +317,7 @@ const ToolPanel: React.FC<{ selectedKey: string }> = observer(({ selectedKey })
    </Stack>
 });
 
-export const ToolViewInner: React.FC = observer(() => {
+const ToolViewInner: React.FC = observer(() => {
 
    const [selectedKey, setSelectedKey] = useState<string>(defaultCategory);
 
diff --git a/Engine/Source/Programs/Horde/HordeDashboard/src/components/UnderConstruction.tsx b/Engine/Source/Programs/Horde/HordeDashboard/src/components/UnderConstruction.tsx
deleted file mode 100644
diff --git a/Engine/Source/Programs/Horde/HordeDashboard/src/components/Welcome.tsx b/Engine/Source/Programs/Horde/HordeDashboard/src/components/Welcome.tsx
deleted file mode 100644
diff --git a/Engine/Source/Programs/Horde/HordeDashboard/src/components/agents/AgentStore.ts b/Engine/Source/Programs/Horde/HordeDashboard/src/components/agents/AgentStore.ts
--- a/Engine/Source/Programs/Horde/HordeDashboard/src/components/agents/AgentStore.ts
+++ b/Engine/Source/Programs/Horde/HordeDashboard/src/components/agents/AgentStore.ts
@@ -4,7 +4,7 @@ import { action, makeObservable, observable } from 'mobx';
 import backend from '../../backend';
 import { AgentData, CategoryAgents, GetAgentResponse, PoolData } from '../../backend/Api';
 
-export class AgentStore {
+class AgentStore {
 
     constructor() {
         makeObservable(this);
diff --git a/Engine/Source/Programs/Horde/HordeDashboard/src/components/agents/AgentView.tsx b/Engine/Source/Programs/Horde/HordeDashboard/src/components/agents/AgentView.tsx
--- a/Engine/Source/Programs/Horde/HordeDashboard/src/components/agents/AgentView.tsx
+++ b/Engine/Source/Programs/Horde/HordeDashboard/src/components/agents/AgentView.tsx
@@ -1451,7 +1451,7 @@ const agentContextMenuProps: IContextualMenuItem[] = [
 const agentStatus = ["Active", "Ready", "Disabled", "Ephemeral", "No Pools", "Pending Conform", "Pending Shutdown", "Offline", "Offline (Autoscaler)", "Offline (Manual)", "Offline (Unexpected)"];
 
 
-export const AgentMenuBar: React.FC<{ agentView?: boolean }> = observer(({ agentView }) => {
+const AgentMenuBar: React.FC<{ agentView?: boolean }> = observer(({ agentView }) => {
 
    let selectedButton: any = <div></div>;
    if (localState.agentsSelectedCount !== 0) {
@@ -1532,7 +1532,7 @@ export const AgentMenuBar: React.FC<{ agentView?: boolean }> = observer(({ agent
 });
 
 
-export const AgentPivot: React.FC = () => {
+const AgentPivot: React.FC = () => {
 
    const { hordeClasses, modeColors } = getHordeStyling();
 
@@ -1615,7 +1615,7 @@ const PoolEditorConfirmation: React.FC = observer(() => {
 
 });
 
-export const PoolEditorModal: React.FC = observer(() => {
+const PoolEditorModal: React.FC = observer(() => {
 
    const [state, setState] = useState<{ sortBy: "Agents" | "Name" }>({ sortBy: "Name" });
    const { hordeClasses } = getHordeStyling();
@@ -1859,7 +1859,7 @@ export const PoolEditorModal: React.FC = observer(() => {
    );
 });
 
-export const PoolSelectionModal: React.FC = observer(() => {
+const PoolSelectionModal: React.FC = observer(() => {
 
    const agentStyles = getAgentStyles();
 
@@ -1968,7 +1968,7 @@ export const PoolSelectionModal: React.FC = observer(() => {
    );
 });
 
-export const SearchUpdate: React.FC = observer(() => {
+const SearchUpdate: React.FC = observer(() => {
 
    const [searchParams, setSearchParams] = useSearchParams();
    const [state, setState] = useState<{ search?: string }>({ search: searchParams.toString() });
@@ -1986,7 +1986,7 @@ export const SearchUpdate: React.FC = observer(() => {
    return null;
 });
 
-export const AgentViewInner: React.FC<{ agentId?: string, poolId?: string, searchParams?: URLSearchParams, agentView?: boolean }> = observer(({ agentId, poolId, searchParams, agentView }) => {
+const AgentViewInner: React.FC<{ agentId?: string, poolId?: string, searchParams?: URLSearchParams, agentView?: boolean }> = observer(({ agentId, poolId, searchParams, agentView }) => {
 
    poolId = poolId?.toLowerCase();
 
diff --git a/Engine/Source/Programs/Horde/HordeDashboard/src/components/agents/HistoryModalStore.ts b/Engine/Source/Programs/Horde/HordeDashboard/src/components/agents/HistoryModalStore.ts
--- a/Engine/Source/Programs/Horde/HordeDashboard/src/components/agents/HistoryModalStore.ts
+++ b/Engine/Source/Programs/Horde/HordeDashboard/src/components/agents/HistoryModalStore.ts
@@ -24,7 +24,7 @@ type InfoPanelSubItem = {
    value: string;
 }
 
-export class HistoryModalStore {
+class HistoryModalStore {
    constructor() {
       makeObservable(this);
    }
diff --git a/Engine/Source/Programs/Horde/HordeDashboard/src/components/bisection/BisectionList.tsx b/Engine/Source/Programs/Horde/HordeDashboard/src/components/bisection/BisectionList.tsx
--- a/Engine/Source/Programs/Horde/HordeDashboard/src/components/bisection/BisectionList.tsx
+++ b/Engine/Source/Programs/Horde/HordeDashboard/src/components/bisection/BisectionList.tsx
@@ -605,7 +605,7 @@ const customStyles = mergeStyleSets({
 });
 
 
-export const CancelBisectionModal: React.FC<{ bisectTaskId: string, onClose: (canceled: boolean) => void }> = ({ bisectTaskId, onClose }) => {
+const CancelBisectionModal: React.FC<{ bisectTaskId: string, onClose: (canceled: boolean) => void }> = ({ bisectTaskId, onClose }) => {
 
    const [state, setState] = useState<{ submitting?: boolean, error?: string }>({});
 
diff --git a/Engine/Source/Programs/Horde/HordeDashboard/src/components/bisection/BisectionModal.tsx b/Engine/Source/Programs/Horde/HordeDashboard/src/components/bisection/BisectionModal.tsx
deleted file mode 100644
diff --git a/Engine/Source/Programs/Horde/HordeDashboard/src/components/build/NewBuildV2.tsx b/Engine/Source/Programs/Horde/HordeDashboard/src/components/build/NewBuildV2.tsx
--- a/Engine/Source/Programs/Horde/HordeDashboard/src/components/build/NewBuildV2.tsx
+++ b/Engine/Source/Programs/Horde/HordeDashboard/src/components/build/NewBuildV2.tsx
@@ -2208,7 +2208,7 @@ export const NewBuildV2: React.FC<{ streamId: string; show: boolean; onClose: (n
    </Stack>
 }
 
-export const ValidationErrorModal: React.FC<{ errors: ValidationError[], onClose: () => void }> = ({ errors, onClose }) => {
+const ValidationErrorModal: React.FC<{ errors: ValidationError[], onClose: () => void }> = ({ errors, onClose }) => {
 
    const { hordeClasses } = getHordeStyling();
 
diff --git a/Engine/Source/Programs/Horde/HordeDashboard/src/components/devices/DevicePoolTelemetry.tsx b/Engine/Source/Programs/Horde/HordeDashboard/src/components/devices/DevicePoolTelemetry.tsx
--- a/Engine/Source/Programs/Horde/HordeDashboard/src/components/devices/DevicePoolTelemetry.tsx
+++ b/Engine/Source/Programs/Horde/HordeDashboard/src/components/devices/DevicePoolTelemetry.tsx
@@ -217,7 +217,7 @@ const DistributionGraph: React.FC<{ handler: PoolTelemetryHandler, platform?: st
       </Stack>
 }
 
-export const DevicePoolGraph: React.FC<{ handler: PoolTelemetryHandler, pool?: string, platform?: string, onChange: (options: object) => void }> = ({ handler, pool, platform, onChange }) => {
+const DevicePoolGraph: React.FC<{ handler: PoolTelemetryHandler, pool?: string, platform?: string, onChange: (options: object) => void }> = ({ handler, pool, platform, onChange }) => {
 
    const [container, setContainer] = useState<HTMLDivElement | null>(null);
    const graph = useRef<PoolTelemetryGraph>(new PoolTelemetryGraph(handler));
@@ -260,7 +260,7 @@ export const DevicePoolGraph: React.FC<{ handler: PoolTelemetryHandler, pool?: s
       </Stack>
 }
 
-export const TelemetryDropdown: React.FC<{ handler: PoolTelemetryHandler, selectedPool?: string, selectedPlatform?: string, isPoolDisabled?: boolean, isPlatformDisabled?: boolean, onChange: (options: object) => void }> = ({ handler, selectedPool, selectedPlatform, isPoolDisabled, isPlatformDisabled, onChange }) => {
+const TelemetryDropdown: React.FC<{ handler: PoolTelemetryHandler, selectedPool?: string, selectedPlatform?: string, isPoolDisabled?: boolean, isPlatformDisabled?: boolean, onChange: (options: object) => void }> = ({ handler, selectedPool, selectedPlatform, isPoolDisabled, isPlatformDisabled, onChange }) => {
    const pools = Array.from(handler.pools.values()).sort((a, b) => a.name.localeCompare(b.name));
    const platforms = Array.from(handler.platforms.values()).sort((a, b) => a.name.localeCompare(b.name));
 
diff --git a/Engine/Source/Programs/Horde/HordeDashboard/src/components/devices/DeviceView.tsx b/Engine/Source/Programs/Horde/HordeDashboard/src/components/devices/DeviceView.tsx
--- a/Engine/Source/Programs/Horde/HordeDashboard/src/components/devices/DeviceView.tsx
+++ b/Engine/Source/Programs/Horde/HordeDashboard/src/components/devices/DeviceView.tsx
@@ -246,7 +246,7 @@ const pivotKeyAutomation = "pivot-key-automation";
 const pivotKeyShared = "pivot-key-shared";
 const localState = new LocalState();
 
-export const SearchUpdate: React.FC = observer(() => {
+const SearchUpdate: React.FC = observer(() => {
 
    const [, setSearchParams] = useSearchParams();
 
@@ -949,7 +949,7 @@ const DevicePanel: React.FC = observer(() => {
    );
 });
 
-export const CheckoutConfirmModal: React.FC<{ check: "in" | "out" | "error", devices: GetDeviceResponse[], onClose: () => void }> = ({ check, devices, onClose }) => {
+const CheckoutConfirmModal: React.FC<{ check: "in" | "out" | "error", devices: GetDeviceResponse[], onClose: () => void }> = ({ check, devices, onClose }) => {
 
    const { hordeClasses } = getHordeStyling();
 
diff --git a/Engine/Source/Programs/Horde/HordeDashboard/src/components/docs/DocView.tsx b/Engine/Source/Programs/Horde/HordeDashboard/src/components/docs/DocView.tsx
--- a/Engine/Source/Programs/Horde/HordeDashboard/src/components/docs/DocView.tsx
+++ b/Engine/Source/Programs/Horde/HordeDashboard/src/components/docs/DocView.tsx
@@ -21,7 +21,7 @@ type State = {
 }
 
 
-export const docClasses = mergeStyleSets({
+const docClasses = mergeStyleSets({
    raised: {
       boxShadow: "0 1.6px 3.6px 0 rgba(0,0,0,0.132), 0 0.3px 0.9px 0 rgba(0,0,0,0.108)",
       padding: "32px 40px"
diff --git a/Engine/Source/Programs/Horde/HordeDashboard/src/components/jobDetailsV2/JobDetailStepHistoryV2.tsx b/Engine/Source/Programs/Horde/HordeDashboard/src/components/jobDetailsV2/JobDetailStepHistoryV2.tsx
deleted file mode 100644
diff --git a/Engine/Source/Programs/Horde/HordeDashboard/src/components/jobDetailsV2/JobDetailViewSteps.tsx b/Engine/Source/Programs/Horde/HordeDashboard/src/components/jobDetailsV2/JobDetailViewSteps.tsx
--- a/Engine/Source/Programs/Horde/HordeDashboard/src/components/jobDetailsV2/JobDetailViewSteps.tsx
+++ b/Engine/Source/Programs/Horde/HordeDashboard/src/components/jobDetailsV2/JobDetailViewSteps.tsx
@@ -157,7 +157,7 @@ const RenderDynamic: React.FC<{ jobDetails: JobDetailsV2, dataView: StepsDataVie
 });
 
 
-export const StepsPanelInner: React.FC<{ jobDetails: JobDetailsV2, depStepId?: string }> = observer(({ jobDetails, depStepId }) => {
+const StepsPanelInner: React.FC<{ jobDetails: JobDetailsV2, depStepId?: string }> = observer(({ jobDetails, depStepId }) => {
 
    const { jobId } = useParams<{ jobId: string }>();
    const [lastSelectedAgent, setLastSelectedAgent] = useState<string | undefined>(undefined);
