Horde Dashboard: Improved navigation with job search link sharing

diff --git a/Engine/Source/Programs/Horde/HordeDashboard/src/components/JobSearchSimple.tsx b/Engine/Source/Programs/Horde/HordeDashboard/src/components/JobSearchSimple.tsx
--- a/Engine/Source/Programs/Horde/HordeDashboard/src/components/JobSearchSimple.tsx
+++ b/Engine/Source/Programs/Horde/HordeDashboard/src/components/JobSearchSimple.tsx
@@ -2,7 +2,7 @@ import { Checkbox, ContextualMenuItemType, DatePicker, DefaultButton, DetailsLis
 import { observer } from "mobx-react-lite";
 import moment from "moment";
 import { useEffect, useMemo, useState } from "react";
-import { Link, useSearchParams } from "react-router-dom";
+import { Link, useLocation, useSearchParams } from "react-router-dom";
 import backend from "../backend";
 import { GetJobResponse, GetJobsTabResponse, GetJobStepRefResponse, JobData, JobQuery, TabType, GetTemplateRefResponse, GetJobsTabLabelColumnResponse, GetLabelStateResponse, LabelOutcome, LabelState, GetJobsTabColumnResponse, JobStepOutcome } from "../backend/Api";
 import dashboard from "../backend/Dashboard";
@@ -18,6 +18,7 @@ import { StreamChooser } from "./projects/StreamChooser";
 import { getLabelColor } from "horde/styles/colors";
 import { CalloutController, JobLabelCallout } from "./JobLabelCallout";
 import LZString from "lz-string"
+import { errorBarStore } from "./error/ErrorStore";
 
 type SearchState = {
    templates: Set<string>, 
@@ -54,6 +55,7 @@ export const JobSearchSimpleModal: React.FC<{ streamIdIn: string, changeListIn?:
    const [searchParams, setSearchParams] = useSearchParams();
    const [templateData, setTemplateData] = useState<{ streamId: string; templates: GetTemplateRefResponse[] } | undefined>(undefined);
    const [controller] = useState(new CalloutController());
+   const location = useLocation();
    const [isFirstRender, setIsFirstRender] = useState(true);
 
    const { hordeClasses, detailClasses } = getHordeStyling();
@@ -282,41 +284,69 @@ export const JobSearchSimpleModal: React.FC<{ streamIdIn: string, changeListIn?:
    }
    
    // Clear search parameters from the URL
-   const clearSearchParams = () => {
-      const tab = searchParams.get("tab")!;
-      setSearchParams("");
-
-      const search = new URLSearchParams();
-      if(tab) search.append("tab", tab);
-
-      setSearchParams(search);
+   const clearSearchParams = (removeFromURL: boolean = true) => {
+      if(removeFromURL) {
+         const tab = searchParams.get("tab")!;
+         setSearchParams("");
+         
+         const search = new URLSearchParams();
+         if(tab) search.append("tab", tab);
+         setSearchParams(search);
+      }
+      else {
+         const search = new URLSearchParams(searchParams);
+         search.set("search", "");
+         setSearchParams(search);
+      }
    }
 
    // Sets state from search params and autoloads jobs
    useEffect(() => {
-      if(!searchParams.has("search") || !isFirstRender || !templateData) return;
-      
+      // StreamView's param detection + the 'get' check below allow linking to an empty search modal "https://...?tab=...&search="
+      if(!searchParams.has("search") || !searchParams.get("search") || !isFirstRender || !templateData) return;
+
       setIsFirstRender(false);
       
-      // Decompress URL data
-      const stringyState = JSON.parse(LZString.decompressFromEncodedURIComponent(searchParams.get("search")!));
-      
-      const state: SearchState = {
-         ...stringyState,
-         minDate: stringyState.minDate ? new Date(stringyState.minDate) : undefined,
-         maxDate: stringyState.maxDate ? new Date(stringyState.maxDate) : undefined,
-         templates: new Set(stringyState.templates),
-         statuses: new Set(stringyState.statuses)
-      };
-      
-      if(state.global){
-         streamChooserId++;
-         state.templates = new Set();
-         state.streamIdOverride = undefined;
+      const searchValue = searchParams.get("search")!;
+      try {
+         // Decompress URL data
+         const stringyState = JSON.parse(LZString.decompressFromEncodedURIComponent(searchValue));
+         
+         const state: SearchState = {
+            ...stringyState,
+            minDate: stringyState.minDate ? new Date(stringyState.minDate) : undefined,
+            maxDate: stringyState.maxDate ? new Date(stringyState.maxDate) : undefined,
+            templates: new Set(stringyState.templates),
+            statuses: new Set(stringyState.statuses)
+         };
+         
+         if(state.global){
+            streamChooserId++;
+            state.templates = new Set();
+            state.streamIdOverride = undefined;
+         }
+
+         queryJobs(state);
+      } catch(err) {
+         const truncateMsg = searchValue.length > 50;
+         
+         // Truncate param string for display friendliness. Doesn't offer complete information, but better than nothing.
+         const searchValueMsg = truncateMsg ? searchValue.slice(0,50) + "..." : searchValue;
+         errorBarStore.set({title: "Job Search Error", message: `Could not load search from URL search param \'${searchValueMsg}\'`}, true)
+         
+         // Best to offer complete information somewhere
+         if(truncateMsg){
+            console.error("Bad search param value:\n\n" + searchValue + "\n\n");
+         }
+
+         clearSearchParams();
       }
-      
-      queryJobs(state)
-   }, [templateData])
+   }, [templateData, location.search])
+
+   // Enable retrigger of url state loading when navigating between the result of a clear action and a search action
+   useEffect(() => {
+      setIsFirstRender(true);
+   }, [location.search])
 
    if (!stream) {
       console.error("unable to get stream");
@@ -950,6 +980,7 @@ export const JobSearchSimpleModal: React.FC<{ streamIdIn: string, changeListIn?:
                               disabled={searchResults.querying} 
                               text="Search" 
                               onClick={() => {
+                                 setIsFirstRender(false); // Important, prevents double query
                                  queryJobs(searchState);
                               }} 
                            />
@@ -958,7 +989,7 @@ export const JobSearchSimpleModal: React.FC<{ streamIdIn: string, changeListIn?:
                            <DefaultButton disabled={searchResults.querying} text="Clear" onClick={() => { 
                                  setSearchState({ preflights: false, userId: undefined, templates: new Set(), statuses: new Set()}) 
                                  setSearchResults({items: [], groups: [], querying: false});
-                                 clearSearchParams();
+                                 clearSearchParams(false);
                               }} />
                         </Stack>
                      </Stack>
diff --git a/Engine/Source/Programs/Horde/HordeDashboard/src/components/StreamView.tsx b/Engine/Source/Programs/Horde/HordeDashboard/src/components/StreamView.tsx
--- a/Engine/Source/Programs/Horde/HordeDashboard/src/components/StreamView.tsx
+++ b/Engine/Source/Programs/Horde/HordeDashboard/src/components/StreamView.tsx
@@ -3,8 +3,8 @@
 import { DefaultButton, FontIcon, HoverCard, HoverCardType, IContextualMenuProps, Pivot, PivotItem, PrimaryButton, Spinner, SpinnerSize, Stack, Text, TextField, mergeStyleSets } from '@fluentui/react';
 import { action, makeObservable, observable } from 'mobx';
 import { observer } from 'mobx-react-lite';
-import React, { useState } from 'react';
-import { Link, Navigate, useNavigate, useParams, useSearchParams } from 'react-router-dom';
+import React, { useEffect, useState } from 'react';
+import { Link, Navigate, useLocation, useNavigate, useParams, useSearchParams } from 'react-router-dom';
 import backend, { useBackend } from '../backend';
 import { GetJobsTabResponse, GetLabelStateResponse, GetStreamTabResponse, GetTemplateRefResponse, JobData, JobsTabData, LabelOutcome, LabelState, ProjectData, TabStyle } from '../backend/Api';
 import dashboard, { StatusColor } from '../backend/Dashboard';
@@ -309,6 +309,7 @@ const StreamViewInner: React.FC = observer(() => {
 
    const [shown, setShown] = useState(query.get("newbuild") ? true : false);
    const [searchParams] = useSearchParams();
+   const location = useLocation();
    const searchInURL: boolean = searchParams.has("search");
    const [findJobsShown, setFindJobsShown] = useState(searchInURL);
 
@@ -321,6 +322,12 @@ const StreamViewInner: React.FC = observer(() => {
    const stream = projectStore.streamById(streamId);
    const project = stream?.project;
 
+   // Respond to search param even when navigating between search queries within the same route
+   // One consequence being pressing the browser's back button can sometimes close the search modal depending on search state
+   useEffect(() => {
+      setFindJobsShown(searchParams.has("search"));
+   }, [location.search])
+
    //let newBuildVersion: string = query.get("newbuildversion") ? query.get("newbuildversion") : "2";
    const newBuildVersion: string = "2"
 
