Horde: Include commit ID in suspect ranking exceptions

diff --git a/Engine/Source/Programs/Horde/Plugins/Build/HordeServer.Build/Issues/IssueService.cs b/Engine/Source/Programs/Horde/Plugins/Build/HordeServer.Build/Issues/IssueService.cs
--- a/Engine/Source/Programs/Horde/Plugins/Build/HordeServer.Build/Issues/IssueService.cs
+++ b/Engine/Source/Programs/Horde/Plugins/Build/HordeServer.Build/Issues/IssueService.cs
@@ -1342,12 +1342,12 @@ namespace HordeServer.Issues
 					}
 					catch (NullReferenceException nullRefEx)
 					{
-						_logger.LogWarning("Exception while ranking suspect changes: {Message} - {Callstack}", nullRefEx.Message,  nullRefEx.StackTrace);
+						_logger.LogError("Exception ranking suspect change {Commit}: {Message} - {Callstack}", commit.Id, nullRefEx.Message, nullRefEx.StackTrace);
 
 					}
 					catch (EpicGames.Perforce.PerforceException perforceEx)
 					{
-						_logger.LogWarning("Exception while ranking suspect changes: {Message} - {Callstack}", perforceEx.Message, perforceEx.StackTrace);
+						_logger.LogError("Exception ranking suspect change {Commit}: {Message} - {Callstack}", commit.Id, perforceEx.Message, perforceEx.StackTrace);
 
 						// re-throw any perforce exceptions that aren't known or expected here
 						if (!perforceEx.Message.Contains("does not contain any files in", System.StringComparison.OrdinalIgnoreCase))
