Updated Horde to create partitioned workspaces by default.

diff --git a/Engine/Source/Programs/Horde/Plugins/Build/HordeServer.Build/BuildConfig.cs b/Engine/Source/Programs/Horde/Plugins/Build/HordeServer.Build/BuildConfig.cs
--- a/Engine/Source/Programs/Horde/Plugins/Build/HordeServer.Build/BuildConfig.cs
+++ b/Engine/Source/Programs/Horde/Plugins/Build/HordeServer.Build/BuildConfig.cs
@@ -515,7 +515,7 @@ namespace HordeServer
 		/// <summary>
 		/// Whether to use partitioned workspaces on this server
 		/// </summary>
-		public bool SupportsPartitionedWorkspaces { get; set; } = false;
+		public bool SupportsPartitionedWorkspaces { get; set; } = true;
 
 		/// <summary>
 		/// List of servers
