Horde Dashboard: Extract useQuery and useSearchParams into single hook

diff --git a/Engine/Source/Programs/Horde/HordeDashboard/src/base/utilities/hooks.ts b/Engine/Source/Programs/Horde/HordeDashboard/src/base/utilities/hooks.ts
--- a/Engine/Source/Programs/Horde/HordeDashboard/src/base/utilities/hooks.ts
+++ b/Engine/Source/Programs/Horde/HordeDashboard/src/base/utilities/hooks.ts
@@ -1,5 +1,6 @@
 import { useState, useEffect } from 'react';
-
+import { useSearchParams } from 'react-router-dom';
+import LZString from "lz-string"
 
 export function useWindowSize():{width:number, height:number} {
 
@@ -26,3 +27,52 @@ export function useWindowSize():{width:number, height:number} {
 	return windowSize;
 }
 
+export function useQuery() {
+	const [searchParams, setSearchParams] = useSearchParams();
+
+	const get = (key: string): string | null => {
+		return searchParams.get(key);
+	}
+
+	const getAll = (key: string): string[] => {
+		return searchParams.getAll(key);
+	}
+
+	const set = (key: string, value: string) => {
+		const newParams = new URLSearchParams(searchParams);
+		newParams.set(key, value);
+		setSearchParams(newParams);
+	}
+
+	const has = (key: string, value?: string): boolean => {
+		return searchParams.has(key, value);
+	}
+
+	const clear = (keepList: string[] = []) => {
+		const newParams = new URLSearchParams();
+
+		for(const key of keepList){
+			const value = get(key);
+			if(value) newParams.set(key, value);
+		}
+
+		setSearchParams(newParams);
+	}
+
+	const setCompressed = (key: string, value: string) => {
+		const compressedValue = LZString.compressToEncodedURIComponent(value);
+		set(key, compressedValue);
+	}
+
+	const getCompressed = (key: string): string | null => {
+		const compressedValue = get(key);
+		if(compressedValue){
+			return LZString.decompressFromEncodedURIComponent(compressedValue);
+		} else {
+			console.error("Could not find key \'" + key + "\' to decompress.");
+			return null;
+		}
+	}
+
+	return {get, getAll, set, has, clear, setCompressed, getCompressed, raw: {searchParams, setSearchParams}}
+}
\ No newline at end of file
diff --git a/Engine/Source/Programs/Horde/HordeDashboard/src/components/IssueViewV2.tsx b/Engine/Source/Programs/Horde/HordeDashboard/src/components/IssueViewV2.tsx
--- a/Engine/Source/Programs/Horde/HordeDashboard/src/components/IssueViewV2.tsx
+++ b/Engine/Source/Programs/Horde/HordeDashboard/src/components/IssueViewV2.tsx
@@ -15,7 +15,7 @@ import { displayTimeZone, getHumanTime, getShortNiceTime } from "../base/utiliti
 import { getHordeStyling } from "../styles/Styles";
 import { getHordeTheme } from "../styles/theme";
 import { errorDialogStore } from "./error/ErrorStore";
-import { useQuery } from "./JobDetailCommon";
+import { useQuery } from "horde/base/utilities/hooks";
 import { renderLine } from "./LogRender";
 import { UserSelect } from "./UserSelect";
 
diff --git a/Engine/Source/Programs/Horde/HordeDashboard/src/components/JobDetailCommon.tsx b/Engine/Source/Programs/Horde/HordeDashboard/src/components/JobDetailCommon.tsx
--- a/Engine/Source/Programs/Horde/HordeDashboard/src/components/JobDetailCommon.tsx
+++ b/Engine/Source/Programs/Horde/HordeDashboard/src/components/JobDetailCommon.tsx
@@ -1,6 +1,5 @@
 // Copyright Epic Games, Inc. All Rights Reserved.
 
-import { useLocation } from 'react-router-dom';
 import { BatchData, GetJobStepRefResponse, JobStepBatchError, JobStepBatchState, JobStepOutcome, JobStepState, StepData } from '../backend/Api';
 
 type StepItem = {
@@ -12,10 +11,6 @@ type StepItem = {
    agentPool?: string;
 };
 
-export function useQuery() {
-   return new URLSearchParams(useLocation().search);
-}
-
 export const getStepStatusMessage = (step: StepData) => {
 
    let message = "";
diff --git a/Engine/Source/Programs/Horde/HordeDashboard/src/components/JobRedirector.tsx b/Engine/Source/Programs/Horde/HordeDashboard/src/components/JobRedirector.tsx
--- a/Engine/Source/Programs/Horde/HordeDashboard/src/components/JobRedirector.tsx
+++ b/Engine/Source/Programs/Horde/HordeDashboard/src/components/JobRedirector.tsx
@@ -3,9 +3,9 @@
 import React, { useState } from 'react';
 import { useNavigate } from 'react-router-dom';
 import backend from '../backend';
-import { useQuery } from './JobDetailCommon';
+import { useQuery } from "horde/base/utilities/hooks";
 import {errorDialogStore} from "./error/ErrorStore";
-import { JobStepState } from 'horde/backend/Api';
+
 
 function setError(message: string) {
 
@@ -30,7 +30,7 @@ export const JobRedirector: React.FC = () => {
    const query = useQuery();
 
    const streamId = !query.get("streamId") ? "" : query.get("streamId")!;
-   const templateId = query.getAll("templateId")
+   const templateId = query.getAll("templateId");
    const nodeName = !query.get("nodeName") ? "" : query.get("nodeName")!;
 
    const complete = query.get("complete") !== null;
diff --git a/Engine/Source/Programs/Horde/HordeDashboard/src/components/JobSearchSimple.tsx b/Engine/Source/Programs/Horde/HordeDashboard/src/components/JobSearchSimple.tsx
--- a/Engine/Source/Programs/Horde/HordeDashboard/src/components/JobSearchSimple.tsx
+++ b/Engine/Source/Programs/Horde/HordeDashboard/src/components/JobSearchSimple.tsx
@@ -2,7 +2,7 @@ import { Checkbox, ContextualMenuItemType, DatePicker, DefaultButton, DetailsLis
 import { observer } from "mobx-react-lite";
 import moment from "moment";
 import { useEffect, useMemo, useState } from "react";
-import { Link, useLocation, useSearchParams } from "react-router-dom";
+import { Link, useLocation } from "react-router-dom";
 import backend from "../backend";
 import { GetJobResponse, GetJobsTabResponse, GetJobStepRefResponse, JobData, JobQuery, TabType, GetTemplateRefResponse, GetJobsTabLabelColumnResponse, GetLabelStateResponse, LabelOutcome, LabelState, GetJobsTabColumnResponse, JobStepOutcome } from "../backend/Api";
 import dashboard from "../backend/Dashboard";
@@ -19,6 +19,7 @@ import { getLabelColor } from "horde/styles/colors";
 import { CalloutController, JobLabelCallout } from "./JobLabelCallout";
 import LZString from "lz-string"
 import { errorBarStore } from "./error/ErrorStore";
+import { useQuery } from "horde/base/utilities/hooks";
 
 type SearchState = {
    templates: Set<string>, 
@@ -52,7 +53,7 @@ export const JobSearchSimpleModal: React.FC<{ streamIdIn: string, changeListIn?:
 
    const [searchState, setSearchState] = useState<SearchState>({ preflights: false, minCL: changeListIn, maxCL: changeListIn, templates: new Set(jobTab?.templates ?? []), statuses: allStatuses });
    const [searchResults, setSearchResults] = useState<{items: JobItem[], groups?: IGroup[], querying: boolean}>({items: [], querying: false});
-   const [searchParams, setSearchParams] = useSearchParams();
+   const query = useQuery();
    const [templateData, setTemplateData] = useState<{ streamId: string; templates: GetTemplateRefResponse[] } | undefined>(undefined);
    const [controller] = useState(new CalloutController());
    const location = useLocation();
@@ -77,25 +78,18 @@ export const JobSearchSimpleModal: React.FC<{ streamIdIn: string, changeListIn?:
 
    // Update the search params in the url for sharing and bookmarking queires
    const updateSearch = (state: SearchState) => {
-      // Preserve current tab
-      const tab = searchParams.get("tab")!;
-      setSearchParams("");
+      const tab = query.get("tab");
+      query.clear();
+      if(tab) query.set("tab", tab);
 
-      const search = new URLSearchParams();
-      if(tab) search.append("tab", tab);
-      
-      // Compress current state to prevent 414s
-      // JSON.stringify doesn't handle Sets
+      // JSON.stringify doesn't support Sets
       const stringyState = {
          ...state,
          templates: Array.from(state.templates),
          statuses: Array.from(state.statuses)
       };
-      const compressedStateString = LZString.compressToEncodedURIComponent(JSON.stringify(stringyState));
-
-      search.append("search", compressedStateString);
-
-      setSearchParams(search);
+      // Compress current state to prevent 414
+      query.setCompressed("search", JSON.stringify(stringyState));
    }
 
    const queryJobs = async (state: SearchState) => {
@@ -282,32 +276,15 @@ export const JobSearchSimpleModal: React.FC<{ streamIdIn: string, changeListIn?:
       }
 
    }
-   
-   // Clear search parameters from the URL
-   const clearSearchParams = (removeFromURL: boolean = true) => {
-      if(removeFromURL) {
-         const tab = searchParams.get("tab")!;
-         setSearchParams("");
-         
-         const search = new URLSearchParams();
-         if(tab) search.append("tab", tab);
-         setSearchParams(search);
-      }
-      else {
-         const search = new URLSearchParams(searchParams);
-         search.set("search", "");
-         setSearchParams(search);
-      }
-   }
 
    // Sets state from search params and autoloads jobs
    useEffect(() => {
       // StreamView's param detection + the 'get' check below allow linking to an empty search modal "https://...?tab=...&search="
-      if(!searchParams.has("search") || !searchParams.get("search") || !isFirstRender || !templateData) return;
+      if(!query.has("search") || !query.get("search") || !isFirstRender || !templateData) return;
 
       setIsFirstRender(false);
       
-      const searchValue = searchParams.get("search")!;
+      const searchValue = query.get("search")!;
       try {
          // Decompress URL data
          const stringyState = JSON.parse(LZString.decompressFromEncodedURIComponent(searchValue));
@@ -339,7 +316,7 @@ export const JobSearchSimpleModal: React.FC<{ streamIdIn: string, changeListIn?:
             console.error("Bad search param value:\n\n" + searchValue + "\n\n");
          }
 
-         clearSearchParams();
+         query.clear(["tab"]);
       }
    }, [templateData, location.search])
 
@@ -844,7 +821,7 @@ export const JobSearchSimpleModal: React.FC<{ streamIdIn: string, changeListIn?:
    }
 
    // Can auto load normally (e.g., clicking "related jobs") if not already a search link
-   if(autoloadIn && isFirstRender && !searchParams.has("search")){
+   if(autoloadIn && isFirstRender && !query.has("search")){
       queryJobs(searchState);
       setIsFirstRender(false);
       return null;
@@ -856,7 +833,7 @@ export const JobSearchSimpleModal: React.FC<{ streamIdIn: string, changeListIn?:
       topOffsetFixed={true} 
       styles={{ main: { padding: 8, width: 1400, height: '820px', backgroundColor: hordeTheme.horde.contentBackground, hasBeenOpened: false, top: "80px", position: "absolute" } }} 
       className={hordeClasses.modal} 
-      onDismiss={() => { clearSearchParams(); onClose(); }}
+      onDismiss={() => { query.clear(["tab"]); onClose(); }}
    >
       {searchResults.querying && <StatusModal text={"Finding Jobs"} />}
       <JobLabelCallout controller={controller} />
@@ -869,7 +846,7 @@ export const JobSearchSimpleModal: React.FC<{ streamIdIn: string, changeListIn?:
                <Stack grow horizontalAlign="end">
                   <IconButton
                      iconProps={{ iconName: 'Cancel' }}
-                     onClick={() => { clearSearchParams(); onClose(); }}
+                     onClick={() => { query.clear(["tab"]); onClose(); }}
                   />
                </Stack>
             </Stack>
@@ -989,7 +966,7 @@ export const JobSearchSimpleModal: React.FC<{ streamIdIn: string, changeListIn?:
                            <DefaultButton disabled={searchResults.querying} text="Clear" onClick={() => { 
                                  setSearchState({ preflights: false, userId: undefined, templates: new Set(), statuses: new Set()}) 
                                  setSearchResults({items: [], groups: [], querying: false});
-                                 clearSearchParams(false);
+                                 query.set("search", "");
                               }} />
                         </Stack>
                      </Stack>
diff --git a/Engine/Source/Programs/Horde/HordeDashboard/src/components/JobView.tsx b/Engine/Source/Programs/Horde/HordeDashboard/src/components/JobView.tsx
--- a/Engine/Source/Programs/Horde/HordeDashboard/src/components/JobView.tsx
+++ b/Engine/Source/Programs/Horde/HordeDashboard/src/components/JobView.tsx
@@ -13,7 +13,7 @@ import { filterJob, JobFilterSimple } from '../base/utilities/filter';
 import { displayTimeZone } from '../base/utilities/timeUtils';
 import { getJobStateColor, getLabelColor } from '../styles/colors';
 import { ChangeButton } from './ChangeButton';
-import { useQuery } from './JobDetailCommon';
+import { useQuery } from "horde/base/utilities/hooks";
 import { CalloutController, JobLabelCallout } from './JobLabelCallout';
 import { JobOperationsContextMenu } from './JobOperationsContextMenu';
 import { getHordeStyling } from '../styles/Styles';
diff --git a/Engine/Source/Programs/Horde/HordeDashboard/src/components/LogView.tsx b/Engine/Source/Programs/Horde/HordeDashboard/src/components/LogView.tsx
--- a/Engine/Source/Programs/Horde/HordeDashboard/src/components/LogView.tsx
+++ b/Engine/Source/Programs/Horde/HordeDashboard/src/components/LogView.tsx
@@ -22,7 +22,7 @@ import { Breadcrumbs } from './Breadcrumbs';
 import { ChangeContextMenu, ChangeContextMenuTarget } from './ChangeButton';
 import { HistoryModal } from './agents/HistoryModal';
 import { IssueModalV2 } from './IssueViewV2';
-import { useQuery } from './JobDetailCommon';
+import { useQuery } from "horde/base/utilities/hooks";
 import { LogItem, renderLine } from './LogRender';
 import { JobLogSource, LogSource } from './LogSource';
 import { getLogStyles, logMetricNormal, logMetricSmall } from "./LogStyle";
@@ -761,7 +761,7 @@ const LogList: React.FC<{ logId: string }> = observer(({ logId }) => {
    }
 
    if (!handler) {
-      LogSource.create(logId, query).then(async (source) => {
+      LogSource.create(logId, query.raw.searchParams).then(async (source) => {
          await source.init();
 
          const handler = new LogHandler();
diff --git a/Engine/Source/Programs/Horde/HordeDashboard/src/components/NewBuild.tsx b/Engine/Source/Programs/Horde/HordeDashboard/src/components/NewBuild.tsx
--- a/Engine/Source/Programs/Horde/HordeDashboard/src/components/NewBuild.tsx
+++ b/Engine/Source/Programs/Horde/HordeDashboard/src/components/NewBuild.tsx
@@ -8,7 +8,7 @@ import backend, { useBackend } from '../backend';
 import { BoolParameterData, CreateJobRequest, GetJobsTabResponse, GroupParameterData, JobsTabData, ListParameterData, ListParameterItemData, ListParameterStyle, ParameterData, ParameterType, Priority, TabType, GetTemplateRefResponse, TextParameterData, ChangeQueryConfig } from '../backend/Api';
 import templateCache from '../backend/TemplateCache';
 import { errorDialogStore } from "./error/ErrorStore";
-import { useQuery } from './JobDetailCommon';
+import { useQuery } from "horde/base/utilities/hooks";
 import { JobDetailsV2 } from './jobDetailsV2/JobDetailsViewCommon';
 import dashboard from '../backend/Dashboard';
 import { getHordeStyling } from '../styles/Styles';
diff --git a/Engine/Source/Programs/Horde/HordeDashboard/src/components/NotificationDropdown.tsx b/Engine/Source/Programs/Horde/HordeDashboard/src/components/NotificationDropdown.tsx
--- a/Engine/Source/Programs/Horde/HordeDashboard/src/components/NotificationDropdown.tsx
+++ b/Engine/Source/Programs/Horde/HordeDashboard/src/components/NotificationDropdown.tsx
@@ -7,7 +7,7 @@ import backend from "../backend";
 import { CreateSubscriptionRequest, GetNotificationResponse, GetSubscriptionResponse, JobState, JobStepState, LabelCompleteEventRecord, LabelState, StepCompleteEventRecord, SubscriptonNotificationType, UpdateNotificationsRequest } from "../backend/Api";
 import { action, makeObservable, observable } from "mobx";
 import { JobDetailsV2 } from './jobDetailsV2/JobDetailsViewCommon';
-import { useQuery } from "./JobDetailCommon";
+import { useQuery } from "horde/base/utilities/hooks";
 import { getHordeStyling } from "../styles/Styles";
 
 
diff --git a/Engine/Source/Programs/Horde/HordeDashboard/src/components/Preflight.tsx b/Engine/Source/Programs/Horde/HordeDashboard/src/components/Preflight.tsx
--- a/Engine/Source/Programs/Horde/HordeDashboard/src/components/Preflight.tsx
+++ b/Engine/Source/Programs/Horde/HordeDashboard/src/components/Preflight.tsx
@@ -3,7 +3,7 @@
 import React, { useState } from 'react';
 import { useNavigate } from 'react-router-dom';
 import backend, { useBackend } from '../backend';
-import { useQuery } from './JobDetailCommon';
+import { useQuery } from "horde/base/utilities/hooks";
 import {errorDialogStore} from './error/ErrorStore';
 import { GetJobsTabResponse, JobsTabData, TabType } from '../backend/Api';
 
@@ -41,7 +41,7 @@ export const PreflightRedirector: React.FC = () => {
 
    const parameters: string[] = [];
 
-   query.forEach((value, key) => {
+   query.raw.searchParams.forEach((value, key) => {
       if (key.startsWith("id")) {
          parameters.push(`${key}=${value}`);
       }
diff --git a/Engine/Source/Programs/Horde/HordeDashboard/src/components/StreamSummary.tsx b/Engine/Source/Programs/Horde/HordeDashboard/src/components/StreamSummary.tsx
--- a/Engine/Source/Programs/Horde/HordeDashboard/src/components/StreamSummary.tsx
+++ b/Engine/Source/Programs/Horde/HordeDashboard/src/components/StreamSummary.tsx
@@ -12,7 +12,7 @@ import dashboard from '../backend/Dashboard';
 import { ProjectStore } from '../backend/ProjectStore';
 import { getElapsedString, getShortNiceTime } from '../base/utilities/timeUtils';
 import { IssueModalV2 } from './IssueViewV2';
-import { useQuery } from './JobDetailCommon';
+import { useQuery } from "horde/base/utilities/hooks";
 import { SchedulePane } from './SchedulePane';
 import { IssueStatusIconV2 } from './StatusIcon';
 import { BuildHealthTestReportPanel } from './TestReportPanel';
diff --git a/Engine/Source/Programs/Horde/HordeDashboard/src/components/StreamView.tsx b/Engine/Source/Programs/Horde/HordeDashboard/src/components/StreamView.tsx
--- a/Engine/Source/Programs/Horde/HordeDashboard/src/components/StreamView.tsx
+++ b/Engine/Source/Programs/Horde/HordeDashboard/src/components/StreamView.tsx
@@ -14,7 +14,7 @@ import { useWindowSize } from '../base/utilities/hooks';
 import { getHordeStyling } from '../styles/Styles';
 import { getHordeTheme } from '../styles/theme';
 import { BreadcrumbItem, Breadcrumbs } from './Breadcrumbs';
-import { useQuery } from './JobDetailCommon';
+import { useQuery } from "horde/base/utilities/hooks";
 import { JobSearchSimpleModal } from './JobSearchSimple';
 import { JobView } from './JobView';
 import { JobViewAll } from './JobViewAll';
diff --git a/Engine/Source/Programs/Horde/HordeDashboard/src/components/TestReportView.tsx b/Engine/Source/Programs/Horde/HordeDashboard/src/components/TestReportView.tsx
--- a/Engine/Source/Programs/Horde/HordeDashboard/src/components/TestReportView.tsx
+++ b/Engine/Source/Programs/Horde/HordeDashboard/src/components/TestReportView.tsx
@@ -9,7 +9,7 @@ import hordePlugins from  '../legacyPlugins';
 import { projectStore } from '../backend/ProjectStore';
 import { useParams } from 'react-router-dom';
 import { TestDataCollection, TestDataWrapper } from '../backend/TestDataHandler'
-import { useQuery } from './JobDetailCommon';
+import { useQuery } from "horde/base/utilities/hooks";
 import { getHordeStyling } from '../styles/Styles';
 
 const getFirstComponentFromPlugin = (mount: PluginMount, id: string): React.FC<any> | undefined => {
diff --git a/Engine/Source/Programs/Horde/HordeDashboard/src/components/UserHome.tsx b/Engine/Source/Programs/Horde/HordeDashboard/src/components/UserHome.tsx
--- a/Engine/Source/Programs/Horde/HordeDashboard/src/components/UserHome.tsx
+++ b/Engine/Source/Programs/Horde/HordeDashboard/src/components/UserHome.tsx
@@ -19,7 +19,7 @@ import { BisectionList } from './bisection/BisectionList';
 import { Breadcrumbs } from './Breadcrumbs';
 import { ChangeButton } from "./ChangeButton";
 import { IssueModalV2 } from './IssueViewV2';
-import { useQuery } from './JobDetailCommon';
+import { useQuery } from "horde/base/utilities/hooks";
 import { JobOperationsContextMenu } from "./JobOperationsContextMenu";
 import { PreflightConfigModal } from './preflights/PreflightConfigCheck';
 import { IssueStatusIcon, StepStatusIcon } from "./StatusIcon";
diff --git a/Engine/Source/Programs/Horde/HordeDashboard/src/components/UtilizationReportView.tsx b/Engine/Source/Programs/Horde/HordeDashboard/src/components/UtilizationReportView.tsx
--- a/Engine/Source/Programs/Horde/HordeDashboard/src/components/UtilizationReportView.tsx
+++ b/Engine/Source/Programs/Horde/HordeDashboard/src/components/UtilizationReportView.tsx
@@ -10,7 +10,7 @@ import backend from '../backend';
 import { GetUtilizationTelemetryResponse } from '../backend/Api';
 import { useWindowSize } from '../base/utilities/hooks';
 import { Breadcrumbs } from './Breadcrumbs';
-import { useQuery } from './JobDetailCommon';
+import { useQuery } from "horde/base/utilities/hooks";
 import { TopNav } from './TopNav';
 import { getHordeStyling } from '../styles/Styles';
 
diff --git a/Engine/Source/Programs/Horde/HordeDashboard/src/components/artifacts/ArtifactsSearch.tsx b/Engine/Source/Programs/Horde/HordeDashboard/src/components/artifacts/ArtifactsSearch.tsx
--- a/Engine/Source/Programs/Horde/HordeDashboard/src/components/artifacts/ArtifactsSearch.tsx
+++ b/Engine/Source/Programs/Horde/HordeDashboard/src/components/artifacts/ArtifactsSearch.tsx
@@ -10,7 +10,7 @@ import { getActiveStreamId } from "../../base/utilities/streamUtils";
 import { getHordeStyling } from "../../styles/Styles";
 import { StreamChooser, streamIdAll } from "../projects/StreamChooser";
 import { JobArtifactsModal } from "./ArtifactsModal";
-import { useQuery } from "../JobDetailCommon";
+import { useQuery } from "horde/base/utilities/hooks";
 import { ArtifactButton } from "./ArtifactButton";
 
 type ArtifactSearchState = {
diff --git a/Engine/Source/Programs/Horde/HordeDashboard/src/components/jobDetailsV2/JobDetailHealthV2.tsx b/Engine/Source/Programs/Horde/HordeDashboard/src/components/jobDetailsV2/JobDetailHealthV2.tsx
--- a/Engine/Source/Programs/Horde/HordeDashboard/src/components/jobDetailsV2/JobDetailHealthV2.tsx
+++ b/Engine/Source/Programs/Horde/HordeDashboard/src/components/jobDetailsV2/JobDetailHealthV2.tsx
@@ -10,7 +10,7 @@ import { GetExternalIssueResponse, GetIssueResponse, IssueData } from '../../bac
 import { ISideRailLink } from '../../base/components/SideRail';
 import { getElapsedString, getShortNiceTime } from '../../base/utilities/timeUtils';
 import { IssueModalV2 } from '../IssueViewV2';
-import { useQuery } from '../JobDetailCommon';
+import { useQuery } from "horde/base/utilities/hooks";
 import { IssueStatusIcon } from '../StatusIcon';
 import { JobDataView, JobDetailsV2 } from "./JobDetailsViewCommon";
 import { getHordeStyling } from '../../styles/Styles';
diff --git a/Engine/Source/Programs/Horde/HordeDashboard/src/components/jobDetailsV2/JobDetailLabels.tsx b/Engine/Source/Programs/Horde/HordeDashboard/src/components/jobDetailsV2/JobDetailLabels.tsx
--- a/Engine/Source/Programs/Horde/HordeDashboard/src/components/jobDetailsV2/JobDetailLabels.tsx
+++ b/Engine/Source/Programs/Horde/HordeDashboard/src/components/jobDetailsV2/JobDetailLabels.tsx
@@ -1,13 +1,13 @@
 import { ConstrainMode, DefaultButton, DetailsList, DetailsListLayoutMode, DirectionalHint, IColumn, SelectionMode, Stack, Text, TooltipHost } from "@fluentui/react";
 import { observer } from "mobx-react-lite";
 import { useNavigate, useLocation } from "react-router-dom";
-import { GetStepResponse, JobData, JobStepOutcome, LabelOutcome, LabelState, StepData } from "../../backend/Api";
+import { GetStepResponse, JobStepOutcome, LabelOutcome, LabelState, StepData } from "../../backend/Api";
 import { JobLabel } from "../../backend/JobDetails";
 import { getLabelColor } from "../../styles/colors";
-import { useQuery } from "../JobDetailCommon";
+import { useQuery } from "horde/base/utilities/hooks";
 import { JobDataView, JobDetailsV2 } from "./JobDetailsViewCommon";
 import { getHordeStyling } from "../../styles/Styles";
-import dashboard, { StatusColor } from "../../backend/Dashboard";
+import dashboard from "../../backend/Dashboard";
 import { Link } from "react-router-dom";
 import { StepStatusIcon } from "../StatusIcon";
 
diff --git a/Engine/Source/Programs/Horde/HordeDashboard/src/components/jobDetailsV2/JobDetailViewV2.tsx b/Engine/Source/Programs/Horde/HordeDashboard/src/components/jobDetailsV2/JobDetailViewV2.tsx
--- a/Engine/Source/Programs/Horde/HordeDashboard/src/components/jobDetailsV2/JobDetailViewV2.tsx
+++ b/Engine/Source/Programs/Horde/HordeDashboard/src/components/jobDetailsV2/JobDetailViewV2.tsx
@@ -6,7 +6,7 @@ import { GetJobsTabResponse } from "../../backend/Api";
 import { useWindowSize } from "../../base/utilities/hooks";
 import { getHordeStyling } from "../../styles/Styles";
 import { BreadcrumbItem, Breadcrumbs } from "../Breadcrumbs";
-import { useQuery } from "../JobDetailCommon";
+import { useQuery } from "horde/base/utilities/hooks";
 import { TopNav } from "../TopNav";
 import { JobArtifactsPanel } from "./JobDetailArtifacts";
 import { BisectionPanel } from "./JobDetailBisection";
diff --git a/Engine/Source/Programs/Horde/HordeDashboard/src/components/jobDetailsV2/JobDetailsViewCommon.tsx b/Engine/Source/Programs/Horde/HordeDashboard/src/components/jobDetailsV2/JobDetailsViewCommon.tsx
--- a/Engine/Source/Programs/Horde/HordeDashboard/src/components/jobDetailsV2/JobDetailsViewCommon.tsx
+++ b/Engine/Source/Programs/Horde/HordeDashboard/src/components/jobDetailsV2/JobDetailsViewCommon.tsx
@@ -11,7 +11,7 @@ import { PollBase } from "../../backend/PollBase";
 import { projectStore } from "../../backend/ProjectStore";
 import TemplateCache from "../../backend/TemplateCache";
 import { ISideRailLink } from "../../base/components/SideRail";
-import { useQuery } from "../JobDetailCommon";
+import { useQuery } from "horde/base/utilities/hooks";
 import { LabelStatusIcon, StepStatusIcon } from "../StatusIcon";
 import { getHordeStyling } from "../../styles/Styles";
 
diff --git a/Engine/Source/Programs/Horde/HordeDashboard/src/components/jobDetailsV2/JobOperationsBar.tsx b/Engine/Source/Programs/Horde/HordeDashboard/src/components/jobDetailsV2/JobOperationsBar.tsx
--- a/Engine/Source/Programs/Horde/HordeDashboard/src/components/jobDetailsV2/JobOperationsBar.tsx
+++ b/Engine/Source/Programs/Horde/HordeDashboard/src/components/jobDetailsV2/JobOperationsBar.tsx
@@ -3,13 +3,13 @@
 import { CommandBar, CommandBarButton, ICommandBarItemProps, IContextualMenuItem, Stack } from '@fluentui/react';
 import { observer } from 'mobx-react-lite';
 import React, { useState } from 'react';
-import { Link, useNavigate, useSearchParams } from 'react-router-dom';
+import { Link, useNavigate } from 'react-router-dom';
 import { ArtifactContextType, JobState, JobStepOutcome, JobStepState, StepData } from '../../backend/Api';
 import { getSiteConfig } from '../../backend/Config';
 import dashboard from '../../backend/Dashboard';
 import { getHordeStyling } from '../../styles/Styles';
 import { EditJobModal } from '../EditJobModal';
-import { useQuery } from '../JobDetailCommon';
+import { useQuery } from "horde/base/utilities/hooks";
 import { NewBuild } from '../NewBuild';
 import { NotificationDropdown } from '../NotificationDropdown';
 import { PauseStepModal } from '../StepPauseModal';
@@ -35,8 +35,7 @@ export const JobOperations: React.FC<{ jobDetails: JobDetailsV2 }> = observer(({
    const [editShown, setEditShown] = useState(false);
    const [retryStepsShown, setRetryStepsShown] = useState(false);
 
-   const [searchParams] = useSearchParams();
-   const [findJobsShown, setFindJobsShown] = useState(searchParams.has("search"));
+   const [findJobsShown, setFindJobsShown] = useState(query.has("search"));
 
    const [parametersState, setParametersState] = useState(query.get("newbuild") ? ParameterState.Clone : ParameterState.Hidden);
 
diff --git a/Engine/Source/Programs/Horde/HordeDashboard/src/components/jobDetailsV2/StepDetailView.tsx b/Engine/Source/Programs/Horde/HordeDashboard/src/components/jobDetailsV2/StepDetailView.tsx
--- a/Engine/Source/Programs/Horde/HordeDashboard/src/components/jobDetailsV2/StepDetailView.tsx
+++ b/Engine/Source/Programs/Horde/HordeDashboard/src/components/jobDetailsV2/StepDetailView.tsx
@@ -8,7 +8,7 @@ import backend from "../../backend";
 import { GetArtifactResponse } from "../../backend/Api";
 import { getHordeStyling } from "../../styles/Styles";
 import { HistoryModal } from "../agents/HistoryModal";
-import { useQuery } from "../JobDetailCommon";
+import { useQuery } from "horde/base/utilities/hooks";
 import { BisectionPanel } from "./JobDetailBisection";
 import { HealthPanel } from "./JobDetailHealthV2";
 import { PreflightPanel } from "./JobDetailPreflight";
diff --git a/Engine/Source/Programs/Horde/HordeDashboard/src/components/preflights/PreflightConfigCheckRedirector.tsx b/Engine/Source/Programs/Horde/HordeDashboard/src/components/preflights/PreflightConfigCheckRedirector.tsx
--- a/Engine/Source/Programs/Horde/HordeDashboard/src/components/preflights/PreflightConfigCheckRedirector.tsx
+++ b/Engine/Source/Programs/Horde/HordeDashboard/src/components/preflights/PreflightConfigCheckRedirector.tsx
@@ -2,7 +2,7 @@
 
 import React from 'react';
 import { Navigate } from 'react-router-dom';
-import { useQuery } from '../JobDetailCommon';
+import { useQuery } from "horde/base/utilities/hooks";
 import { errorDialogStore } from '../error/ErrorStore';
 
 function setError(message: string) {
diff --git a/Engine/Source/Programs/Horde/HordeDashboard/src/legacyPlugins/AutomatedTestSession/components/AutomatedTestSession.tsx b/Engine/Source/Programs/Horde/HordeDashboard/src/legacyPlugins/AutomatedTestSession/components/AutomatedTestSession.tsx
--- a/Engine/Source/Programs/Horde/HordeDashboard/src/legacyPlugins/AutomatedTestSession/components/AutomatedTestSession.tsx
+++ b/Engine/Source/Programs/Horde/HordeDashboard/src/legacyPlugins/AutomatedTestSession/components/AutomatedTestSession.tsx
@@ -9,7 +9,7 @@ import { Link } from 'react-router-dom';
 import dashboard, { StatusColor } from "../../../backend/Dashboard";
 import { msecToElapsed } from '../../../base/utilities/timeUtils';
 import { BreadcrumbItem, Breadcrumbs } from '../../../components/Breadcrumbs';
-import { useQuery } from '../../../components/JobDetailCommon';
+import { useQuery } from "horde/base/utilities/hooks";
 import { getProjectName, getStreamData, testDataHandler } from '../../../components/TestReportView';
 import { TopNav } from '../../../components/TopNav';
 import { getHordeStyling } from '../../../styles/Styles';
