[Horde][Analytics] Add ExternalIssueKey & Fingerprint to the telemetry payload for issues. Also includes AnalyticsIntegrationTests.

diff --git a/Engine/Source/Programs/Horde/Plugins/Analytics/HordeServer.Analytics.Integration.Tests/AnalyticsIntegrationTestExtensions.cs b/Engine/Source/Programs/Horde/Plugins/Analytics/HordeServer.Analytics.Integration.Tests/AnalyticsIntegrationTestExtensions.cs
new file mode 100644
--- /dev/null
+++ b/Engine/Source/Programs/Horde/Plugins/Analytics/HordeServer.Analytics.Integration.Tests/AnalyticsIntegrationTestExtensions.cs
@@ -0,0 +1,188 @@
+// Copyright Epic Games, Inc. All Rights Reserved.
+
+using System.Collections;
+using System.Reflection;
+using System.Runtime.Serialization;
+using Microsoft.Extensions.Logging;
+
+namespace HordeServer.Analytics.Integration.Tests
+{
+	public static class AnalyticsIntegrationTestExtensions
+	{
+		private static readonly ILogger s_logger = LoggerFactory.Create(b => b.AddConsole()).CreateLogger("AnalyticsIntegrationTestExtensions");
+
+		private static object? TryMapAnonymousToTypedObject(object anon, Type targetType)
+		{
+			if (anon == null)
+			{
+				return null;
+			}
+
+#pragma warning disable SYSLIB0050 // Type or member is obsolete
+			object typedObj = FormatterServices.GetUninitializedObject(targetType);
+#pragma warning restore SYSLIB0050 // Type or member is obsolete
+
+			foreach (PropertyInfo targetProp in targetType.GetProperties(BindingFlags.Public | BindingFlags.Instance))
+			{
+				if (!targetProp.CanWrite)
+				{
+					continue;
+				}
+
+				PropertyInfo? anonProp = anon.GetType().GetProperty(targetProp.Name);
+				if (anonProp == null)
+				{
+					continue;
+				}
+
+				object? anonValue = anonProp.GetValue(anon);
+				if (anonValue == null)
+				{
+					continue;
+				}
+
+				try
+				{
+					object? converted = Convert.ChangeType(anonValue, Nullable.GetUnderlyingType(targetProp.PropertyType) ?? targetProp.PropertyType);
+					targetProp.SetValue(typedObj, converted);
+				}
+				catch
+				{
+					if (targetProp.PropertyType.IsInstanceOfType(anonValue))
+					{
+						targetProp.SetValue(typedObj, anonValue);
+					}
+				}
+			}
+
+			return typedObj;
+		}
+
+		/// <summary>
+		/// Attempts to map a hashtable to an object.
+		/// </summary>
+		/// <typeparam name="T">The type of object to return.</typeparam>
+		/// <param name="table">The input table to attempt conversion from.</param>
+		/// <returns>The resulting object of type if conversion was possible.</returns>
+		/// <exception cref="ArgumentNullException">Throws if input table is null.</exception>
+		public static T MapToObject<T>(this Hashtable table) where T : class
+		{
+			ArgumentNullException.ThrowIfNull(table);
+
+			Type targetType = typeof(T);
+#pragma warning disable SYSLIB0050 // Type or member is obsolete
+			T instance = (T)System.Runtime.Serialization.FormatterServices.GetUninitializedObject(typeof(T));
+#pragma warning restore SYSLIB0050 // Type or member is obsolete
+
+			foreach (PropertyInfo prop in targetType.GetProperties(BindingFlags.Public | BindingFlags.Instance))
+			{
+				if (!prop.CanWrite)
+				{
+					continue;
+				}
+
+				string propName = prop.Name;
+
+				if (!table.ContainsKey(propName))
+				{
+					continue;
+				}
+
+				object? value = table[propName];
+
+				if (value == null)
+				{
+					prop.SetValue(instance, null);
+					continue;
+				}
+
+				Type propType = prop.PropertyType;
+
+				try
+				{
+					Type underlyingType = Nullable.GetUnderlyingType(propType) ?? propType;
+
+					if (underlyingType.IsEnum)
+					{
+						string? valueString = value.ToString();
+
+						if (String.IsNullOrEmpty(valueString))
+						{
+							continue;
+						}
+
+						object enumValue = Enum.Parse(underlyingType, value: valueString);
+						prop.SetValue(instance, enumValue);
+					}
+					else if (underlyingType == typeof(Guid))
+					{
+						string? valueString = value.ToString();
+
+						if (String.IsNullOrEmpty(valueString))
+						{
+							continue;
+						}
+
+						prop.SetValue(instance, Guid.Parse(input: value.ToString()));
+					}
+					else if (underlyingType == typeof(DateTime))
+					{
+						prop.SetValue(instance, Convert.ChangeType(value, typeof(DateTime)));
+					}
+					else if (underlyingType == typeof(bool))
+					{
+						prop.SetValue(instance, Convert.ToBoolean(value));
+					}
+					else if (underlyingType.IsArray && value is IEnumerable enumerableValue)
+					{
+						Type? elemType = underlyingType.GetElementType();
+
+						if (elemType == null)
+						{
+							continue;
+						}
+
+						object[] list = enumerableValue.Cast<object>()
+							.Select(v => Convert.ChangeType(v, elemType))
+							.ToArray();
+
+						Array array = Array.CreateInstance(elemType, list.Length);
+						Array.Copy(list, array, list.Length);
+
+						prop.SetValue(instance, array);
+					}
+					else if (underlyingType.IsGenericType
+							 && typeof(IEnumerable).IsAssignableFrom(underlyingType)
+							 && value is IEnumerable enumerable)
+					{
+						prop.SetValue(instance, value);
+					}
+					else
+					{
+						try
+						{
+							object converted = Convert.ChangeType(value, underlyingType);
+							prop.SetValue(instance, converted);
+						}
+						catch
+						{
+							// Fallback: try to map anonymous type to expected property type by property name
+							object? fallbackValue = TryMapAnonymousToTypedObject(value, underlyingType);
+							if (fallbackValue != null)
+							{
+								prop.SetValue(instance, fallbackValue);
+							}
+						}
+					}
+				}
+				catch (Exception ex)
+				{
+					s_logger?.LogDebug("Failed to process prop:{PropName} with value: {Value} due to exception: {Ex}", propName, value, ex);
+				}
+			}
+
+			return instance;
+		}
+	}
+
+}
diff --git a/Engine/Source/Programs/Horde/Plugins/Analytics/HordeServer.Analytics.Integration.Tests/HordeServer.Analytics.Integration.Tests.csproj b/Engine/Source/Programs/Horde/Plugins/Analytics/HordeServer.Analytics.Integration.Tests/HordeServer.Analytics.Integration.Tests.csproj
new file mode 100644
--- /dev/null
+++ b/Engine/Source/Programs/Horde/Plugins/Analytics/HordeServer.Analytics.Integration.Tests/HordeServer.Analytics.Integration.Tests.csproj
@@ -0,0 +1,43 @@
+<Project Sdk="Microsoft.NET.Sdk">
+
+  <PropertyGroup>
+    <TargetFramework>net8.0</TargetFramework>
+    <ImplicitUsings>enable</ImplicitUsings>
+    <Nullable>enable</Nullable>
+    <IsPackable>false</IsPackable>
+    <IsTestProject>true</IsTestProject>
+    <RunAnalyzersDuringBuild>False</RunAnalyzersDuringBuild>
+    <Configurations>Debug;Release;Analyze</Configurations>
+    <NoWarn>1701;1702;1591</NoWarn>
+  </PropertyGroup>
+
+  <PropertyGroup Condition=" '$(Configuration)|$(Platform)' == 'Analyze|AnyCPU' ">
+    <RunAnalyzersDuringBuild>True</RunAnalyzersDuringBuild>
+    <TreatWarningsAsErrors>false</TreatWarningsAsErrors>
+    <EnforceCodeStyleInBuild>True</EnforceCodeStyleInBuild>
+    <GenerateDocumentationFile>True</GenerateDocumentationFile>
+  </PropertyGroup>
+
+  <ItemGroup>
+    <PackageReference Include="coverlet.collector" Version="6.0.2">
+      <PrivateAssets>all</PrivateAssets>
+      <IncludeAssets>runtime; build; native; contentfiles; analyzers; buildtransitive</IncludeAssets>
+    </PackageReference>
+    <PackageReference Include="Microsoft.NET.Test.Sdk" Version="17.11.1" />
+    <PackageReference Include="MSTest.TestAdapter" Version="3.6.2" />
+    <PackageReference Include="MSTest.TestFramework" Version="3.6.2" />
+  </ItemGroup>
+
+  <ItemGroup>
+    <ProjectReference Include="..\..\Build\HordeServer.Build.Tests\HordeServer.Build.Tests.csproj" />
+    <ProjectReference Include="..\..\Build\HordeServer.Build\HordeServer.Build.csproj" />
+    <ProjectReference Include="..\HordeServer.Analytics.Tests\HordeServer.Analytics.Tests.csproj" />
+    <ProjectReference Include="..\HordeServer.Analytics\HordeServer.Analytics.csproj" />
+    <ProjectReference Include="..\..\..\HordeServer.Tests.Shared\HordeServer.Tests.Shared.csproj" />
+  </ItemGroup>
+
+  <ItemGroup>
+    <Using Include="Microsoft.VisualStudio.TestTools.UnitTesting" />
+  </ItemGroup>
+
+</Project>
diff --git a/Engine/Source/Programs/Horde/Plugins/Analytics/HordeServer.Analytics.Integration.Tests/Telemetry/TelemetryIntegrationTests.cs b/Engine/Source/Programs/Horde/Plugins/Analytics/HordeServer.Analytics.Integration.Tests/Telemetry/TelemetryIntegrationTests.cs
new file mode 100644
--- /dev/null
+++ b/Engine/Source/Programs/Horde/Plugins/Analytics/HordeServer.Analytics.Integration.Tests/Telemetry/TelemetryIntegrationTests.cs
@@ -0,0 +1,210 @@
+// Copyright Epic Games, Inc. All Rights Reserved.
+
+using System.Collections;
+using System.Reflection;
+using EpicGames.Core;
+using EpicGames.Horde.Issues;
+using EpicGames.Horde.Jobs;
+using EpicGames.Horde.Telemetry;
+using HordeServer.Analytics.Tests;
+using HordeServer.Issues;
+using HordeServer.Telemetry;
+using HordeServer.Telemetry.Metrics;
+using HordeServer.Tests.Issues;
+using Microsoft.Extensions.DependencyInjection;
+
+namespace HordeServer.Analytics.Integration.Tests.Telemetry
+{
+	/// <summary>
+	/// Test telemetry writer. Used in dependency injection, and will track all received telemetry events for subsequent assertions.
+	/// </summary>
+	public class TestTelemetryWriter : ITelemetryWriter
+	{
+		#region -- Static API --
+
+		private static Hashtable ToHashtable(object anon)
+		{
+			Hashtable resultingTable = new Hashtable();
+			foreach (PropertyInfo prop in anon.GetType().GetProperties(BindingFlags.Instance | BindingFlags.Public))
+			{
+				resultingTable[prop.Name] = prop.GetValue(anon, null);
+			}
+			return resultingTable;
+		}
+
+		#endregion -- Static API --
+
+		#region -- Members & Properties --
+
+		public bool Enabled => true;
+		public IReadOnlyDictionary<TelemetryStoreId, IList<Tuple<TelemetryRecordMeta?, Hashtable>>> RecordedTelemetry => _recordedTelemetry;
+		private readonly Dictionary<TelemetryStoreId, IList<Tuple<TelemetryRecordMeta?, Hashtable>>> _recordedTelemetry;
+
+		#endregion -- Members & Properties --
+
+		#region -- Constructor --
+
+		public TestTelemetryWriter()
+		{
+			_recordedTelemetry = new Dictionary<TelemetryStoreId, IList<Tuple<TelemetryRecordMeta?, Hashtable>>>();
+		}
+
+		#endregion -- Constructor --
+
+		#region -- Public API --
+
+		internal void ClearStoredTelemetry()
+		{
+			_recordedTelemetry.Clear();
+		}
+
+		#endregion -- Public API --
+
+		#region -- Interface API --
+
+		public void WriteEvent(TelemetryStoreId telemetryStoreId, object payload)
+		{
+#pragma warning disable CS8625 // Cannot convert null literal to non-nullable reference type.
+			WriteEvent(telemetryStoreId, null, payload);
+#pragma warning restore CS8625 // Cannot convert null literal to non-nullable reference type.
+		}
+
+		public void WriteEvent(TelemetryStoreId telemetryStoreId, TelemetryRecordMeta metadata, object payload)
+		{
+			IList<Tuple<TelemetryRecordMeta?, Hashtable>>? telemetryStorePayloads;
+			Hashtable convertedPayload = ToHashtable(payload);
+
+			if (_recordedTelemetry.TryGetValue(telemetryStoreId, out telemetryStorePayloads))
+			{
+				telemetryStorePayloads.Add(new Tuple<TelemetryRecordMeta?, Hashtable>(null, convertedPayload));
+			}
+			else
+			{
+				Hashtable convertedHashTable = ToHashtable(payload);
+				telemetryStorePayloads = new List<Tuple<TelemetryRecordMeta?, Hashtable>>();
+				telemetryStorePayloads.Add(new Tuple<TelemetryRecordMeta?, Hashtable>(null, convertedPayload));
+				_recordedTelemetry.Add(telemetryStoreId, telemetryStorePayloads);
+			}
+		}
+
+		#endregion -- Interface API --
+	}
+
+	/// <summary>
+	/// Test class used to test the integration of <see cref="AnalyticsPlugin"/> with various other upstream systems.
+	/// </summary>
+	/// <remarks><see cref="BuildPlugin"/> represents the current upstream integration.</remarks>
+	[TestClass]
+	public class TelemetryIntegrationTests : AbstractIssueServiceTests
+	{
+		#region -- Test Constants --
+
+		const string EventNameKeyStr = "EventName";
+
+		#endregion -- Test Constants -- 
+
+		#region -- Interface API --
+
+		protected override void ConfigureServices(IServiceCollection services)
+		{
+			base.ConfigureServices(services);
+
+			services.AddSingleton<TestTelemetryWriter>();
+			services.AddSingleton<ITelemetryWriter>(sp => sp.GetRequiredService<TestTelemetryWriter>());
+		}
+
+		#endregion -- Interface API --
+
+		#region -- Test Cases --
+
+		/// <summary>
+		/// Test to verify that <see cref="IssueCollection"/> is properly emitting through the <see cref="ITelemetryWriter"/>.
+		/// </summary>
+		[TestMethod]
+		public async Task TestStateIssueAndIssueSpanEmissionAsync()
+		{
+			// Configure the metric
+			TelemetryStoreConfig telemetryStoreConfig = new TelemetryStoreConfig();
+			telemetryStoreConfig.Id = TelemetryStoreId.Default;
+			AnalyticsConfig analyticsConfig = new AnalyticsConfig();
+			analyticsConfig.Stores.Add(telemetryStoreConfig);
+
+			await UpdateConfigAsync(config =>
+			{
+				config.Plugins.GetBuildConfig().GetStream(_mainStreamId).TelemetryStoreId = TelemetryStoreId.Default;
+				config.Plugins.AddAnalyticsTestConfig(analyticsConfig);
+			});
+
+			// Test variables
+			TestTelemetryWriter concreteTelemetryWriter = ServiceProvider.GetRequiredService<TestTelemetryWriter>();
+
+			// Validate our preconditions
+			Assert.IsNotNull(concreteTelemetryWriter);
+			Assert.IsTrue(!concreteTelemetryWriter.RecordedTelemetry.ContainsKey(TelemetryStoreId.Default));
+
+			// Generate issue telemetry data
+			{
+				string[] lines =
+				{
+					@"WARNING: Engine\Source\Programs\UnrealBuildTool\ProjectFiles\Rider\ToolchainInfo.cs: Missing copyright boilerplate"
+				};
+
+				IJob job = CreateJob(_mainStreamId, 105, "Test Build", _graph);
+				await ParseEventsAsync(job, 0, 0, lines);
+				await UpdateCompleteStepAsync(job, 0, 0, JobStepOutcome.Failure);
+
+				// Validate we have received a telemetry event.
+				Assert.IsTrue(concreteTelemetryWriter.RecordedTelemetry.ContainsKey(TelemetryStoreId.Default));
+
+				// Verify presence of Issue telemetry
+				IList<Tuple<TelemetryRecordMeta?, Hashtable>>  actualTelemetryPayloads = [.. concreteTelemetryWriter.RecordedTelemetry[TelemetryStoreId.Default]];
+				Assert.IsTrue(actualTelemetryPayloads.Select(x => x.Item2[EventNameKeyStr] as string == IssueTelemetry.DefaultEventName).Any());
+				Assert.IsTrue(actualTelemetryPayloads.Select(x => x.Item2[EventNameKeyStr] as string == IssueSpanTelemetry.DefaultEventName).Any());
+			}
+
+			// Update the issue and make sure we get back external issues
+			const string ExpectedExternalIssueKey = "NewIssueKey-1010";
+			{
+				// Reset the writer
+				concreteTelemetryWriter.ClearStoredTelemetry();
+
+				IReadOnlyList<IIssue> issues = await IssueCollection.FindIssuesAsync();
+				await IssueService.UpdateIssueAsync(issues[0].Id, externalIssueKey: ExpectedExternalIssueKey);
+
+				IList<Tuple<TelemetryRecordMeta?, Hashtable>>  actualTelemetryPayloads = [.. concreteTelemetryWriter.RecordedTelemetry[TelemetryStoreId.Default]];
+				Hashtable actualStateIssuePayload = actualTelemetryPayloads.Where(x => x.Item2[EventNameKeyStr] as string == IssueTelemetry.DefaultEventName).First().Item2;
+				IssueTelemetry? actualIssueTelemetry = actualStateIssuePayload.MapToObject<IssueTelemetry>();
+
+				// Verify new payloads have been set for external issue
+				Assert.AreEqual(ExpectedExternalIssueKey, actualIssueTelemetry.ExternalIssueKey);
+			}
+
+
+			// Generate issues with fingerprints
+			{
+				// Reset the writer
+				concreteTelemetryWriter.ClearStoredTelemetry();
+
+				string[] lines =
+				{
+					FileReference.Combine(_workspaceDir, "FOO.CPP").FullName + @"(170) : warning C6011: Dereferencing NULL pointer 'CurrentProperty'. : Lines: 159, 162, 163, 169, 170, 174, 176, 159, 162, 163, 169, 170",
+					FileReference.Combine(_workspaceDir, "foo.cpp").FullName + @"(170) : warning C6011: Dereferencing NULL pointer 'CurrentProperty'. : Lines: 159, 162, 163, 169, 170, 174, 176, 159, 162, 163, 169, 170"
+				};
+
+				IJob job = CreateJob(_mainStreamId, 120, "Compile Test", _graph);
+				await ParseEventsAsync(job, 0, 0, lines);
+				await UpdateCompleteStepAsync(job, 0, 0, JobStepOutcome.Failure);
+
+				IReadOnlyList<IIssue> expectedIssues = await IssueCollection.FindIssuesAsync();
+				IIssue expectedIssue = expectedIssues[0];
+				IList<Tuple<TelemetryRecordMeta?, Hashtable>>  actualTelemetryPayloads = [.. concreteTelemetryWriter.RecordedTelemetry[TelemetryStoreId.Default]];
+				Hashtable actualStateIssuePayload = actualTelemetryPayloads.Where(x => x.Item2[EventNameKeyStr] as string == IssueTelemetry.DefaultEventName).First().Item2;
+
+				// Verify new payloads have been set for fingerprints
+				Assert.IsNotNull(actualStateIssuePayload["Fingerprints"]);
+			}
+		}
+
+		#endregion -- Test Cases --
+	}
+}
\ No newline at end of file
diff --git a/Engine/Source/Programs/Horde/Plugins/Analytics/HordeServer.Analytics/Properties/AssemblyInfo.cs b/Engine/Source/Programs/Horde/Plugins/Analytics/HordeServer.Analytics/Properties/AssemblyInfo.cs
--- a/Engine/Source/Programs/Horde/Plugins/Analytics/HordeServer.Analytics/Properties/AssemblyInfo.cs
+++ b/Engine/Source/Programs/Horde/Plugins/Analytics/HordeServer.Analytics/Properties/AssemblyInfo.cs
@@ -4,4 +4,5 @@ using System.Resources;
 using System.Runtime.CompilerServices;
 
 [assembly: InternalsVisibleTo("HordeServer.Analytics.Tests")]
+[assembly: InternalsVisibleTo("HordeServer.Analytics.Integration.Tests")]
 [assembly: NeutralResourcesLanguage("en-US")]
diff --git a/Engine/Source/Programs/Horde/Plugins/Build/HordeServer.Build.Tests/Issues/IssueServiceTests.cs b/Engine/Source/Programs/Horde/Plugins/Build/HordeServer.Build.Tests/Issues/IssueServiceTests.cs
--- a/Engine/Source/Programs/Horde/Plugins/Build/HordeServer.Build.Tests/Issues/IssueServiceTests.cs
+++ b/Engine/Source/Programs/Horde/Plugins/Build/HordeServer.Build.Tests/Issues/IssueServiceTests.cs
@@ -36,9 +36,9 @@ using Moq;
 namespace HordeServer.Tests.Issues
 {
 	[TestClass]
-	public class IssueServiceTests : BuildTestSetup
+	public abstract class AbstractIssueServiceTests : BuildTestSetup
 	{
-		class TestJsonLogger : ILogger, IAsyncDisposable
+		protected class TestJsonLogger : ILogger, IAsyncDisposable
 		{
 			readonly ILog _log;
 			readonly LogBuilder _builder;
@@ -108,21 +108,32 @@ namespace HordeServer.Tests.Issues
 			}
 		}
 
-		const string MainStreamName = "//UE4/Main";
-		readonly StreamId _mainStreamId = new StreamId("ue4-main");
+		protected const string MainStreamName = "//UE4/Main";
+		protected readonly StreamId _mainStreamId = new StreamId("ue4-main");
 
-		const string ReleaseStreamName = "//UE4/Release";
-		readonly StreamId _releaseStreamId = new StreamId("ue4-release");
+		protected const string ReleaseStreamName = "//UE4/Release";
+		protected readonly StreamId _releaseStreamId = new StreamId("ue4-release");
 
-		const string DevStreamName = "//UE4/Dev";
-		readonly StreamId _devStreamId = new StreamId("ue4-dev");
-		IGraph _graph = default!;
-		PerforceServiceStub _perforce = default!;
-		UserId _jerryId;
-		UserId _bobId;
-		UserId _chrisId;
-		DirectoryReference _autoSdkDir = default!;
-		DirectoryReference _workspaceDir = default!;
+		protected const string DevStreamName = "//UE4/Dev";
+		protected readonly StreamId _devStreamId = new StreamId("ue4-dev");
+		protected IGraph _graph = default!;
+		internal PerforceServiceStub _perforce = default!;
+		protected UserId _jerryId;
+		protected UserId _bobId;
+		protected UserId _chrisId;
+		protected DirectoryReference _autoSdkDir = default!;
+		protected DirectoryReference _workspaceDir = default!;
+
+		public static INode MockNode(string name, IReadOnlyNodeAnnotations annotations)
+		{
+			Mock<INode> node = new Mock<INode>(MockBehavior.Strict);
+			node.SetupGet(x => x.Name).Returns(name);
+			node.SetupGet(x => x.Annotations).Returns(annotations);
+			return node.Object;
+		}
+
+		delegate void TryGetBatchDelegate(JobStepBatchId id, out IJobStepBatch? batch);
+		delegate void TryGetStepDelegate(JobStepId id, out IJobStep? step);
 
 		[TestInitialize]
 		public async Task SetupAsync()
@@ -229,17 +240,6 @@ namespace HordeServer.Tests.Issues
 			_graph = graphMock.Object;
 		}
 
-		public static INode MockNode(string name, IReadOnlyNodeAnnotations annotations)
-		{
-			Mock<INode> node = new Mock<INode>(MockBehavior.Strict);
-			node.SetupGet(x => x.Name).Returns(name);
-			node.SetupGet(x => x.Annotations).Returns(annotations);
-			return node.Object;
-		}
-
-		delegate void TryGetBatchDelegate(JobStepBatchId id, out IJobStepBatch? batch);
-		delegate void TryGetStepDelegate(JobStepId id, out IJobStep? step);
-
 		public IJob CreateJob(StreamId streamId, int change, string name, IGraph graph, TimeSpan time = default, bool promoteByDefault = true, bool updateIssues = true)
 		{
 			JobId jobId = JobIdUtils.GenerateNewId();
@@ -312,7 +312,7 @@ namespace HordeServer.Tests.Issues
 			return job.Object;
 		}
 
-		async Task UpdateCompleteStepAsync(IJob job, int batchIdx, int stepIdx, JobStepOutcome outcome)
+		protected async Task UpdateCompleteStepAsync(IJob job, int batchIdx, int stepIdx, JobStepOutcome outcome)
 		{
 			IJobStepBatch batch = job.Batches[batchIdx];
 			IJobStep step = batch.Steps[stepIdx];
@@ -327,7 +327,7 @@ namespace HordeServer.Tests.Issues
 			}
 		}
 
-		async Task AddEventAsync(IJob job, int batchIdx, int stepIdx, LogLevel logLevel, string? message = null, EventId? id = null)
+		protected async Task AddEventAsync(IJob job, int batchIdx, int stepIdx, LogLevel logLevel, string? message = null, EventId? id = null)
 		{
 			ArrayBufferWriter<byte> buffer = new ArrayBufferWriter<byte>();
 			using (Utf8JsonWriter writer = new Utf8JsonWriter(buffer))
@@ -349,7 +349,7 @@ namespace HordeServer.Tests.Issues
 			await AddEventAsync(job, batchIdx, stepIdx, severity, buffer.WrittenMemory.ToArray());
 		}
 
-		async Task AddEventAsync(IJob job, int batchIdx, int stepIdx, LogEventSeverity severity, byte[] data)
+		protected async Task AddEventAsync(IJob job, int batchIdx, int stepIdx, LogEventSeverity severity, byte[] data)
 		{
 			LogId logId = job.Batches[batchIdx].Steps[stepIdx].LogId!.Value;
 
@@ -368,7 +368,7 @@ namespace HordeServer.Tests.Issues
 			await log.AddEventsAsync(new List<NewLogEventData> { new NewLogEventData { LineIndex = 0, LineCount = 1, Severity = severity } }, CancellationToken.None);
 		}
 
-		private async Task<TestJsonLogger> CreateLoggerAsync(IJob job, int batchIdx, int stepIdx)
+		protected async Task<TestJsonLogger> CreateLoggerAsync(IJob job, int batchIdx, int stepIdx)
 		{
 			LogId logId = job.Batches[batchIdx].Steps[stepIdx].LogId!.Value;
 			IStorageNamespace storageNamespace = StorageService.GetNamespace(Namespace.Logs);
@@ -377,7 +377,7 @@ namespace HordeServer.Tests.Issues
 			return new TestJsonLogger(log, storageNamespace);
 		}
 
-		private async Task ParseEventsAsync(IJob job, int batchIdx, int stepIdx, string[] lines)
+		protected async Task ParseEventsAsync(IJob job, int batchIdx, int stepIdx, string[] lines)
 		{
 			void WriteLinesToLogger(ILogger logger)
 			{
@@ -393,7 +393,7 @@ namespace HordeServer.Tests.Issues
 			await WriteEventsAsync(job, batchIdx, stepIdx, WriteLinesToLogger);
 		}
 
-		private async Task WriteEventsAsync(IJob job, int batchIdx, int stepIdx, Action<ILogger> writeEvents)
+		protected async Task WriteEventsAsync(IJob job, int batchIdx, int stepIdx, Action<ILogger> writeEvents)
 		{
 			LogId logId = job.Batches[batchIdx].Steps[stepIdx].LogId!.Value;
 			ILog log = (await LogCollection.GetAsync(logId))!;
@@ -408,6 +408,11 @@ namespace HordeServer.Tests.Issues
 				writeEvents(perforceLogger);
 			}
 		}
+	}
+
+	[TestClass]
+	public class IssueServiceTests : AbstractIssueServiceTests
+	{
 
 		[TestMethod]
 		public async Task DeleteStreamTestAsync()
diff --git a/Engine/Source/Programs/Horde/Plugins/Build/HordeServer.Build/Issues/IssueCollection.cs b/Engine/Source/Programs/Horde/Plugins/Build/HordeServer.Build/Issues/IssueCollection.cs
--- a/Engine/Source/Programs/Horde/Plugins/Build/HordeServer.Build/Issues/IssueCollection.cs
+++ b/Engine/Source/Programs/Horde/Plugins/Build/HordeServer.Build/Issues/IssueCollection.cs
@@ -819,6 +819,8 @@ namespace HordeServer.Issues
 					AcknowledgedAt = issue.AcknowledgedAt,
 					CreatedAt = issue.CreatedAt,
 					OwnerId = issue.OwnerId,
+					ExternalIssueKey = issue.ExternalIssueKey,
+					Fingerprints = issue.Fingerprints.Select(x=> new {Type = x.Type, Keys = x.Keys}).ToArray(),
 					FixChange = issue.FixCommitId,
 					FixedSystemic = issue.FixSystemic,
 					LastSeenAt = issue.LastSeenAt,
diff --git a/Engine/Source/Programs/Horde/Plugins/Build/HordeServer.Build/Issues/IssueTelemetry.cs b/Engine/Source/Programs/Horde/Plugins/Build/HordeServer.Build/Issues/IssueTelemetry.cs
new file mode 100644
--- /dev/null
+++ b/Engine/Source/Programs/Horde/Plugins/Build/HordeServer.Build/Issues/IssueTelemetry.cs
@@ -0,0 +1,101 @@
+// Copyright Epic Games, Inc. All Rights Reserved.
+
+using EpicGames.Horde.Commits;
+using EpicGames.Horde.Issues;
+using EpicGames.Horde.Jobs;
+using EpicGames.Horde.Jobs.Templates;
+using EpicGames.Horde.Streams;
+using EpicGames.Horde.Users;
+using MongoDB.Bson;
+
+namespace HordeServer.Issues
+{
+	/// <summary>
+	/// Record used for Issue telemetry. Models a <see cref="IIssue"/>.
+	/// </summary>
+	/// <remarks>Modifications to this record's structure should be versioned appropriately, and call sites should continue to publish in a backwards compatible manner.</remarks>
+	public record IssueTelemetry(string EventName, int Id, DateTime? AcknowledgedAt, DateTime CreatedAt, UserId? OwnerId, string? ExternalIssueKey, FingerprintTelemetry[] Fingerprints, CommitId? FixChange, bool FixedSystemic, DateTime LastSeenAt, DateTime? NominatedAt, UserId? NominatedById, DateTime? ResolvedAt, UserId? ResolvedById, IssueSeverity Severity, string Summary, DateTime? VerifiedAt)
+	{
+		/// <summary>
+		/// Default event name for the IssueTelemetry.
+		/// </summary>
+		public static readonly string DefaultEventName = "State.Issue";
+
+		/// <summary>
+		/// Telemetry payload for <see cref="IIssue"/>.
+		/// </summary>
+		public IssueTelemetry(IIssue issue)
+			: this(
+					DefaultEventName,
+					issue.Id,
+					issue.AcknowledgedAt,
+					issue.CreatedAt,
+					issue.OwnerId,
+					issue.ExternalIssueKey,
+					issue.Fingerprints.Select(x => new FingerprintTelemetry(x)).ToArray(),
+					issue.FixCommitId,
+					issue.FixSystemic,
+					issue.LastSeenAt,
+					issue.NominatedAt,
+					issue.NominatedById,
+					issue.ResolvedAt,
+					issue.ResolvedById,
+					issue.Severity,
+					issue.Summary,
+					issue.VerifiedAt
+				  )
+		{
+		}
+	}
+
+	/// <summary>
+	/// Record used for IssueFingerprint telemetry. Models a <see cref="IIssueFingerprint"/>.
+	/// </summary>
+	/// <remarks>Modifications to this record's structure should be versioned appropriately, and call sites should continue to publish in a backwards compatible manner.</remarks>
+	public record FingerprintTelemetry(string Type, IReadOnlySet<IssueKey> Keys, IReadOnlySet<IssueKey>? RejectKeys)
+	{
+		/// <summary>
+		/// Constructor.
+		/// </summary>
+		/// <param name="issueFingerprint">The issueFingerprint to use when constructing the telemetry.</param>
+		public FingerprintTelemetry(IIssueFingerprint issueFingerprint)
+			: this(issueFingerprint.Type, issueFingerprint.Keys, issueFingerprint.RejectKeys)
+		{
+		}
+	}
+
+	/// <summary>
+	/// Record used for FailreInfo telemetry.
+	/// </summary>
+	/// <remarks>Modifications to this record's structure should be versioned appropriately, and call sites should continue to publish in a backwards compatible manner.</remarks>
+	public record FailureInfoTelemetry(JobId JobId, string JobName, CommitIdWithOrder Change, ObjectId StepId);
+
+	/// <summary>
+	/// Record used for IssueSpan telemetry. Modles a <see cref="IIssueSpan"/>.
+	/// </summary>
+	/// <remarks>Modifications to this record's structure should be versioned appropriately, and call sites should continue to publish in a backwards compatible manner.</remarks>
+	public record IssueSpanTelemetry(string EventName, ObjectId Id, int IssueId, FailureInfoTelemetry FirstFailure, FailureInfoTelemetry? LastFailure, StreamId StreamId, string StreamName, TemplateId TemplateRefId)
+	{
+		/// <summary>
+		/// Default event name for the IssueTelemetry.
+		/// </summary>
+		public const string DefaultEventName = "State.IssueSpan";
+
+		/// <summary>
+		/// Constructor.
+		/// </summary>
+		/// <param name="issueSpan">The issueSpan to use when constructing the telemetry.</param>
+		public IssueSpanTelemetry(IIssueSpan issueSpan) : this(DefaultEventName,
+			issueSpan.Id,
+			issueSpan.IssueId,
+			new FailureInfoTelemetry(issueSpan.FirstFailure.JobId, issueSpan.FirstFailure.JobName, issueSpan.FirstFailure.CommitId, issueSpan.FirstFailure.SpanId),
+			issueSpan.LastFailure != null
+				? new FailureInfoTelemetry(issueSpan.LastFailure.JobId, issueSpan.LastFailure.JobName, issueSpan.LastFailure.CommitId, issueSpan.LastFailure.SpanId)
+				: null,
+			issueSpan.StreamId,
+			issueSpan.StreamName,
+			issueSpan.TemplateRefId)
+		{
+		}
+	}
+}
diff --git a/Engine/Source/Programs/Horde/Plugins/Build/HordeServer.Build/Properties/AssemblyInfo.cs b/Engine/Source/Programs/Horde/Plugins/Build/HordeServer.Build/Properties/AssemblyInfo.cs
--- a/Engine/Source/Programs/Horde/Plugins/Build/HordeServer.Build/Properties/AssemblyInfo.cs
+++ b/Engine/Source/Programs/Horde/Plugins/Build/HordeServer.Build/Properties/AssemblyInfo.cs
@@ -5,4 +5,5 @@ using System.Runtime.CompilerServices;
 
 [assembly: InternalsVisibleTo("HordeServer.Tests")]
 [assembly: InternalsVisibleTo("HordeServer.Build.Tests")]
+[assembly: InternalsVisibleTo("HordeServer.Analytics.Integration.Tests")]
 [assembly: NeutralResourcesLanguage("en-US")]
