[Horde][Analytics] General cleanup & test improvements of Analytics.

- Issue telemetry cleanup to use strongly typed records.
- Updated payload test for Fingerprints to better verify contents.
- Example driver for telemetry sinks, and Issue telemetry cleanup to use strongly typed records.
- SchemaVersion added to both State.Issue and State.IssueSpan

diff --git a/Engine/Source/Programs/Horde/Plugins/Analytics/HordeServer.Analytics.Integration.Tests/Telemetry/TelemetryEpicSinkIntegrationTests.cs b/Engine/Source/Programs/Horde/Plugins/Analytics/HordeServer.Analytics.Integration.Tests/Telemetry/TelemetryEpicSinkIntegrationTests.cs
new file mode 100644
--- /dev/null
+++ b/Engine/Source/Programs/Horde/Plugins/Analytics/HordeServer.Analytics.Integration.Tests/Telemetry/TelemetryEpicSinkIntegrationTests.cs
@@ -0,0 +1,141 @@
+// Copyright Epic Games, Inc. All Rights Reserved.
+
+using EpicGames.Core;
+using EpicGames.Horde.Jobs;
+using EpicGames.Horde.Telemetry;
+using HordeServer.Analytics.Tests;
+using HordeServer.Telemetry;
+using HordeServer.Telemetry.Metrics;
+using HordeServer.Tests.Issues;
+using Microsoft.Extensions.Configuration;
+using Microsoft.Extensions.DependencyInjection;
+using Microsoft.Extensions.Options;
+
+namespace HordeServer.Analytics.Integration.Tests.Telemetry
+{
+	/// <summary>
+	/// Test Server Info used to prevent pollution of Epic Telemetry data.
+	/// </summary>
+	public class TestServerInfo : IServerInfo
+	{
+		#region -- Properties & Members --
+
+		private readonly IServerInfo _serverInfo;
+
+		#endregion -- Properties & Members --
+
+		#region -- Constructor --
+
+		public TestServerInfo(IServerInfo originalServerInfo)
+		{
+			_serverInfo = originalServerInfo;
+		}
+
+		#endregion -- Constructor --
+
+		#region -- Interface API --
+
+		/// <inheritdoc/>
+		public SemVer Version => _serverInfo.Version;
+
+		/// <inheritdoc/>
+		public string Environment => "Local";
+
+		/// <inheritdoc/>
+		public string SessionId => _serverInfo.SessionId;
+
+		/// <inheritdoc/>
+		public DirectoryReference AppDir => _serverInfo.AppDir;
+
+		/// <inheritdoc/>
+		public DirectoryReference DataDir => _serverInfo.DataDir;
+
+		/// <inheritdoc/>
+		public IConfiguration Configuration => _serverInfo.Configuration;
+
+		/// <inheritdoc/>
+		public bool ReadOnlyMode => _serverInfo.ReadOnlyMode;
+
+		/// <inheritdoc/>
+		public bool EnableDebugEndpoints => _serverInfo.EnableDebugEndpoints;
+
+		/// <inheritdoc/>
+		public Uri ServerUrl => _serverInfo.ServerUrl;
+
+		/// <inheritdoc/>
+		public Uri DashboardUrl => _serverInfo.DashboardUrl;
+
+		/// <inheritdoc/>
+		public bool IsRunModeActive(RunMode mode)
+		{
+			return _serverInfo.IsRunModeActive(mode);
+		}
+
+		#endregion -- Interface API --
+	}
+
+	/// <summary>
+	/// Test class used to drive data to the Epic Telemetry system in order to verify formats.
+	/// </summary>
+	/// <remarks>Primarily a data driver.</remarks>
+	public class TelemetryEpicSinkIntegrationTests : AbstractIssueServiceTests
+	{
+		private readonly string _epicTelemetryURI = String.Empty;
+
+		protected override void ConfigureServices(IServiceCollection services)
+		{
+			base.ConfigureServices(services);
+
+			EpicTelemetryConfig epicTelemetryConfig = new EpicTelemetryConfig();
+			epicTelemetryConfig.Enabled = true;
+			epicTelemetryConfig.Url = new Uri(_epicTelemetryURI);
+
+			AnalyticsServerConfig analyticsServerConfig = new AnalyticsServerConfig();
+			analyticsServerConfig.Sinks.Epic = epicTelemetryConfig;
+			ServiceProvider provider = services.BuildServiceProvider();
+
+			IServerInfo initialServerInfo = provider.GetRequiredService<IServerInfo>();
+
+			services.AddSingleton(Options.Create(analyticsServerConfig));
+			services.AddSingleton<IServerInfo>(sp =>
+			{
+				return new TestServerInfo(initialServerInfo);
+			});
+		}
+
+		#region -- Test Cases --
+
+		public async Task TestEpicSinkAsync()
+		{
+			// Configure the metric
+			TelemetryStoreConfig telemetryStoreConfig = new TelemetryStoreConfig();
+			telemetryStoreConfig.Id = TelemetryStoreId.Default;
+
+			AnalyticsConfig analyticsConfig = new AnalyticsConfig();
+			analyticsConfig.Stores.Add(telemetryStoreConfig);
+
+			await UpdateConfigAsync(config =>
+			{
+				config.Plugins.GetBuildConfig().GetStream(_mainStreamId).TelemetryStoreId = TelemetryStoreId.Default;
+				config.Plugins.AddAnalyticsTestConfig(analyticsConfig);
+			});
+
+			// Generate issue telemetry data
+			{
+				string[] lines =
+				{
+					@"WARNING: Engine\Source\Programs\UnrealBuildTool\ProjectFiles\Rider\ToolchainInfo.cs: Missing copyright boilerplate"
+				};
+
+				IJob job = CreateJob(_mainStreamId, 105, "Test Build", _graph);
+				await ParseEventsAsync(job, 0, 0, lines);
+				await UpdateCompleteStepAsync(job, 0, 0, JobStepOutcome.Failure);
+			}
+
+			TelemetryManager telemetryManager = ServiceProvider.GetRequiredService<TelemetryManager>();
+			await telemetryManager.FlushAsync(CancellationToken.None);
+		}
+
+		#endregion -- Test Cases --
+	}
+}
diff --git a/Engine/Source/Programs/Horde/Plugins/Analytics/HordeServer.Analytics.Integration.Tests/Telemetry/TelemetryIntegrationTests.cs b/Engine/Source/Programs/Horde/Plugins/Analytics/HordeServer.Analytics.Integration.Tests/Telemetry/TelemetryIntegrationTests.cs
--- a/Engine/Source/Programs/Horde/Plugins/Analytics/HordeServer.Analytics.Integration.Tests/Telemetry/TelemetryIntegrationTests.cs
+++ b/Engine/Source/Programs/Horde/Plugins/Analytics/HordeServer.Analytics.Integration.Tests/Telemetry/TelemetryIntegrationTests.cs
@@ -157,7 +157,7 @@ namespace HordeServer.Analytics.Integration.Tests.Telemetry
 				Assert.IsTrue(concreteTelemetryWriter.RecordedTelemetry.ContainsKey(TelemetryStoreId.Default));
 
 				// Verify presence of Issue telemetry
-				IList<Tuple<TelemetryRecordMeta?, Hashtable>>  actualTelemetryPayloads = [.. concreteTelemetryWriter.RecordedTelemetry[TelemetryStoreId.Default]];
+				IList<Tuple<TelemetryRecordMeta?, Hashtable>> actualTelemetryPayloads = [.. concreteTelemetryWriter.RecordedTelemetry[TelemetryStoreId.Default]];
 				Assert.IsTrue(actualTelemetryPayloads.Select(x => x.Item2[EventNameKeyStr] as string == IssueTelemetry.DefaultEventName).Any());
 				Assert.IsTrue(actualTelemetryPayloads.Select(x => x.Item2[EventNameKeyStr] as string == IssueSpanTelemetry.DefaultEventName).Any());
 			}
@@ -171,7 +171,7 @@ namespace HordeServer.Analytics.Integration.Tests.Telemetry
 				IReadOnlyList<IIssue> issues = await IssueCollection.FindIssuesAsync();
 				await IssueService.UpdateIssueAsync(issues[0].Id, externalIssueKey: ExpectedExternalIssueKey);
 
-				IList<Tuple<TelemetryRecordMeta?, Hashtable>>  actualTelemetryPayloads = [.. concreteTelemetryWriter.RecordedTelemetry[TelemetryStoreId.Default]];
+				IList<Tuple<TelemetryRecordMeta?, Hashtable>> actualTelemetryPayloads = [.. concreteTelemetryWriter.RecordedTelemetry[TelemetryStoreId.Default]];
 				Hashtable actualStateIssuePayload = actualTelemetryPayloads.Where(x => x.Item2[EventNameKeyStr] as string == IssueTelemetry.DefaultEventName).First().Item2;
 				IssueTelemetry? actualIssueTelemetry = actualStateIssuePayload.MapToObject<IssueTelemetry>();
 
@@ -197,11 +197,16 @@ namespace HordeServer.Analytics.Integration.Tests.Telemetry
 
 				IReadOnlyList<IIssue> expectedIssues = await IssueCollection.FindIssuesAsync();
 				IIssue expectedIssue = expectedIssues[0];
-				IList<Tuple<TelemetryRecordMeta?, Hashtable>>  actualTelemetryPayloads = [.. concreteTelemetryWriter.RecordedTelemetry[TelemetryStoreId.Default]];
-				Hashtable actualStateIssuePayload = actualTelemetryPayloads.Where(x => x.Item2[EventNameKeyStr] as string == IssueTelemetry.DefaultEventName).First().Item2;
-
-				// Verify new payloads have been set for fingerprints
-				Assert.IsNotNull(actualStateIssuePayload["Fingerprints"]);
+				IIssueFingerprint expectedFingerprint = expectedIssue.Fingerprints[0];
+				IList<Tuple<TelemetryRecordMeta?, Hashtable>> actualTelemetryPayloads = [.. concreteTelemetryWriter.RecordedTelemetry[TelemetryStoreId.Default]];
+				Hashtable actualStateIssueSpanPayload = actualTelemetryPayloads.Where(x => x.Item2[EventNameKeyStr] as string == IssueSpanTelemetry.DefaultEventName).First().Item2;
+				IssueSpanTelemetry? actualIssueSpanTelemetry = actualStateIssueSpanPayload.MapToObject<IssueSpanTelemetry>();
+
+				// Verify new payloads have been set for fingerprint
+				Assert.IsNotNull(actualIssueSpanTelemetry);
+				Assert.AreEqual(expectedFingerprint.Type, actualIssueSpanTelemetry.Fingerprint.Type);
+				Assert.AreEqual(expectedFingerprint.Keys.Count, actualIssueSpanTelemetry.Fingerprint.Keys.Count);
+				Assert.AreEqual(expectedFingerprint.Keys.First(), actualIssueSpanTelemetry.Fingerprint.Keys.First());
 			}
 		}
 
diff --git a/Engine/Source/Programs/Horde/Plugins/Analytics/HordeServer.Analytics.Integration.Tests/Telemetry/TelemetryMongoDBSinkIntegrationTests.cs b/Engine/Source/Programs/Horde/Plugins/Analytics/HordeServer.Analytics.Integration.Tests/Telemetry/TelemetryMongoDBSinkIntegrationTests.cs
new file mode 100644
--- /dev/null
+++ b/Engine/Source/Programs/Horde/Plugins/Analytics/HordeServer.Analytics.Integration.Tests/Telemetry/TelemetryMongoDBSinkIntegrationTests.cs
@@ -0,0 +1,83 @@
+// Copyright Epic Games, Inc. All Rights Reserved.
+
+using EpicGames.Horde.Jobs;
+using EpicGames.Horde.Telemetry;
+using HordeServer.Analytics.Tests;
+using HordeServer.Server;
+using HordeServer.Telemetry;
+using HordeServer.Telemetry.Metrics;
+using HordeServer.Tests.Issues;
+using HordeServer.VersionControl.Perforce;
+using Microsoft.Extensions.DependencyInjection;
+using Microsoft.Extensions.Logging;
+using Microsoft.Extensions.Options;
+using MongoDB.Bson;
+using MongoDB.Driver;
+using StackExchange.Redis;
+
+namespace HordeServer.Analytics.Integration.Tests.Telemetry
+{
+	/// <summary>
+	/// Test class used to drive data to the Mongo DB Telemetry system in order to verify formats.
+	/// </summary>
+	/// <remarks>Primarily a data driver.</remarks>
+	[TestClass]
+	public class TelemetryMongoDBSinkIntegrationTests : AbstractIssueServiceTests
+	{
+		ILogger<TelemetryMongoDBSinkIntegrationTests> _logger = CreateConsoleLogger<TelemetryMongoDBSinkIntegrationTests>();
+		protected override void ConfigureServices(IServiceCollection services)
+		{
+			base.ConfigureServices(services);
+
+			MongoTelemetryConfig mongoTelemetryConfig = new MongoTelemetryConfig();
+			mongoTelemetryConfig.Enabled = true;
+			mongoTelemetryConfig.RetainDays = 1;
+
+			AnalyticsServerConfig analyticsServerConfig = new AnalyticsServerConfig();
+			analyticsServerConfig.Sinks.Mongo = mongoTelemetryConfig;
+			services.AddSingleton(Options.Create(analyticsServerConfig));
+		}
+
+		#region -- Test Cases --
+
+		public async Task TestMongoDBTelemetrySinkAsync()
+		{
+			// Configure the metric
+			TelemetryStoreConfig telemetryStoreConfig = new TelemetryStoreConfig();
+			telemetryStoreConfig.Id = TelemetryStoreId.Default;
+
+			AnalyticsConfig analyticsConfig = new AnalyticsConfig();
+			analyticsConfig.Stores.Add(telemetryStoreConfig);
+
+			await UpdateConfigAsync(config =>
+			{
+				config.Plugins.GetBuildConfig().GetStream(_mainStreamId).TelemetryStoreId = TelemetryStoreId.Default;
+				config.Plugins.AddAnalyticsTestConfig(analyticsConfig);
+			});
+
+			IMongoService mongoService = ServiceProvider.GetRequiredService<IMongoService>();
+			IMongoCollection<BsonDocument> collection = mongoService.Database.GetCollection<BsonDocument>("Telemetry");
+			long count = await collection.CountDocumentsAsync(FilterDefinition<BsonDocument>.Empty);
+			Assert.AreEqual(0, count);
+			// Generate issue telemetry data
+			{
+				string[] lines =
+				{
+					@"WARNING: Engine\Source\Programs\UnrealBuildTool\ProjectFiles\Rider\ToolchainInfo.cs: Missing copyright boilerplate"
+				};
+
+				IJob job = CreateJob(_mainStreamId, 105, "Test Build", _graph);
+				await ParseEventsAsync(job, 0, 0, lines);
+				await UpdateCompleteStepAsync(job, 0, 0, JobStepOutcome.Failure);
+			}
+
+			TelemetryManager telemetryManager = ServiceProvider.GetRequiredService<TelemetryManager>();
+			await telemetryManager.FlushAsync(CancellationToken.None);
+
+			count = await collection.CountDocumentsAsync(FilterDefinition<BsonDocument>.Empty);
+			Assert.IsTrue(count > 0);
+		}
+
+		#endregion -- Test Cases --
+	}
+}
diff --git a/Engine/Source/Programs/Horde/Plugins/Build/HordeServer.Build/Issues/IssueCollection.cs b/Engine/Source/Programs/Horde/Plugins/Build/HordeServer.Build/Issues/IssueCollection.cs
--- a/Engine/Source/Programs/Horde/Plugins/Build/HordeServer.Build/Issues/IssueCollection.cs
+++ b/Engine/Source/Programs/Horde/Plugins/Build/HordeServer.Build/Issues/IssueCollection.cs
@@ -812,26 +812,7 @@ namespace HordeServer.Issues
 
 			foreach (TelemetryStoreId telemetryStoreId in telemetryStoreIds)
 			{
-				_telemetryWriter.WriteEvent(telemetryStoreId, new
-				{
-					EventName = "State.Issue",
-					Id = issue.Id,
-					AcknowledgedAt = issue.AcknowledgedAt,
-					CreatedAt = issue.CreatedAt,
-					OwnerId = issue.OwnerId,
-					ExternalIssueKey = issue.ExternalIssueKey,
-					Fingerprints = issue.Fingerprints.Select(x=> new {Type = x.Type, Keys = x.Keys}).ToArray(),
-					FixChange = issue.FixCommitId,
-					FixedSystemic = issue.FixSystemic,
-					LastSeenAt = issue.LastSeenAt,
-					NominatedAt = issue.NominatedAt,
-					NominatedById = issue.NominatedById,
-					ResolvedAt = issue.ResolvedAt,
-					ResolvedById = issue.ResolvedById,
-					Severity = issue.Severity,
-					Summary = issue.Summary,
-					VerifiedAt = issue.VerifiedAt
-				});
+				_telemetryWriter.WriteEvent(telemetryStoreId, new IssueTelemetry(issue));
 			}
 		}
 
@@ -1626,18 +1607,7 @@ namespace HordeServer.Issues
 		{
 			if (_buildConfig.CurrentValue.TryGetStream(issueSpan.StreamId, out StreamConfig? streamConfig) && !streamConfig.TelemetryStoreId.IsEmpty)
 			{
-				_telemetryWriter.WriteEvent(streamConfig.TelemetryStoreId, new
-				{
-					EventName = "State.IssueSpan",
-					Id = issueSpan.Id,
-					IssueId = issueSpan.IssueId,
-					Fingerprint = new { Type = issueSpan.Fingerprint.Type, Keys = issueSpan.Fingerprint.Keys },
-					FirstFailure = new { JobId = issueSpan.FirstFailure.JobId, JobName = issueSpan.FirstFailure.JobName, Change = issueSpan.FirstFailure.CommitId, StepId = issueSpan.FirstFailure.SpanId },
-					LastFailure = (issueSpan.LastFailure != null) ? new { JobId = issueSpan.LastFailure.JobId, JobName = issueSpan.LastFailure.JobName, Change = issueSpan.LastFailure.CommitId, StepId = issueSpan.LastFailure.SpanId } : null,
-					StreamId = issueSpan.StreamId,
-					StreamName = issueSpan.StreamName,
-					TemplateRefId = issueSpan.TemplateRefId
-				});
+				_telemetryWriter.WriteEvent(streamConfig.TelemetryStoreId, new IssueSpanTelemetry(issueSpan));
 			}
 		}
 
diff --git a/Engine/Source/Programs/Horde/Plugins/Build/HordeServer.Build/Issues/IssueTelemetry.cs b/Engine/Source/Programs/Horde/Plugins/Build/HordeServer.Build/Issues/IssueTelemetry.cs
--- a/Engine/Source/Programs/Horde/Plugins/Build/HordeServer.Build/Issues/IssueTelemetry.cs
+++ b/Engine/Source/Programs/Horde/Plugins/Build/HordeServer.Build/Issues/IssueTelemetry.cs
@@ -13,26 +13,31 @@ namespace HordeServer.Issues
 	/// <summary>
 	/// Record used for Issue telemetry. Models a <see cref="IIssue"/>.
 	/// </summary>
-	/// <remarks>Modifications to this record's structure should be versioned appropriately, and call sites should continue to publish in a backwards compatible manner.</remarks>
-	public record IssueTelemetry(string EventName, int Id, DateTime? AcknowledgedAt, DateTime CreatedAt, UserId? OwnerId, string? ExternalIssueKey, FingerprintTelemetry[] Fingerprints, CommitId? FixChange, bool FixedSystemic, DateTime LastSeenAt, DateTime? NominatedAt, UserId? NominatedById, DateTime? ResolvedAt, UserId? ResolvedById, IssueSeverity Severity, string Summary, DateTime? VerifiedAt)
+	/// <remarks>Modifications to this record's structure should be versioned appropriately through <see cref="SchemaVersion"/> via <see cref="CurrentSchemaVersion"/>, and call sites should continue to publish in a backwards compatible manner.</remarks>
+	public record IssueTelemetry(string EventName, int SchemaVersion, int Id, DateTime? AcknowledgedAt, DateTime CreatedAt, UserId? OwnerId, string? ExternalIssueKey, CommitId? FixChange, bool FixedSystemic, DateTime LastSeenAt, DateTime? NominatedAt, UserId? NominatedById, DateTime? ResolvedAt, UserId? ResolvedById, IssueSeverity Severity, string Summary, DateTime? VerifiedAt)
 	{
 		/// <summary>
 		/// Default event name for the IssueTelemetry.
 		/// </summary>
 		public static readonly string DefaultEventName = "State.Issue";
 
+		/// <summary>
+		/// Current schema version.
+		/// </summary>
+		public static readonly int CurrentSchemaVersion = 1;
+
 		/// <summary>
 		/// Telemetry payload for <see cref="IIssue"/>.
 		/// </summary>
 		public IssueTelemetry(IIssue issue)
 			: this(
 					DefaultEventName,
+					CurrentSchemaVersion,
 					issue.Id,
 					issue.AcknowledgedAt,
 					issue.CreatedAt,
 					issue.OwnerId,
 					issue.ExternalIssueKey,
-					issue.Fingerprints.Select(x => new FingerprintTelemetry(x)).ToArray(),
 					issue.FixCommitId,
 					issue.FixSystemic,
 					issue.LastSeenAt,
@@ -52,14 +57,14 @@ namespace HordeServer.Issues
 	/// Record used for IssueFingerprint telemetry. Models a <see cref="IIssueFingerprint"/>.
 	/// </summary>
 	/// <remarks>Modifications to this record's structure should be versioned appropriately, and call sites should continue to publish in a backwards compatible manner.</remarks>
-	public record FingerprintTelemetry(string Type, IReadOnlySet<IssueKey> Keys, IReadOnlySet<IssueKey>? RejectKeys)
+	public record FingerprintTelemetry(string Type, IReadOnlySet<IssueKey> Keys)
 	{
 		/// <summary>
 		/// Constructor.
 		/// </summary>
 		/// <param name="issueFingerprint">The issueFingerprint to use when constructing the telemetry.</param>
 		public FingerprintTelemetry(IIssueFingerprint issueFingerprint)
-			: this(issueFingerprint.Type, issueFingerprint.Keys, issueFingerprint.RejectKeys)
+			: this(issueFingerprint.Type, issueFingerprint.Keys)
 		{
 		}
 	}
@@ -73,21 +78,29 @@ namespace HordeServer.Issues
 	/// <summary>
 	/// Record used for IssueSpan telemetry. Modles a <see cref="IIssueSpan"/>.
 	/// </summary>
-	/// <remarks>Modifications to this record's structure should be versioned appropriately, and call sites should continue to publish in a backwards compatible manner.</remarks>
-	public record IssueSpanTelemetry(string EventName, ObjectId Id, int IssueId, FailureInfoTelemetry FirstFailure, FailureInfoTelemetry? LastFailure, StreamId StreamId, string StreamName, TemplateId TemplateRefId)
+	/// <remarks>Modifications to this record's structure should be versioned appropriately through <see cref="SchemaVersion"/> via <see cref="CurrentSchemaVersion"/>, and call sites should continue to publish in a backwards compatible manner.</remarks>
+	public record IssueSpanTelemetry(string EventName, int SchemaVersion, ObjectId Id, int IssueId, FingerprintTelemetry Fingerprint, FailureInfoTelemetry FirstFailure, FailureInfoTelemetry? LastFailure, StreamId StreamId, string StreamName, TemplateId TemplateRefId)
 	{
 		/// <summary>
 		/// Default event name for the IssueTelemetry.
 		/// </summary>
 		public const string DefaultEventName = "State.IssueSpan";
 
+		/// <summary>
+		/// Current schema version.
+		/// </summary>
+		public static readonly int CurrentSchemaVersion = 1;
+
 		/// <summary>
 		/// Constructor.
 		/// </summary>
 		/// <param name="issueSpan">The issueSpan to use when constructing the telemetry.</param>
-		public IssueSpanTelemetry(IIssueSpan issueSpan) : this(DefaultEventName,
+		public IssueSpanTelemetry(IIssueSpan issueSpan) : this(
+			DefaultEventName,
+			CurrentSchemaVersion,
 			issueSpan.Id,
 			issueSpan.IssueId,
+			new FingerprintTelemetry(issueSpan.Fingerprint),
 			new FailureInfoTelemetry(issueSpan.FirstFailure.JobId, issueSpan.FirstFailure.JobName, issueSpan.FirstFailure.CommitId, issueSpan.FirstFailure.SpanId),
 			issueSpan.LastFailure != null
 				? new FailureInfoTelemetry(issueSpan.LastFailure.JobId, issueSpan.LastFailure.JobName, issueSpan.LastFailure.CommitId, issueSpan.LastFailure.SpanId)
