Backed out the change to create partitioned workspaces by default in Horde as it re-introduced the 'partitioned client type has not been configured for this server.' error on a standard setup of Perforce.

#rnx

diff --git a/Engine/Source/Programs/Horde/Plugins/Build/HordeServer.Build/BuildConfig.cs b/Engine/Source/Programs/Horde/Plugins/Build/HordeServer.Build/BuildConfig.cs
--- a/Engine/Source/Programs/Horde/Plugins/Build/HordeServer.Build/BuildConfig.cs
+++ b/Engine/Source/Programs/Horde/Plugins/Build/HordeServer.Build/BuildConfig.cs
@@ -515,7 +515,7 @@ namespace HordeServer
 		/// <summary>
 		/// Whether to use partitioned workspaces on this server
 		/// </summary>
-		public bool SupportsPartitionedWorkspaces { get; set; } = true;
+		public bool SupportsPartitionedWorkspaces { get; set; } = false;
 
 		/// <summary>
 		/// List of servers
