Horde: Change how secrets in the global config are resolved

The attribute ResolveSecret is used to identifiy the string properties in the Global Config that can be resolved to a secret during load

diff --git a/Engine/Source/Programs/Horde/HordeServer.Shared/Configuration/ResolveSecretAttribute.cs b/Engine/Source/Programs/Horde/HordeServer.Shared/Configuration/ResolveSecretAttribute.cs
new file mode 100644
--- /dev/null
+++ b/Engine/Source/Programs/Horde/HordeServer.Shared/Configuration/ResolveSecretAttribute.cs
@@ -0,0 +1,9 @@
+// Copyright Epic Games, Inc. All Rights Reserved.
+
+namespace HordeServer.Configuration;
+
+/// <summary>
+/// Convert a string representation of a secret to a concrete value during config updates
+/// </summary>
+[AttributeUsage(AttributeTargets.Property)]
+public sealed class ResolveSecretAttribute : Attribute;
diff --git a/Engine/Source/Programs/Horde/HordeServer.Shared/Plugins/IPluginConfig.cs b/Engine/Source/Programs/Horde/HordeServer.Shared/Plugins/IPluginConfig.cs
--- a/Engine/Source/Programs/Horde/HordeServer.Shared/Plugins/IPluginConfig.cs
+++ b/Engine/Source/Programs/Horde/HordeServer.Shared/Plugins/IPluginConfig.cs
@@ -20,7 +20,7 @@ namespace HordeServer.Plugins
 	/// <summary>
 	/// Options passed to <see cref="IPluginConfig.PostLoad(PluginConfigOptions)"/>
 	/// </summary>
-	public record class PluginConfigOptions(ConfigVersion Version, IEnumerable<IPluginConfig> Plugins, AclConfig ParentAcl, IServiceProvider? ServiceProvider = null, ILogger? Logger = null);
+	public record class PluginConfigOptions(ConfigVersion Version, IEnumerable<IPluginConfig> Plugins, AclConfig ParentAcl, ILogger? Logger = null);
 
 	/// <summary>
 	/// Empty implementation of <see cref="IPluginConfig"/>
diff --git a/Engine/Source/Programs/Horde/HordeServer.Tests.Shared/ServerTestSetup.cs b/Engine/Source/Programs/Horde/HordeServer.Tests.Shared/ServerTestSetup.cs
--- a/Engine/Source/Programs/Horde/HordeServer.Tests.Shared/ServerTestSetup.cs
+++ b/Engine/Source/Programs/Horde/HordeServer.Tests.Shared/ServerTestSetup.cs
@@ -76,7 +76,8 @@ namespace HordeServer.Tests
 		protected async Task SetConfigAsync(GlobalConfig globalConfig)
 		{
 			await ConfigService.WaitForInitialConfigAsync();
-			globalConfig.PostLoad(ServerSettings, _pluginCollection.LoadedPlugins, DefaultAclModifiers, ServiceProvider, ConfigServiceLogger);
+			globalConfig.PostLoad(ServerSettings, _pluginCollection.LoadedPlugins, DefaultAclModifiers, ConfigServiceLogger);
+			await ConfigService.ResolveConfigSecretsAsync(globalConfig, CancellationToken.None);
 			ConfigService.OverrideConfig(globalConfig);
 		}
 
@@ -85,7 +86,8 @@ namespace HordeServer.Tests
 			await ConfigService.WaitForInitialConfigAsync();
 			GlobalConfig globalConfig = GlobalConfig.CurrentValue;
 			action(globalConfig);
-			globalConfig.PostLoad(ServerSettings, _pluginCollection.LoadedPlugins, DefaultAclModifiers, ServiceProvider, ConfigServiceLogger);
+			globalConfig.PostLoad(ServerSettings, _pluginCollection.LoadedPlugins, DefaultAclModifiers, ConfigServiceLogger);
+			await ConfigService.ResolveConfigSecretsAsync(globalConfig, CancellationToken.None);
 			ConfigService.OverrideConfig(globalConfig);
 		}
 
diff --git a/Engine/Source/Programs/Horde/HordeServer/Configuration/ConfigService.cs b/Engine/Source/Programs/Horde/HordeServer/Configuration/ConfigService.cs
--- a/Engine/Source/Programs/Horde/HordeServer/Configuration/ConfigService.cs
+++ b/Engine/Source/Programs/Horde/HordeServer/Configuration/ConfigService.cs
@@ -1,16 +1,20 @@
 // Copyright Epic Games, Inc. All Rights Reserved.
 
 using System;
+using System.Collections;
 using System.Collections.Generic;
 using System.IO;
 using System.IO.Compression;
 using System.Linq;
+using System.Reflection;
 using System.Text;
 using System.Text.Json;
 using System.Text.Json.Nodes;
+using System.Text.RegularExpressions;
 using System.Threading;
 using System.Threading.Tasks;
 using EpicGames.Core;
+using EpicGames.Horde.Secrets;
 using EpicGames.Horde.Users;
 using EpicGames.Redis;
 using HordeServer.Acls;
@@ -18,6 +22,7 @@ using HordeServer.Plugins;
 using HordeServer.Server;
 using HordeServer.Users;
 using HordeServer.Utilities;
+using Microsoft.Extensions.DependencyInjection;
 using Microsoft.Extensions.Diagnostics.HealthChecks;
 using Microsoft.Extensions.Hosting;
 using Microsoft.Extensions.Logging;
@@ -285,7 +290,9 @@ namespace HordeServer.Configuration
 			try
 			{
 				GlobalConfig globalConfig = await ReadGlobalConfigAsync(context, cancellationToken);
-				globalConfig.PostLoad(_serverSettings, _pluginCollection.LoadedPlugins, _defaultAclModifiers, _serviceProvider, _logger);
+				globalConfig.PostLoad(_serverSettings, _pluginCollection.LoadedPlugins, _defaultAclModifiers, _logger);
+
+				await ResolveConfigSecretsAsync(globalConfig, cancellationToken);
 
 				foreach (OverrideConfigFile file in overrideFiles.Values)
 				{
@@ -395,7 +402,7 @@ namespace HordeServer.Configuration
 					data = SerializeSnapshot(snapshot);
 				}
 
-				ConfigState state = new ConfigState(IoHash.Compute(data.Span), CreateGlobalConfig(data));
+				ConfigState state = new ConfigState(IoHash.Compute(data.Span), await CreateGlobalConfigAsync(data, CancellationToken.None));
 				await _health.UpdateAsync(HealthStatus.Healthy);
 
 				return state;
@@ -611,7 +618,9 @@ namespace HordeServer.Configuration
 				}
 
 				// Execute a PostLoad before returning so we can validate that everything is valid
-				globalConfig.PostLoad(_serverSettings, _pluginCollection.LoadedPlugins, _defaultAclModifiers, _serviceProvider, _logger);
+				globalConfig.PostLoad(_serverSettings, _pluginCollection.LoadedPlugins, _defaultAclModifiers, _logger);
+
+				await ResolveConfigSecretsAsync(globalConfig, cancellationToken);
 
 				return snapshot;
 			}
@@ -844,13 +853,13 @@ namespace HordeServer.Configuration
 			IoHash hash = IoHash.Compute(data.Span);
 			if (hash != state.Hash && state.Hash != IoHash.Zero) // Don't replace any explicit config override
 			{
-				GlobalConfig globalConfig = CreateGlobalConfig(data);
+				GlobalConfig globalConfig = await CreateGlobalConfigAsync(data, CancellationToken.None);
 				_logger.LogInformation("Updating config state from Redis (hash: {OldHash} -> {NewHash}, new size: {Size})", state.Hash, hash, data.Length);
 				Set(hash, globalConfig);
 			}
 		}
 
-		GlobalConfig CreateGlobalConfig(ReadOnlyMemory<byte> data)
+		async Task<GlobalConfig> CreateGlobalConfigAsync(ReadOnlyMemory<byte> data, CancellationToken cancellationToken = default)
 		{
 			// Decompress the snapshot object
 			ConfigSnapshot snapshot = DeserializeSnapshot(data);
@@ -860,8 +869,152 @@ namespace HordeServer.Configuration
 			globalConfig.Revision = IoHash.Compute(data.Span).ToString();
 
 			// Run the postload callbacks on all the config objects
-			globalConfig.PostLoad(_serverSettings, _pluginCollection.LoadedPlugins, _defaultAclModifiers, _serviceProvider, _logger);
+			globalConfig.PostLoad(_serverSettings, _pluginCollection.LoadedPlugins, _defaultAclModifiers, _logger);
+
+			await ResolveConfigSecretsAsync(globalConfig, cancellationToken);
+
 			return globalConfig;
 		}
+
+		class SecretReplacementInfo
+		{
+			public required string SecretName { get; init; }
+			public required object SourceObject { get; init; }
+			public required PropertyInfo Property { get; init; }
+		}
+
+		/// <summary>
+		/// Convert string representations of secrets to concrete values in the config
+		/// </summary>
+		/// <param name="config">The config to update</param>
+		/// <param name="cancellationToken">Cancellation token for the operation</param>
+		public async Task ResolveConfigSecretsAsync(GlobalConfig config, CancellationToken cancellationToken)
+		{
+			if (!config.ResolveSecrets)
+			{
+				return;
+			}
+
+			SecretsConfig? secretsConfig = config.Plugins.Values.OfType<SecretsConfig>().FirstOrDefault();
+			if (secretsConfig == null)
+			{
+				return;
+			}
+
+			await using AsyncServiceScope scope = _serviceProvider.CreateAsyncScope();
+			ISecretCollection? secretCollection = scope.ServiceProvider.GetService<ISecretCollection>();
+			if (secretCollection == null)
+			{
+				return;
+			}
+
+			Dictionary<SecretId, List<SecretReplacementInfo>> secrets = [];
+			foreach ((object obj, PropertyInfo prop) in GetPropertiesWithAttribute<string, ResolveSecretAttribute>(config))
+			{
+				string? value = prop.GetValue(obj) as string;
+				if (String.IsNullOrEmpty(value))
+				{
+					continue;
+				}
+
+				Match match = Regex.Match(value, @"^horde:secret:(?<SecretId>[^\.]+)\.(?<PropertyName>[^\.]+)?$");
+				if (!match.Success)
+				{
+					continue;
+				}
+				SecretId secretId = new(match.Groups["SecretId"].Value);
+				string secretName = match.Groups["PropertyName"].Value;
+
+				if (!secrets.TryGetValue(secretId, out List<SecretReplacementInfo>? replacementInfo))
+				{
+					replacementInfo = [];
+					secrets.Add(secretId, replacementInfo);
+				}
+				replacementInfo.Add(new SecretReplacementInfo {
+					SecretName = secretName, Property = prop, SourceObject = obj,
+				});
+			}
+
+			foreach ((SecretId secretId, List<SecretReplacementInfo> replacementInfos) in secrets)
+			{
+				ISecret? secret = await secretCollection.GetAsync(secretId, secretsConfig, cancellationToken);
+				if (secret == null)
+				{
+					throw new KeyNotFoundException($"The given secret '{secretId}' was not found");
+				}
+
+				foreach (SecretReplacementInfo replacementInfo in replacementInfos)
+				{
+					if (!secret.Data.TryGetValue(replacementInfo.SecretName, out string? value))
+					{
+						throw new KeyNotFoundException($"The given property '{replacementInfo.SecretName}' for the secret '{secretId}' was not found");
+					}
+					replacementInfo.Property.SetValue(replacementInfo.SourceObject, value);
+				}
+			}
+		}
+
+		static List<KeyValuePair<object, PropertyInfo>> GetPropertiesWithAttribute<T, TAttribute>(object obj) where TAttribute : Attribute
+		{
+			List<KeyValuePair<object, PropertyInfo>> results = [];
+			Stack<object?> stack = new([obj]);
+			HashSet<object?> visited = [];
+
+			while (stack.Count > 0)
+			{
+				object? currentObj = stack.Pop();
+				if (currentObj == null || visited.Contains(currentObj))
+				{
+					continue;
+				}
+
+				if (currentObj is not string && currentObj is IEnumerable elems)
+				{
+					foreach (object? elem in elems)
+					{
+						if (elem != null)
+						{
+							stack.Push(elem);
+						}
+					}
+				}
+				else
+				{
+					PropertyInfo[] props = currentObj.GetType().GetProperties(BindingFlags.Public | BindingFlags.Instance);
+					foreach (PropertyInfo prop in props)
+					{
+						if (!prop.CanRead)
+						{
+							continue;
+						}
+
+						if (prop.PropertyType == typeof(T))
+						{
+							if (prop.GetCustomAttributes<TAttribute>().Any())
+							{
+								KeyValuePair<object, PropertyInfo> result = new(currentObj, prop);
+								results.Add(result);
+							}
+						}
+						else
+						{
+							// Check if reached properties that should not be enumerated any further
+							if (!(prop.PropertyType.IsPrimitive || prop.GetIndexParameters().Length > 0 || prop.PropertyType.IsByRefLike))
+							{
+								object? value = prop.GetValue(currentObj);
+								if (value != null)
+								{
+									stack.Push(value);
+								}
+							}
+						}
+					}
+				}
+
+				visited.Add(currentObj);
+			}
+
+			return results;
+		}
 	}
 }
diff --git a/Engine/Source/Programs/Horde/HordeServer/Server/GlobalConfig.cs b/Engine/Source/Programs/Horde/HordeServer/Server/GlobalConfig.cs
--- a/Engine/Source/Programs/Horde/HordeServer/Server/GlobalConfig.cs
+++ b/Engine/Source/Programs/Horde/HordeServer/Server/GlobalConfig.cs
@@ -85,6 +85,11 @@ namespace HordeServer.Server
 		/// </summary>
 		public JsonObject Parameters { get; set; } = new JsonObject();
 
+		/// <summary>
+		/// Whether to enable resolving secrets in config properties
+		/// </summary>
+		public bool ResolveSecrets { get; set; }
+
 		/// <summary>
 		/// Access control list
 		/// </summary>
@@ -105,7 +110,7 @@ namespace HordeServer.Server
 		/// <summary>
 		/// Called after the config file has been read
 		/// </summary>
-		public void PostLoad(ServerSettings serverSettings, IReadOnlyList<ILoadedPlugin> loadedPlugins, IEnumerable<IDefaultAclModifier> defaultAclModifiers, IServiceProvider? serviceProvider = null, ILogger? logger = null)
+		public void PostLoad(ServerSettings serverSettings, IReadOnlyList<ILoadedPlugin> loadedPlugins, IEnumerable<IDefaultAclModifier> defaultAclModifiers, ILogger? logger = null)
 		{
 			ServerSettings = serverSettings;
 			AclAction[] aclActions = AclConfig.GetActions(
@@ -155,7 +160,7 @@ namespace HordeServer.Server
 				pluginConfigs = sortedPluginConfigs;
 			}
 
-			PluginConfigOptions pluginConfigOptions = new PluginConfigOptions(VersionEnum, Plugins.Values, Acl, serviceProvider, logger);
+			PluginConfigOptions pluginConfigOptions = new PluginConfigOptions(VersionEnum, Plugins.Values, Acl, logger);
 			foreach (IPluginConfig pluginConfig in pluginConfigs)
 			{
 				pluginConfig.PostLoad(pluginConfigOptions);
diff --git a/Engine/Source/Programs/Horde/Plugins/Build/HordeServer.Build.Tests/Perforce/PerforceSecretsTests.cs b/Engine/Source/Programs/Horde/Plugins/Build/HordeServer.Build.Tests/Perforce/PerforceSecretsTests.cs
--- a/Engine/Source/Programs/Horde/Plugins/Build/HordeServer.Build.Tests/Perforce/PerforceSecretsTests.cs
+++ b/Engine/Source/Programs/Horde/Plugins/Build/HordeServer.Build.Tests/Perforce/PerforceSecretsTests.cs
@@ -17,8 +17,8 @@ public class PerforceSecretsTests : BuildTestSetup
 	[TestInitialize]
 	public async Task SetupAsync()
 	{
-		GlobalConfig globalConfig = new();
-		globalConfig.Plugins.AddBuildConfig(new BuildConfig { ResolvePerforceCredentialsWithSecrets = true });
+		GlobalConfig globalConfig = new() { ResolveSecrets = true };
+		globalConfig.Plugins.AddBuildConfig(new BuildConfig());
 		globalConfig.Plugins.AddSecretsConfig(new SecretsConfig());
 		await SetConfigAsync(globalConfig);
 	}
diff --git a/Engine/Source/Programs/Horde/Plugins/Build/HordeServer.Build/BuildConfig.cs b/Engine/Source/Programs/Horde/Plugins/Build/HordeServer.Build/BuildConfig.cs
--- a/Engine/Source/Programs/Horde/Plugins/Build/HordeServer.Build/BuildConfig.cs
+++ b/Engine/Source/Programs/Horde/Plugins/Build/HordeServer.Build/BuildConfig.cs
@@ -12,7 +12,6 @@ using EpicGames.Horde.Common;
 using EpicGames.Horde.Jobs;
 using EpicGames.Horde.Jobs.Templates;
 using EpicGames.Horde.Projects;
-using EpicGames.Horde.Secrets;
 using EpicGames.Horde.Streams;
 using EpicGames.Perforce;
 using HordeServer.Acls;
@@ -20,13 +19,13 @@ using HordeServer.Agents;
 using HordeServer.Agents.Pools;
 using HordeServer.Artifacts;
 using HordeServer.Commits;
+using HordeServer.Configuration;
 using HordeServer.Devices;
 using HordeServer.VersionControl.Perforce;
 using HordeServer.Plugins;
 using HordeServer.Projects;
 using HordeServer.Streams;
 using HordeServer.Utilities;
-using Microsoft.Extensions.DependencyInjection;
 
 #pragma warning disable CA2227 // Change x to be read-only by removing the property setter
 
@@ -48,11 +47,6 @@ namespace HordeServer
 		/// </summary>
 		public List<PerforceCluster> PerforceClusters { get; set; } = new List<PerforceCluster>();
 
-		/// <summary>
-		/// Whether to enable resolving Perforce passwords and tickets against secrets
-		/// </summary>
-		public bool ResolvePerforceCredentialsWithSecrets { get; set; }
-
 		/// <summary>
 		/// Device configuration
 		/// </summary>
@@ -104,8 +98,6 @@ namespace HordeServer
 			Acl = configOptions.ParentAcl;
 			Streams = Projects.SelectMany(x => x.Streams).ToList();
 
-			UpdatePerforceCredentials(configOptions.Plugins, configOptions.ServiceProvider);
-
 			_projectLookup.Clear();
 			_streamLookup.Clear();
 			foreach (ProjectConfig project in Projects)
@@ -127,49 +119,6 @@ namespace HordeServer
 			UpdateWorkspacesForPools(configOptions.Plugins);
 		}
 
-		private static string? ResolveSecret(string? value, ISecretCollection secretCollection, SecretsConfig secretsConfig)
-		{
-			if (!String.IsNullOrEmpty(value))
-			{
-				ISecretProperty? newValue = secretCollection.ResolveAsync(value, secretsConfig).Result;
-				if (newValue != null)
-				{
-					return newValue.Value;
-				}
-			}
-			return value;
-		}
-
-		private void UpdatePerforceCredentials(IEnumerable<IPluginConfig> plugins, IServiceProvider? serviceProvider)
-		{
-			if (!ResolvePerforceCredentialsWithSecrets)
-			{
-				return;
-			}
-
-			SecretsConfig? secretsConfig = plugins.OfType<SecretsConfig>().FirstOrDefault();
-			if (secretsConfig == null)
-			{
-				return;
-			}
-
-			using IServiceScope? scope = serviceProvider?.CreateScope();
-			ISecretCollection? secretCollection = scope?.ServiceProvider.GetService<ISecretCollection>();
-			if (secretCollection == null)
-			{
-				return;
-			}
-
-			foreach (PerforceCluster perforceCluster in PerforceClusters)
-			{
-				foreach (PerforceCredentials credentials in perforceCluster.Credentials)
-				{
-					credentials.Password = ResolveSecret(credentials.Password, secretCollection, secretsConfig);
-					credentials.Ticket = ResolveSecret(credentials.Ticket, secretCollection, secretsConfig);
-				}
-			}
-		}
-
 		void UpdateWorkspacesForPools(IEnumerable<IPluginConfig> plugins)
 		{
 			// Try to get the compute config, and skip if it isn't configured
@@ -477,11 +426,13 @@ namespace HordeServer
 		/// <summary>
 		/// Password for the user
 		/// </summary>
+		[ResolveSecret]
 		public string? Password { get; set; } = String.Empty;
 
 		/// <summary>
 		/// Login ticket for the user (will be used instead of password if set)
 		/// </summary>
+		[ResolveSecret]
 		public string? Ticket { get; set; } = String.Empty;
 	}
 
diff --git a/Engine/Source/Programs/Horde/Plugins/Build/HordeServer.Build/BuildPlugin.cs b/Engine/Source/Programs/Horde/Plugins/Build/HordeServer.Build/BuildPlugin.cs
--- a/Engine/Source/Programs/Horde/Plugins/Build/HordeServer.Build/BuildPlugin.cs
+++ b/Engine/Source/Programs/Horde/Plugins/Build/HordeServer.Build/BuildPlugin.cs
@@ -45,7 +45,7 @@ namespace HordeServer
 	/// <summary>
 	/// Entry point for the build plugin
 	/// </summary>
-	[Plugin("Build", GlobalConfigType = typeof(BuildConfig), ServerConfigType = typeof(BuildServerConfig), DependsOn = ["Secrets"])]
+	[Plugin("Build", GlobalConfigType = typeof(BuildConfig), ServerConfigType = typeof(BuildServerConfig))]
 	public class BuildPlugin : IPluginStartup
 	{
 		readonly IServerInfo _serverInfo;
diff --git a/Engine/Source/Programs/Horde/Plugins/Secrets/HordeServer.Secrets.Tests/Secrets/SecretResolveTests.cs b/Engine/Source/Programs/Horde/Plugins/Secrets/HordeServer.Secrets.Tests/Secrets/SecretResolveTests.cs
--- a/Engine/Source/Programs/Horde/Plugins/Secrets/HordeServer.Secrets.Tests/Secrets/SecretResolveTests.cs
+++ b/Engine/Source/Programs/Horde/Plugins/Secrets/HordeServer.Secrets.Tests/Secrets/SecretResolveTests.cs
@@ -2,7 +2,6 @@
 
 using System.Text.Json;
 using EpicGames.Horde.Secrets;
-using HordeServer.Plugins;
 using HordeServer.Secrets;
 using HordeServer.Server;
 using HordeServer.Utilities;
@@ -14,7 +13,6 @@ namespace HordeServer.Tests.Secrets
 	public class SecretResolveTests
 	{
 		private static ISecretCollection s_secretCollection = default!;
-		private static GlobalConfig s_globalConfig = default!;
 
 		private class OptionsSnapshotStub<T>(T value) : IOptionsSnapshot<T> where T : class
 		{
@@ -41,7 +39,6 @@ namespace HordeServer.Tests.Secrets
 			GlobalConfig config = new();
 			config.Plugins.AddSecretsConfig(secretsConfig);
 			config.PostLoad(new ServerSettings(), [], []);
-			s_globalConfig = config;
 			s_secretCollection = secretCollection;
 		}
 
@@ -56,24 +53,6 @@ namespace HordeServer.Tests.Secrets
 			Assert.AreEqual(expectedValue, secret.Value);
 		}
 
-		[TestMethod]
-		[DataRow("horde:secret:mySecret.username", "username", "j.doe", DisplayName = "Get Username")]
-		[DataRow("horde:secret:mySecret.password", "password", "password123", DisplayName = "Get Password")]
-		public async Task ResolveSecretWithConfigAsync(string value, string expectedName, string expectedValue)
-		{
-			ISecretProperty? secret = await s_secretCollection.ResolveAsync(value, s_globalConfig.Plugins[new PluginName("Secrets")]);
-			Assert.IsNotNull(secret);
-			Assert.AreEqual(expectedName, secret.Name);
-			Assert.AreEqual(expectedValue, secret.Value);
-		}
-
-		[TestMethod]
-		[ExpectedException(typeof(InvalidCastException))]
-		public async Task ResolveSecretWithInvalidConfigAsync()
-		{
-			_ = await s_secretCollection.ResolveAsync("horde:secret:mySecret.password", new BuildConfig());
-		}
-
 		[TestMethod]
 		[DataRow("", DisplayName = "Empty string")]
 		[DataRow(null, DisplayName = "Null string")]
diff --git a/Engine/Source/Programs/Horde/Plugins/Secrets/HordeServer.Secrets.Tests/Secrets/SecretTests.cs b/Engine/Source/Programs/Horde/Plugins/Secrets/HordeServer.Secrets.Tests/Secrets/SecretTests.cs
new file mode 100644
--- /dev/null
+++ b/Engine/Source/Programs/Horde/Plugins/Secrets/HordeServer.Secrets.Tests/Secrets/SecretTests.cs
@@ -0,0 +1,66 @@
+// Copyright Epic Games, Inc. All Rights Reserved.
+
+using EpicGames.Horde.Secrets;
+using HordeServer.Plugins;
+using HordeServer.Secrets;
+using HordeServer.Server;
+using Microsoft.Extensions.Options;
+
+namespace HordeServer.Tests.Secrets
+{
+	[TestClass]
+	public class SecretTests
+	{
+		private static ISecretCollection s_secretCollection = null!;
+		private static GlobalConfig s_globalConfig = null!;
+
+		private class OptionsSnapshotStub<T>(T value) : IOptionsSnapshot<T> where T : class
+		{
+			public T Value { get; } = value;
+			public T Get(string? name)
+			{
+				return Value;
+			}
+		}
+
+		[ClassInitialize]
+		public static void CreateSecretCollection(TestContext _)
+		{
+			SecretsConfig secretsConfig = new()
+			{
+				Secrets = [new SecretConfig
+				{
+					Id = new SecretId("User"),
+					Data = new Dictionary<string, string>
+					{
+						{ "username", "j.doe" },
+						{ "password", "password123" }
+					}
+				}]
+			};
+			SecretCollection secretCollection = new(new SecretCollectionInternal([]), new OptionsSnapshotStub<SecretsConfig>(secretsConfig));
+			GlobalConfig config = new();
+			config.Plugins.AddSecretsConfig(secretsConfig);
+			config.PostLoad(new ServerSettings(), [], []);
+			s_globalConfig = config;
+			s_secretCollection = secretCollection;
+		}
+
+		[TestMethod]
+		[DataRow("username", "j.doe", DisplayName = "Get Username")]
+		[DataRow("password", "password123", DisplayName = "Get Password")]
+		public async Task GetSecretWithConfigAsync(string expectedName, string expectedValue)
+		{
+			ISecret? secret = await s_secretCollection.GetAsync(new SecretId("User"), s_globalConfig.Plugins[new PluginName("Secrets")]);
+			Assert.IsNotNull(secret);
+			Assert.AreEqual(expectedValue, secret.Data[expectedName]);
+		}
+
+		[TestMethod]
+		[ExpectedException(typeof(InvalidCastException))]
+		public async Task GetSecretWithInvalidConfigAsync()
+		{
+			_ = await s_secretCollection.GetAsync(new SecretId(), new BuildConfig());
+		}
+	}
+}
diff --git a/Engine/Source/Programs/Horde/Plugins/Secrets/HordeServer.Secrets/Secrets/SecretCollection.cs b/Engine/Source/Programs/Horde/Plugins/Secrets/HordeServer.Secrets/Secrets/SecretCollection.cs
--- a/Engine/Source/Programs/Horde/Plugins/Secrets/HordeServer.Secrets/Secrets/SecretCollection.cs
+++ b/Engine/Source/Programs/Horde/Plugins/Secrets/HordeServer.Secrets/Secrets/SecretCollection.cs
@@ -31,10 +31,10 @@ namespace HordeServer.Secrets
 		}
 
 		/// <inheritdoc/>
-		public async Task<ISecretProperty?> ResolveAsync(string value, object config, CancellationToken cancellationToken)
+		public async Task<ISecret?> GetAsync(SecretId secretId, object config, CancellationToken cancellationToken)
 		{
 			SecretsConfig secretsConfig = config as SecretsConfig ?? throw new InvalidCastException($"{nameof(config)} is not a SecretsConfig");
-			return await _secretCollectionInternal.ResolveAsync(value, secretsConfig, cancellationToken);
+			return await _secretCollectionInternal.GetAsync(secretId, secretsConfig, cancellationToken);
 		}
 
 		/// <inheritdoc/>
diff --git a/Engine/Source/Programs/Shared/EpicGames.Horde/Secrets/HttpSecretCollection.cs b/Engine/Source/Programs/Shared/EpicGames.Horde/Secrets/HttpSecretCollection.cs
--- a/Engine/Source/Programs/Shared/EpicGames.Horde/Secrets/HttpSecretCollection.cs
+++ b/Engine/Source/Programs/Shared/EpicGames.Horde/Secrets/HttpSecretCollection.cs
@@ -49,17 +49,17 @@ namespace EpicGames.Horde.Secrets
 		}
 
 		/// <inheritdoc/>
-		public async Task<ISecretProperty?> ResolveAsync(string value, CancellationToken cancellationToken = default)
+		public Task<ISecret?> GetAsync(SecretId secretId, object config, CancellationToken cancellationToken)
 		{
-			HordeHttpClient hordeHttpClient = _hordeClient.CreateHttpClient();
-			GetSecretPropertyResponse response = await hordeHttpClient.ResolveSecretAsync(value, cancellationToken);
-			return new SecretProperty(response);
+			throw new NotImplementedException();
 		}
 
 		/// <inheritdoc/>
-		public Task<ISecretProperty?> ResolveAsync(string value, object config, CancellationToken cancellationToken = default)
+		public async Task<ISecretProperty?> ResolveAsync(string value, CancellationToken cancellationToken = default)
 		{
-			throw new NotImplementedException();
+			HordeHttpClient hordeHttpClient = _hordeClient.CreateHttpClient();
+			GetSecretPropertyResponse response = await hordeHttpClient.ResolveSecretAsync(value, cancellationToken);
+			return new SecretProperty(response);
 		}
 	}
 }
diff --git a/Engine/Source/Programs/Shared/EpicGames.Horde/Secrets/ISecretCollection.cs b/Engine/Source/Programs/Shared/EpicGames.Horde/Secrets/ISecretCollection.cs
--- a/Engine/Source/Programs/Shared/EpicGames.Horde/Secrets/ISecretCollection.cs
+++ b/Engine/Source/Programs/Shared/EpicGames.Horde/Secrets/ISecretCollection.cs
@@ -18,20 +18,20 @@ namespace EpicGames.Horde.Secrets
 		Task<ISecret?> GetAsync(SecretId secretId, CancellationToken cancellationToken = default);
 
 		/// <summary>
-		/// Converts the string representation of a secret and property to a concrete value
-		/// The format of the string is 'horde:secret:secret-id.property-name'
+		/// Resolve a secret to concrete values
 		/// </summary>
-		/// <param name="value">A string that contains a secret ID and property name</param>
+		/// <param name="secretId">Identifier for the secret</param>
+		/// <param name="config">A SecretsConfig object</param>
 		/// <param name="cancellationToken">Cancellation token for the operation</param>
-		Task<ISecretProperty?> ResolveAsync(string value, CancellationToken cancellationToken = default);
+		/// <remarks>The method should only be used when the config cannot be retrieved from global config object, for example when it as not initialized during a config load</remarks>
+		Task<ISecret?> GetAsync(SecretId secretId, object config, CancellationToken cancellationToken = default);
 
 		/// <summary>
 		/// Converts the string representation of a secret and property to a concrete value
+		/// The format of the string is 'horde:secret:secret-id.property-name'
 		/// </summary>
 		/// <param name="value">A string that contains a secret ID and property name</param>
-		/// <param name="config">A SecretsConfig object</param>
 		/// <param name="cancellationToken">Cancellation token for the operation</param>
-		/// <remarks>The method should only be used when the config cannot be retrieved from global config object, for example when it as not initialized inside a IPluginConfig.PostLoadAsync method</remarks>
-		Task<ISecretProperty?> ResolveAsync(string value, object config, CancellationToken cancellationToken = default);
+		Task<ISecretProperty?> ResolveAsync(string value, CancellationToken cancellationToken = default);
 	}
 }
