Horde: JobExecutor tokens should not inherit exp claims for user generated jobs

diff --git a/Engine/Source/Programs/Horde/Plugins/Build/HordeServer.Build/Jobs/JobTaskSource.cs b/Engine/Source/Programs/Horde/Plugins/Build/HordeServer.Build/Jobs/JobTaskSource.cs
--- a/Engine/Source/Programs/Horde/Plugins/Build/HordeServer.Build/Jobs/JobTaskSource.cs
+++ b/Engine/Source/Programs/Horde/Plugins/Build/HordeServer.Build/Jobs/JobTaskSource.cs
@@ -895,7 +895,9 @@ namespace HordeServer.Jobs
 			claims.Add(new AclClaimConfig(HordeClaimTypes.LeaseStream, streamConfig.Id.ToString()));
 			claims.Add(new AclClaimConfig(HordeClaimTypes.LeaseProject, streamConfig.ProjectConfig.Id.ToString()));
 			claims.Add(new AclClaimConfig(HordeClaimTypes.LeaseTemplate, job.TemplateId.ToString()));
-			claims.AddRange(job.Claims.ConvertAll(x => new AclClaimConfig(x.Type, x.Value)));
+
+			// Copy job claims, except for expiration as we do not want that forwarded to the job executor token
+			claims.AddRange(job.Claims.Where(x => x.Type != "exp").ToList().ConvertAll(x => new AclClaimConfig(x.Type, x.Value)));
 
 			// Encode the payload
 			ExecuteJobTask task = new ExecuteJobTask();
