Horde: Updating schema documentation

#rnx

diff --git a/Engine/Source/Programs/Horde/Docs/Config/Schema/Globals.md b/Engine/Source/Programs/Horde/Docs/Config/Schema/Globals.md
--- a/Engine/Source/Programs/Horde/Docs/Config/Schema/Globals.md
+++ b/Engine/Source/Programs/Horde/Docs/Config/Schema/Globals.md
@@ -13,6 +13,7 @@ Name | Description
 `downtime` | [ScheduledDowntime](#scheduleddowntime)`[]`<br>List of scheduled downtime
 `plugins` | [GlobalPluginsConfig](#globalpluginsconfig)<br>Plugin config objects
 `parameters` | `object`<br>General parameters for other tools. Can be queried through the api/v1/parameters endpoint.
+`resolveSecrets` | `boolean`<br>Whether to enable resolving secrets in config properties
 `acl` | [AclConfig](#aclconfig)<br>Access control list
 
 ## ConfigInclude
@@ -83,7 +84,6 @@ Configuration for the build plugin
 Name | Description
 ---- | -----------
 `perforceClusters` | [PerforceCluster](#perforcecluster)`[]`<br>List of Perforce clusters
-`resolvePerforceCredentialsWithSecrets` | `boolean`<br>Whether to enable resolving Perforce passwords and tickets against secrets
 `devices` | [DeviceConfig](#deviceconfig)<br>Device configuration
 `maxConformCount` | `integer`<br>Maximum number of conforms to run at once
 `agentShutdownIfDisabledGracePeriod` | `string`<br>Time to wait before shutting down an agent that has been disabled Used if no value is set on the actual pool.
