Horde Dashboard: Resolve old query logic with new hook

diff --git a/Engine/Source/Programs/Horde/HordeDashboard/src/base/utilities/hooks.ts b/Engine/Source/Programs/Horde/HordeDashboard/src/base/utilities/hooks.ts
--- a/Engine/Source/Programs/Horde/HordeDashboard/src/base/utilities/hooks.ts
+++ b/Engine/Source/Programs/Horde/HordeDashboard/src/base/utilities/hooks.ts
@@ -38,6 +38,10 @@ export function useQuery() {
 		return searchParams.getAll(key);
 	}
 
+	const getCopy = (): URLSearchParams => {
+		return new URLSearchParams(searchParams);
+	}
+
 	const set = (key: string, value: string) => {
 		const newParams = new URLSearchParams(searchParams);
 		newParams.set(key, value);
@@ -74,5 +78,5 @@ export function useQuery() {
 		}
 	}
 
-	return {get, getAll, set, has, clear, setCompressed, getCompressed, raw: {searchParams, setSearchParams}}
+	return {get, getAll, getCopy, set, has, clear, setCompressed, getCompressed, raw: {searchParams, setSearchParams}}
 }
\ No newline at end of file
diff --git a/Engine/Source/Programs/Horde/HordeDashboard/src/components/LogView.tsx b/Engine/Source/Programs/Horde/HordeDashboard/src/components/LogView.tsx
--- a/Engine/Source/Programs/Horde/HordeDashboard/src/components/LogView.tsx
+++ b/Engine/Source/Programs/Horde/HordeDashboard/src/components/LogView.tsx
@@ -875,11 +875,11 @@ const LogList: React.FC<{ logId: string }> = observer(({ logId }) => {
 
             const issueId = item!.issueId!.toString();
 
-            query.set("issue", issueId);
-            query.toString();
+            const newQuery = query.getCopy();
+            newQuery.set("issue", issueId);
 
             // this href is for when log is copy/pasted, note that we need the onClick handler when on site, so don't reload page
-            const href = `${location.pathname}?${query.toString()}`;
+            const href = `${location.pathname}?${newQuery.toString()}`;
 
             const fontStyle = {
                fontFamily: "Horde Open Sans SemiBold, sans-serif, sans-serif", color: error ? errorButtonTextColor: warningButtonTextColor, fontSize: handler.fontSize, textDecoration: !item!.issue?.resolvedAt ? undefined : "line-through"
diff --git a/Engine/Source/Programs/Horde/HordeDashboard/src/components/StreamSummary.tsx b/Engine/Source/Programs/Horde/HordeDashboard/src/components/StreamSummary.tsx
--- a/Engine/Source/Programs/Horde/HordeDashboard/src/components/StreamSummary.tsx
+++ b/Engine/Source/Programs/Horde/HordeDashboard/src/components/StreamSummary.tsx
@@ -651,8 +651,9 @@ const HealthPanelIssues: React.FC<{ desktopAlerts?: boolean }> = observer(({ des
       if (props) {
 
          const item = props!.item as HealthItem;
-         query.set("issue", item.issue.id.toString());
-         const href = `${location.pathname}?${query.toString()}`;
+         const newQuery = query.getCopy();
+         newQuery.set("issue", item.issue.id.toString());
+         const href = `${location.pathname}?${newQuery.toString()}`;
 
          return <Link to={href} className="job-item" onClick={() => setIssueHistory(true)}>
             <DetailsRow className={classes.detailsRow} {...props} />
diff --git a/Engine/Source/Programs/Horde/HordeDashboard/src/components/UserHome.tsx b/Engine/Source/Programs/Horde/HordeDashboard/src/components/UserHome.tsx
--- a/Engine/Source/Programs/Horde/HordeDashboard/src/components/UserHome.tsx
+++ b/Engine/Source/Programs/Horde/HordeDashboard/src/components/UserHome.tsx
@@ -301,8 +301,9 @@ const HealthPanel: React.FC = observer(() => {
       if (props) {
 
          const item = props!.item as HealthItem;
-         query.set("issue", item.issue.id.toString());
-         const href = `${location.pathname}?${query.toString()}`;
+         const newQuery = query.getCopy();
+         newQuery.set("issue", item.issue.id.toString());
+         const href = `${location.pathname}?${newQuery.toString()}`;
 
          return <Link to={href} className="job-item" onClick={() => { setIssueHistory(true) }}>
             <DetailsRow className={classes.detailsRow} {...props} />
diff --git a/Engine/Source/Programs/Horde/HordeDashboard/src/components/jobDetailsV2/JobDetailHealthV2.tsx b/Engine/Source/Programs/Horde/HordeDashboard/src/components/jobDetailsV2/JobDetailHealthV2.tsx
--- a/Engine/Source/Programs/Horde/HordeDashboard/src/components/jobDetailsV2/JobDetailHealthV2.tsx
+++ b/Engine/Source/Programs/Horde/HordeDashboard/src/components/jobDetailsV2/JobDetailHealthV2.tsx
@@ -251,9 +251,9 @@ export const HealthPanel: React.FC<{ jobDetails: JobDetailsV2 }> = observer(({ j
             summary = summary.slice(0, 100) + "...";
          }
 
-
-         query.set("issue", item.issue.id.toString());
-         const href = `${location.pathname}?${query.toString()}` + window.location.hash;
+         const newQuery = query.getCopy();
+         newQuery.set("issue", item.issue.id.toString());
+         const href = `${location.pathname}?${newQuery.toString()}` + window.location.hash;
 
          return <Link to={href} className="job-item" onClick={() => setIssueHistory(true)}><Stack horizontal disableShrink={true}>{<IssueStatusIcon issue={issue} />}<Text variant={textSize} style={{ textDecoration: !!issue.resolvedAt ? "line-through" : undefined }}>{`Issue ${issue.id} - ${summary}`}</Text></Stack></Link>;
       }
