Horde: Copy StreamConfig tags to environment variables

The name of environment variable will be the tag's name in uppercase and prefixed with "UE_HORDE_STREAMTAG_".
The value of environment variable will be either the string "True" or "False".

diff --git a/Engine/Source/Programs/Horde/Plugins/Build/HordeServer.Build/Jobs/JobRpcService.cs b/Engine/Source/Programs/Horde/Plugins/Build/HordeServer.Build/Jobs/JobRpcService.cs
--- a/Engine/Source/Programs/Horde/Plugins/Build/HordeServer.Build/Jobs/JobRpcService.cs
+++ b/Engine/Source/Programs/Horde/Plugins/Build/HordeServer.Build/Jobs/JobRpcService.cs
@@ -405,6 +405,11 @@ namespace HordeServer.Jobs
 				response.EnvVars.Add("UE_HORDE_TEMPLATENAME", templateName);
 				response.EnvVars.Add("UE_HORDE_STEPNAME", node.Name);
 
+				foreach (StreamTag tag in streamConfig.StreamTags)
+				{
+					response.EnvVars.Add($"UE_HORDE_STREAMTAG_{tag.Name.ToUpperInvariant()}", tag.Enabled.ToString());
+				}
+
 				if (job.Environment != null)
 				{
 					response.EnvVars.Add(job.Environment);
