Horde Dashboard: Adjust advanced additional arg parsing

diff --git a/Engine/Source/Programs/Horde/HordeDashboard/src/components/build/NewBuildV2.tsx b/Engine/Source/Programs/Horde/HordeDashboard/src/components/build/NewBuildV2.tsx
--- a/Engine/Source/Programs/Horde/HordeDashboard/src/components/build/NewBuildV2.tsx
+++ b/Engine/Source/Programs/Horde/HordeDashboard/src/components/build/NewBuildV2.tsx
@@ -66,6 +66,33 @@ class BuildOptions {
 
    }
 
+   getQuotedArg(arg: string){
+      // Ignore args that don't follow "-set:XXX=YYY"
+      // Also prevents overquoting incorrectly recorded args 
+      // i.e., forgetting "" in -set:XXX=YYY ZZZ AAA BBB => -set:XXX="YYY" "ZZZ" "AAA" "BBB"
+      if(!arg.startsWith("-set")) return arg;
+      
+      // Split argument on space to extract left side assignment
+      const splitOnSpace = arg.split(" ");
+      
+      // Ignore args without spaces
+      if(splitOnSpace.length === 1) return arg;
+
+      // Split assignment on equals, getting the last = sign. Which equals sign doesn't necessarily matter.
+      // It looks best/is more semantic to do the inner most one, though.
+      const splitOnEquals = splitOnSpace[0].split("=");
+
+      // Add double quote just before value
+      splitOnEquals[splitOnEquals.length - 1] = `"${splitOnEquals[splitOnEquals.length - 1]}`;
+      
+      // Replace splitOnSpace[0], Restore string, and add " to end
+      splitOnSpace[0] = splitOnEquals.join("=");
+      let finalString = splitOnSpace.join(" ") + '"';
+      
+      // -set:XXX=YYY=AAA BBB CCC --> -set:XXX=YYY="AAA BBB CCC"
+      return finalString;
+   }
+
    init() {
 
       this.change = undefined;
@@ -128,7 +155,7 @@ class BuildOptions {
 
       this.stream = stream;
 
-      this.advAdditionalArgs = jobDetails?.jobData?.additionalArguments?.join(" ");
+      this.advAdditionalArgs = jobDetails?.jobData?.additionalArguments?.map(this.getQuotedArg).join(" ");
 
       this.load();
 
@@ -1857,17 +1884,7 @@ const AdvancedPanel: React.FC = observer(() => {
             </Stack>
             {showAdditionalArgs && <TextField key={`key_adv_add_args_${options.currentRenderKey}`} multiline style={{ height: 48 }} resizable={false} readOnly={options.readOnly} spellCheck={false} defaultValue={options.advAdditionalArgs} label="Additional Arguments" onChange={(ev, newValue) => options.advAdditionalArgs = newValue} />}
             {showArgumentClipboardButton && <DefaultButton text="Copy Job Arguments to Clipboard" style={{ width: 240 }} onClick={() => copyToClipboard(
-               options.jobDetails?.jobData?.arguments.map(arg => {
-                  if (arg.indexOf("=") !== -1) {
-                     const components = arg.split("=");
-                     if (components[1].indexOf(" ") === -1) {
-                        return arg;
-                     }
-                     return `${components[0]}="${components[1]}"`;
-                  } else {
-                     return arg;
-                  }
-               }).join(" ")
+               options.jobDetails?.jobData?.arguments.map(options.getQuotedArg).join(" ")
             )} />}
             {!!farguments && <TextField style={{ height: fargumentsHeight }} key={`key_adv_job_arguments_${options.currentRenderKey}`} defaultValue={farguments} readOnly={true} label="Job Arguments" multiline resizable={false} />}
             {!!fparameters && <TextField style={{ height: fparametersHeight }} key={`key_adv_job_parameters_${options.currentRenderKey}`} defaultValue={fparameters} readOnly={true} label="Job Parameters" multiline resizable={false} />}
@@ -1947,8 +1964,16 @@ const BuildModal: React.FC<{ setUseLegacyDialog: (value: boolean) => void }> = o
 
       let additionalArgs: string[] = [];
       if (options!.advAdditionalArgs) {
-         const argRegex = /"(\\"|[^"])*?"|[^ ]+/g;
-         options.advAdditionalArgs.trim().replaceAll("\n", " ").match(argRegex)?.forEach(arg => additionalArgs.push(arg.replace(/"|\\"|'|\\'|&quot;/g, "")));
+         // Match either:
+         // 1. A key-value pair where the value is a quoted string that may contain escaped characters
+         //    - Starts with one or more non-whitespace chars followed by '='
+         //    - Then a quoted string starting and ending with "
+         //    - Inside the quotes, allow any characyer except " and \
+         //    - However, allow any escaped character
+         // 2. A standalone non-whitespace character sequence
+         const argRegex = /\S+="(?:[^"\\]|\\.)*"|\S+/g
+         // Match input with above regex, replacing any non-escaped quotes with nothing and escaped quotes with just a quote.
+         options.advAdditionalArgs.trim().replaceAll("\n", " ").match(argRegex)?.forEach(arg => additionalArgs.push(arg.replace(/(?<!\\)"/g, '').replace(/\\"/g, '"')));
       }
 
       // trim all parameters
