Horde: Hydrate CachedCommitDoc objects created in fallback path

diff --git a/Engine/Source/Programs/Horde/Plugins/Build/HordeServer.Build/VersionControl/Perforce/PerforceServiceCache.cs b/Engine/Source/Programs/Horde/Plugins/Build/HordeServer.Build/VersionControl/Perforce/PerforceServiceCache.cs
--- a/Engine/Source/Programs/Horde/Plugins/Build/HordeServer.Build/VersionControl/Perforce/PerforceServiceCache.cs
+++ b/Engine/Source/Programs/Horde/Plugins/Build/HordeServer.Build/VersionControl/Perforce/PerforceServiceCache.cs
@@ -714,6 +714,7 @@ namespace HordeServer.VersionControl.Perforce
 							await foreach (ICommit commit in base.FindAsync(minChange, true, maxChange, true, maxUnfilteredResults, null, cancellationToken))
 							{
 								CachedCommitDoc cachedCommit = await CachedCommitDoc.FromCommitAsync(commit, cancellationToken);
+								cachedCommit.PostLoad(_owner, StreamConfig);
 								await _owner.AddCachedCommitAsync(cachedCommit, cancellationToken);
 								await _owner._mongoService.UpdateSingletonAsync<CacheState>(x => TryUpdateRange(x, StreamConfig, commit.Id.GetPerforceChange(), maxChange), cancellationToken);
 								_owner._logger.LogDebug("Adding new cached commit for {StreamId} at change {Change}", StreamConfig.Id, commit.Id);
