Add visual cue that outcome filter is being used

#rnx

diff --git a/Engine/Source/Programs/Horde/HordeDashboard/plugins/test-automation/testAutomationCommon.tsx b/Engine/Source/Programs/Horde/HordeDashboard/plugins/test-automation/testAutomationCommon.tsx
--- a/Engine/Source/Programs/Horde/HordeDashboard/plugins/test-automation/testAutomationCommon.tsx
+++ b/Engine/Source/Programs/Horde/HordeDashboard/plugins/test-automation/testAutomationCommon.tsx
@@ -26,6 +26,9 @@ export const styles = {
     defaultFilter: {
         backgroundColor: "#0078d4"
     },
+    highlightFilter: {
+        backgroundColor: "#2a95e8ff"
+    },
     defaultCheckbox: {
         root: {
             height: 15
@@ -539,6 +542,8 @@ export const CheckListOption: React.FC<{options: ICheckListOption[], onChange: (
         setUpdate(1);
     }, [options]);
 
+    const isFiltering : boolean = useMemo(() => items.slice(1).some(i => !i.checked), [items, update]);
+
     if (!items.length && !update) return;
 
     const fitlerOptions: IContextualMenuProps = {
@@ -574,15 +579,19 @@ export const CheckListOption: React.FC<{options: ICheckListOption[], onChange: (
                 </Stack>
     }
 
+    const iconName : string = isFiltering? "ReportWarning" : "PageListFilter";
+    const iconColor : string = isFiltering? styles.highlightFilter.backgroundColor : "";
+
     return  <IconButton
-                iconProps={{iconName: "PageListFilter"}}
+                iconProps={{iconName: iconName}}
                 menuProps={fitlerOptions}
                 style={{
                     fontSize:16, height: 28,
                     border: '1px solid', borderColor: dashboard.darktheme ? "#4D4C4B" : "#6D6C6B", borderRadius: 4,
-                    backgroundColor: dashboard.darktheme ? "rgba(255, 255, 255, 0.1)" : "rgba(0, 0, 0, 0.3)",
-                    color: 'white'
+                    backgroundColor: dashboard.darktheme ? "rgba(255, 255, 255, 0.1)" : "",
+                    color: dashboard.darktheme ? 'white' : 'grey'
                 }}
+                styles={{icon: {color: iconColor}}}
             />
 }
 
