Add all device logs to phase artifacts

#rnx

diff --git a/Engine/Source/Programs/Horde/HordeDashboard/plugins/test-automation/testAutomationEvents.tsx b/Engine/Source/Programs/Horde/HordeDashboard/plugins/test-automation/testAutomationEvents.tsx
--- a/Engine/Source/Programs/Horde/HordeDashboard/plugins/test-automation/testAutomationEvents.tsx
+++ b/Engine/Source/Programs/Horde/HordeDashboard/plugins/test-automation/testAutomationEvents.tsx
@@ -191,10 +191,11 @@ export const PhaseEventsPane: React.FC<{phase: TestPhaseStatus, onEventsLoaded?:
     const collectPhaseArtifacts = useCallback((): string[] | undefined => {
         if (!events) return;
         // collect artifact references
-        const artifacts: string[] = [];
-        // phase log
-        const logPath = phase.devices.length && phase.devices[0].logPath;
-        if (logPath) artifacts.push(logPath);
+        const artifacts: Set<string> = new Set();
+        // phase logs
+        phase.devices.forEach(device => {
+            if (device.logPath) artifacts.add(device.logPath);
+        });
         // parse event properties
         events.forEach(entry => {
             if (!!entry.properties) {
@@ -203,31 +204,31 @@ export const PhaseEventsPane: React.FC<{phase: TestPhaseStatus, onEventsLoaded?:
                     switch (property["$type"]) {
                         case LogPropertyType.ImageCompare:
                             const imageCompare = property as ImageCompare;
-                            if (imageCompare.approved) artifacts.push(imageCompare.approved);
-                            if (imageCompare.approved_metadata) artifacts.push(imageCompare.approved_metadata);
-                            if (imageCompare.unapproved) artifacts.push(imageCompare.unapproved);
-                            if (imageCompare.unapproved_metadata) artifacts.push(imageCompare.unapproved_metadata);
-                            if (imageCompare.difference) artifacts.push(imageCompare.difference);
-                            if (imageCompare.difference_report) artifacts.push(imageCompare.difference_report);
+                            if (imageCompare.approved) artifacts.add(imageCompare.approved);
+                            if (imageCompare.approved_metadata) artifacts.add(imageCompare.approved_metadata);
+                            if (imageCompare.unapproved) artifacts.add(imageCompare.unapproved);
+                            if (imageCompare.unapproved_metadata) artifacts.add(imageCompare.unapproved_metadata);
+                            if (imageCompare.difference) artifacts.add(imageCompare.difference);
+                            if (imageCompare.difference_report) artifacts.add(imageCompare.difference_report);
                             break;
 
                         case LogPropertyType.EmbeddedHTML:
                             const embeddedHTML = property as EmbeddedHTML;
-                            if (embeddedHTML.path) artifacts.push(embeddedHTML.path);
+                            if (embeddedHTML.path) artifacts.add(embeddedHTML.path);
                             break;
 
                         case LogPropertyType.Artifacts:
                             const artifactsProp = property as Artifacts
                             Object.entries(artifactsProp).forEach(([key, value]) => {
                                 if (!value || key.startsWith('$')) return;
-                                artifacts.push(value);
+                                artifacts.add(value);
                             });
                             break;
                     }
                 });
             }
         });
-        return artifacts;
+        return [...artifacts];
     }, [events]);
 
     const isEntryNeedDisplay = useCallback((event: EventEntry): boolean => {
diff --git a/Engine/Source/Programs/Horde/HordeDashboard/plugins/test-automation/testDataV2Fetcher.ts b/Engine/Source/Programs/Horde/HordeDashboard/plugins/test-automation/testDataV2Fetcher.ts
--- a/Engine/Source/Programs/Horde/HordeDashboard/plugins/test-automation/testDataV2Fetcher.ts
+++ b/Engine/Source/Programs/Horde/HordeDashboard/plugins/test-automation/testDataV2Fetcher.ts
@@ -195,7 +195,7 @@ export class TestDataV2 implements TestVersionFetcher {
         const start = new Date(testDetails.summary.dateTime);
         const result = new TestSessionDetails(testData.key, testDetails.summary.testName, start, testDetails.summary.timeElapseSec, testData.commitId.name, testData.commitId.order, testData.streamId, testData.jobId, testData.stepId);
         Object.entries(testDetails.devices).forEach(([key, d]) => {
-            result.devices.set(key, new DeviceRef(d.name, decodeURI(d.appInstanceLogPath), d.metadata, this.version));
+            result.devices.set(key, new DeviceRef(d.name, decodeURI(d.appInstanceLogPath ?? ""), d.metadata, this.version));
         });
         result.phases.push(...testDetails.phases.map(p => this.createPhaseStatus(p, result.devices, result)));
         result.phasesSucceededCount = testDetails.summary.phasesSucceededCount;
