Assume artifact path is URI encoded in data model

#rnx

diff --git a/Engine/Source/Programs/Horde/HordeDashboard/plugins/test-automation/testData.ts b/Engine/Source/Programs/Horde/HordeDashboard/plugins/test-automation/testData.ts
--- a/Engine/Source/Programs/Horde/HordeDashboard/plugins/test-automation/testData.ts
+++ b/Engine/Source/Programs/Horde/HordeDashboard/plugins/test-automation/testData.ts
@@ -1528,7 +1528,7 @@ export class ArtifactFactory {
 
         await this.getJobStepSavedArtifacts();
         
-        return `${backend.serverUrl}/api/v2/artifacts/${this.artifactId}/file?path=${encodeURIComponent(referencePath)}&inline=true`;
+        return `${backend.serverUrl}/api/v2/artifacts/${this.artifactId}/file?path=${referencePath}&inline=true`;
     }
 
     async fetch(referencePath: string): Promise<any> {
@@ -1554,12 +1554,12 @@ export class ArtifactFactory {
 
     async download(referencePath: string): Promise<void> {
         await this.getJobStepSavedArtifacts();
-        if (!!this.artifactId) backend.downloadArtifactV2(this.artifactId, referencePath);
+        if (!!this.artifactId) backend.downloadArtifactV2(this.artifactId, decodeURIComponent(referencePath));
     }
 
     async downloadZip(referencePaths: string[]) : Promise<void> {
         await this.getJobStepSavedArtifacts();
-        if (!!this.artifactId) backend.downloadArtifactZipV2(this.artifactId, {filter: referencePaths})
+        if (!!this.artifactId) backend.downloadArtifactZipV2(this.artifactId, {filter: referencePaths.map(r => decodeURIComponent(r))})
     }
 
     private jobId: string;
diff --git a/Engine/Source/Programs/Horde/HordeDashboard/plugins/test-automation/testDataV2Fetcher.ts b/Engine/Source/Programs/Horde/HordeDashboard/plugins/test-automation/testDataV2Fetcher.ts
--- a/Engine/Source/Programs/Horde/HordeDashboard/plugins/test-automation/testDataV2Fetcher.ts
+++ b/Engine/Source/Programs/Horde/HordeDashboard/plugins/test-automation/testDataV2Fetcher.ts
@@ -8,7 +8,7 @@ class TestPhaseStatusV2 extends TestPhaseStatus {
 
     constructor(phase: PhaseSessionDetailsResponse, devices: DeviceRef[], parent: TestSessionDetails) {
         super(phase.key, phase.name, new Date(phase.dateTime), phase.timeElapseSec, devices, phase.outcome, parent)
-        if (phase.eventStreamPath) this._eventStreamPath = decodeURI(phase.eventStreamPath);
+        this._eventStreamPath = phase.eventStreamPath;
         this.hasWarning = phase.hasWarning;
     }
 
@@ -195,7 +195,7 @@ export class TestDataV2 implements TestVersionFetcher {
         const start = new Date(testDetails.summary.dateTime);
         const result = new TestSessionDetails(testData.key, testDetails.summary.testName, start, testDetails.summary.timeElapseSec, testData.commitId.name, testData.commitId.order, testData.streamId, testData.jobId, testData.stepId);
         Object.entries(testDetails.devices).forEach(([key, d]) => {
-            result.devices.set(key, new DeviceRef(d.name, decodeURI(d.appInstanceLogPath ?? ""), d.metadata, this.version));
+            result.devices.set(key, new DeviceRef(d.name, d.appInstanceLogPath, d.metadata, this.version));
         });
         result.phases.push(...testDetails.phases.map(p => this.createPhaseStatus(p, result.devices, result)));
         result.phasesSucceededCount = testDetails.summary.phasesSucceededCount;
