Horde: Escape redirect URI in OAuth login form for Horde accounts

diff --git a/Engine/Source/Programs/Horde/HordeServer/Accounts/OAuthController.cs b/Engine/Source/Programs/Horde/HordeServer/Accounts/OAuthController.cs
--- a/Engine/Source/Programs/Horde/HordeServer/Accounts/OAuthController.cs
+++ b/Engine/Source/Programs/Horde/HordeServer/Accounts/OAuthController.cs
@@ -114,7 +114,28 @@ namespace HordeServer.Accounts
 				FormPostUrl = $"/api/v1/oauth2/login{Request.QueryString}"
 			});
 		}
-
+		
+		private static bool TryParseUrl(string candidate, out Uri uri)
+		{
+			uri = null!;
+			if (!Uri.TryCreate(candidate, UriKind.Absolute, out Uri? parsed))
+			{
+				return false;
+			}
+			
+			// Disallow fragment and userinfo per OAuth spec
+			bool hasFragment = !String.IsNullOrEmpty(parsed.Fragment);
+			bool hasUserInfo = !String.IsNullOrEmpty(parsed.UserInfo);
+			
+			if (hasFragment || hasUserInfo)
+			{
+				return false;
+			}
+			
+			uri = parsed;
+			return true;
+		}
+		
 		/// <summary>
 		/// Post interactive login credentials to the server in exchange for an authorization token
 		/// </summary>
@@ -133,7 +154,10 @@ namespace HordeServer.Accounts
 				return BadRequest("Missing username from form post");
 			}
 
-			// TODO: Validate the redirect URI
+			if (!TryParseUrl(request.RedirectUri, out Uri _))
+			{
+				return BadRequest("Invalid redirect URI");
+			}
 
 			// Check the login credentials and update the session key
 			IAccount? account = null;
@@ -182,7 +206,7 @@ namespace HordeServer.Accounts
 				StringBuilder html = new StringBuilder();
 				html.AppendLine($"<html>");
 				html.AppendLine($"  <body onload=\"javascript:document.forms[0].submit()\">");
-				html.AppendLine($"    <form method=\"post\" action=\"{request.RedirectUri}\">");
+				html.AppendLine($"    <form method=\"post\" action=\"{HttpUtility.HtmlAttributeEncode(request.RedirectUri)}\">");
 				foreach ((string key, string value) in responseFields)
 				{
 					html.AppendLine($"      <input type=\"hidden\" name=\"{HttpUtility.HtmlAttributeEncode(key)}\" value=\"{HttpUtility.HtmlAttributeEncode(value)}\"/>");
@@ -549,7 +573,10 @@ namespace HordeServer.Accounts
 		{
 			if (Settings.Value.AuthMethod != AuthMethod.Horde)
 			{
-				context.Result = new NotFoundResult();
+				context.Result = new BadRequestObjectResult(new
+				{
+					message = "This endpoint is only available when authentication method is set to 'Horde'."
+				});
 			}
 		}
 	}
