Horde: Log an error when UAT terminates with a non-zero exit code

The JobExecutor in the agent detects and flags job steps as failed when Unreal Automation Tool (UAT) terminates with a non-zero exit code. However, when inspecting job step logs, this failure is not immediately apparent because no explicit error messages were logged previously.

This change adds an error log entry to make UAT failures more visible in the logs.

diff --git a/Engine/Source/Programs/Horde/Drivers/JobDriver/Execution/JobExecutor.cs b/Engine/Source/Programs/Horde/Drivers/JobDriver/Execution/JobExecutor.cs
--- a/Engine/Source/Programs/Horde/Drivers/JobDriver/Execution/JobExecutor.cs
+++ b/Engine/Source/Programs/Horde/Drivers/JobDriver/Execution/JobExecutor.cs
@@ -806,9 +806,10 @@ namespace JobDriver.Execution
 				arguments.AppendArgument("-Schema=", preprocessedSchemaFile.FullName);
 			}
 
-			int result = await ExecuteAutomationToolAsync(step, workspaceDir, arguments.ToString(), useP4, logger, cancellationToken);
-			if (result != 0)
+			int exitCode = await ExecuteAutomationToolAsync(step, workspaceDir, arguments.ToString(), useP4, logger, cancellationToken);
+			if (exitCode != 0)
 			{
+				logger.LogError("BuildGraph setup terminated with non-zero exit code: {ExitCode}", exitCode);
 				return false;
 			}
 
@@ -1252,8 +1253,10 @@ namespace JobDriver.Execution
 			}
 
 			// Run UAT
-			if (await ExecuteAutomationToolAsync(step, workspaceDir, arguments, useP4, logger, cancellationToken) != 0)
+			int exitCode = await ExecuteAutomationToolAsync(step, workspaceDir, arguments, useP4, logger, cancellationToken);
+			if (exitCode != 0)
 			{
+				logger.LogError("AutomationTool (UAT) terminated with non-zero exit code: {ExitCode}", exitCode);
 				return false;
 			}
 
