Horde: Update release notes

#rnx

diff --git a/Engine/Source/Programs/Horde/Docs/ReleaseNotes.md b/Engine/Source/Programs/Horde/Docs/ReleaseNotes.md
--- a/Engine/Source/Programs/Horde/Docs/ReleaseNotes.md
+++ b/Engine/Source/Programs/Horde/Docs/ReleaseNotes.md
@@ -1,5 +1,18 @@
 # Release Notes
 
+## 2025-08-07
+
+* Log an error when UAT terminates with a non-zero exit code
+  The JobExecutor in the agent detects and flags job steps as failed when Unreal Automation Tool (UAT) terminates with a non-zero exit code. However, when inspecting job step logs, this failure is not immediately apparent because no explicit error messages were logged previously.
+  This change adds an error log entry to make UAT failures more visible in the logs. (44661300)
+* Escape redirect URI in OAuth login form for Horde accounts (44659453)
+* Hydrate CachedCommitDoc objects created in fallback path (44597929)
+* Horde Dashboard: Adjust advanced additional arg parsing (44597533)
+* Copy StreamConfig tags to environment variables
+  The name of environment variable will be the tag's name in uppercase and prefixed with "UE_HORDE_STREAMTAG_". 
+  The value of environment variable will be either the string "True" or "False". (44554523)
+* Horde Dashboard: Resolve old query logic with new hook (44481989)
+
 ## 2025-07-24
 
 * Horde Dashboard: Remove PrintException Component (44281453)
