[Horde] Adds configurable device tags
Device tags are small bits of metadata that can be attached to each device
When making a reservation, users can specify tags in a variety of ways

Required Tags - Devices that do not have this tag will not be reserved
Preferred Tags - Devices that have this tag will be selected over devices that do not, provided at least one is available
Undesired Tags - Devices that have this tag will be de-prioritized over devices that do not
Blocked Tags - Devices that have this tag will not be reserved

#tests horde device unit tests

diff --git a/Engine/Source/Programs/Horde/Plugins/Build/HordeServer.Build.Tests/Devices/DeviceServiceTests.cs b/Engine/Source/Programs/Horde/Plugins/Build/HordeServer.Build.Tests/Devices/DeviceServiceTests.cs
--- a/Engine/Source/Programs/Horde/Plugins/Build/HordeServer.Build.Tests/Devices/DeviceServiceTests.cs
+++ b/Engine/Source/Programs/Horde/Plugins/Build/HordeServer.Build.Tests/Devices/DeviceServiceTests.cs
@@ -240,7 +240,7 @@ namespace HordeServer.Tests.Devices
 			return job;
 		}
 
-		static LegacyCreateReservationRequest SetupReservationTestAsync(IJob job, string poolId = "TestDevicePool1", string deviceType = "TestDevicePlatform1", JobStepId? stepId = null, string? modelId = null, string? deviceName = null)
+		static LegacyCreateReservationRequest SetupReservationTestAsync(IJob job, string poolId = "TestDevicePool1", string deviceType = "TestDevicePlatform1", JobStepId? stepId = null, string? modelId = null, string? deviceName = null, DeviceTagsRequest? tags = null)
 		{
 			if (stepId == null)
 			{
@@ -261,6 +261,7 @@ namespace HordeServer.Tests.Devices
 			request.JobId = job.Id.ToString();
 			request.StepId = stepId.ToString();
 			request.DeviceName = deviceName;
+			request.Tags = tags;
 
 			return request;
 		}
@@ -584,6 +585,234 @@ namespace HordeServer.Tests.Devices
 			Assert.AreEqual(device.Name, "testdevice1_platform1_testdevicepool1", StringComparer.OrdinalIgnoreCase);
 		}
 
+		[TestMethod]
+		public async Task TestReservationDeviceTagExpectedErrorsAsync()
+		{
+			await CreateFixtureAsync();
+
+			DeviceConfig devices = new DeviceConfig();
+
+			// Create a pool
+			string poolName = "TestDevicePool";
+			DevicePoolConfig devicePool = new DevicePoolConfig() { Id = new DevicePoolId(poolName), Name = poolName, PoolType = DevicePoolType.Automation, ProjectIds = new List<ProjectId>() { new ProjectId("ue5") } };
+			devices.Pools.Add(devicePool);
+
+			// Create a single platform
+			string platformName = "TestDevicePlatform";
+			DevicePlatformConfig platform = new DevicePlatformConfig() { Id = new DevicePlatformId(platformName), Name = platformName };
+			devices.Platforms.Add(platform);
+
+			await UpdateConfigAsync(x => x.Plugins.GetBuildConfig().Devices = devices);
+
+			// Create a device with a single blocked tag
+			CreateDeviceRequest testDevice = new CreateDeviceRequest()
+			{
+				Name = "TestDevice",
+				PoolId = poolName,
+				Address = "10.0.0.1",
+				PlatformId = platformName,
+				Tags = new List<string> { "BlockedTag" }
+			};
+
+			await DeviceController.CreateDeviceAsync(testDevice);
+			await DeviceService.TickForTestingAsync();
+
+			IJob job = await SetupJobAsync();
+
+			// Now test the three expected failure conditions
+
+			// Bad request - should error because of duplicates in tag request
+			DeviceTagsRequest duplicateTags = new DeviceTagsRequest() { RequiredTags = ["DuplicateEntry"], BlockedTags = ["DuplicateEntry"] };
+			LegacyCreateReservationRequest duplicateRequest = SetupReservationTestAsync(job, poolName, platformName, null, null, null, duplicateTags);
+
+			// Missing required - should error because no devices exist with the required tag
+			DeviceTagsRequest requiredTag = new DeviceTagsRequest() { RequiredTags = ["RequiredTag"] };
+			LegacyCreateReservationRequest requiredRequest = SetupReservationTestAsync(job, poolName, platformName, null, null, null, requiredTag);
+
+			// Only blocked - should error because no devices exist without the blocked tag
+			DeviceTagsRequest blockedTag = new DeviceTagsRequest() { BlockedTags = ["BlockedTag"] };
+			LegacyCreateReservationRequest blockedRequest = SetupReservationTestAsync(job, poolName, platformName, null, null, null, blockedTag);
+
+			ActionResult? duplicateResponse = (await DeviceController!.CreateDeviceReservationV1Async(duplicateRequest)).Result;
+			ActionResult? requiredResponse = (await DeviceController!.CreateDeviceReservationV1Async(requiredRequest)).Result;
+			ActionResult? blockedResponse = (await DeviceController!.CreateDeviceReservationV1Async(blockedRequest)).Result;
+
+			Assert.IsInstanceOfType(duplicateResponse, typeof(BadRequestObjectResult));
+			Assert.IsInstanceOfType(requiredResponse, typeof(ConflictObjectResult));
+			Assert.IsInstanceOfType(blockedResponse, typeof(ConflictObjectResult));
+		}
+
+		[TestMethod]
+		public async Task TestReservationDeviceTagRequirementsAsync()
+		{
+			await CreateFixtureAsync();
+
+			DeviceConfig devices = new DeviceConfig();
+
+			// Create a pool
+			string poolName = "TestDevicePool";
+			DevicePoolConfig devicePool = new DevicePoolConfig() { Id = new DevicePoolId(poolName), Name = poolName, PoolType = DevicePoolType.Automation, ProjectIds = new List<ProjectId>() { new ProjectId("ue5") } };
+			devices.Pools.Add(devicePool);
+
+			// Create a single platform
+			string platformName = "TestDevicePlatform";
+			DevicePlatformConfig platform = new DevicePlatformConfig() { Id = new DevicePlatformId(platformName), Name = platformName };
+			devices.Platforms.Add(platform);
+
+			await UpdateConfigAsync(x => x.Plugins.GetBuildConfig().Devices = devices);
+
+			// Create 10 devices
+			string expectedDeviceName = "ExpectedDevice";
+			for (int i = 0; i < 10; ++i)
+			{
+				CreateDeviceRequest testDevice;
+				if (i == 5)
+				{
+					// 6th device has a required tag
+					testDevice = new CreateDeviceRequest()
+					{
+						Name = expectedDeviceName,
+						PoolId = poolName,
+						Address = "10.0.0.1",
+						PlatformId = platformName,
+						Tags = new List<string> { "RequiredTag" }
+					};
+				}
+				else
+				{
+					// rest have a blocked tag
+					testDevice = new CreateDeviceRequest()
+					{
+						Name = "TestDevice_" + (i + 1),
+						PoolId = poolName,
+						Address = "10.0.0.1",
+						PlatformId = platformName,
+						Tags = new List<string> { "BlockedTag" }
+					};
+				}
+
+				await DeviceController.CreateDeviceAsync(testDevice);
+			}
+			
+			await DeviceService.TickForTestingAsync();
+
+			IJob job = await SetupJobAsync();
+
+			// Request a device with a required tag
+			DeviceTagsRequest requiredTag = new DeviceTagsRequest() { RequiredTags = ["RequiredTag"] };
+			LegacyCreateReservationRequest requiredRequest = SetupReservationTestAsync(job, poolName, platformName, null, null, null, requiredTag);
+
+			// Request a device ignoring blocked tags
+			DeviceTagsRequest blockedTag = new DeviceTagsRequest() { BlockedTags = ["BlockedTag"] };
+			LegacyCreateReservationRequest blockedRequest = SetupReservationTestAsync(job, poolName, platformName, null, null, null, blockedTag);
+
+			// Ensure we get the required device
+			GetLegacyReservationResponse requiredReservation = ResultToValue(await DeviceController!.CreateDeviceReservationV1Async(requiredRequest));
+			Assert.AreEqual(1, requiredReservation.DeviceNames.Length);
+			Assert.AreEqual(expectedDeviceName, requiredReservation.DeviceNames[0], StringComparer.OrdinalIgnoreCase);
+
+			// Delete the reservation
+			await DeviceController.DeleteReservationV1Async(requiredReservation.Guid);
+			await DeviceService.TickForTestingAsync();
+
+			// Attempt to reserve a device that is not blocked, which should give us the same device
+			GetLegacyReservationResponse blockedReservation = ResultToValue(await DeviceController!.CreateDeviceReservationV1Async(blockedRequest));
+			Assert.AreEqual(1, blockedReservation.DeviceNames.Length);
+			Assert.AreEqual(expectedDeviceName, blockedReservation.DeviceNames[0], StringComparer.OrdinalIgnoreCase);
+		}
+
+		[TestMethod]
+		public async Task TestReservationDeviceTagPreferencesAsync()
+		{
+			await CreateFixtureAsync();
+
+			DeviceConfig devices = new DeviceConfig();
+
+			// Create a pool
+			string poolName = "TestDevicePool";
+			DevicePoolConfig devicePool = new DevicePoolConfig() { Id = new DevicePoolId(poolName), Name = poolName, PoolType = DevicePoolType.Automation, ProjectIds = new List<ProjectId>() { new ProjectId("ue5") } };
+			devices.Pools.Add(devicePool);
+
+			// Create a single platform
+			string platformName = "TestDevicePlatform";
+			DevicePlatformConfig platform = new DevicePlatformConfig() { Id = new DevicePlatformId(platformName), Name = platformName };
+			devices.Platforms.Add(platform);
+
+			await UpdateConfigAsync(x => x.Plugins.GetBuildConfig().Devices = devices);
+
+			// Create 10 devices
+			string expectedDeviceName = "ExpectedDevice";
+			for (int i = 0; i < 10; ++i)
+			{
+				CreateDeviceRequest testDevice;
+				if (i == 5)
+				{
+					// 6th device has a preferred tag
+					testDevice = new CreateDeviceRequest()
+					{
+						Name = expectedDeviceName,
+						PoolId = poolName,
+						Address = "10.0.0.1",
+						PlatformId = platformName,
+						Tags = new List<string> { "PreferredTag" }
+					};
+				}
+				else
+				{
+					// rest have an undesired tag
+					testDevice = new CreateDeviceRequest()
+					{
+						Name = "TestDevice_" + (i + 1),
+						PoolId = poolName,
+						Address = "10.0.0.1",
+						PlatformId = platformName,
+						Tags = new List<string> { "UndesiredTag" }
+					};
+				}
+
+				await DeviceController.CreateDeviceAsync(testDevice);
+			}
+
+			await DeviceService.TickForTestingAsync();
+
+			IJob job = await SetupJobAsync();
+
+			// Request a device with a preferred tag
+			DeviceTagsRequest preferredTag = new DeviceTagsRequest() { PreferredTags = ["PreferredTag"] };
+			LegacyCreateReservationRequest preferredRequest = SetupReservationTestAsync(job, poolName, platformName, null, null, null, preferredTag);
+
+			// Request a device ignoring blocked tags
+			DeviceTagsRequest undesiredTag = new DeviceTagsRequest() { UndesiredTags = ["UndesiredTag"] };
+			LegacyCreateReservationRequest undesiredRequest = SetupReservationTestAsync(job, poolName, platformName, null, null, null, undesiredTag);
+
+			// Ensure we get the preferred device
+			GetLegacyReservationResponse preferredReservation = ResultToValue(await DeviceController!.CreateDeviceReservationV1Async(preferredRequest));
+			Assert.AreEqual(1, preferredReservation.DeviceNames.Length);
+			Assert.AreEqual(expectedDeviceName, preferredReservation.DeviceNames[0], StringComparer.OrdinalIgnoreCase);
+
+			// Now, while the preferred device is reserved, attempt to reserve it again. 
+			// Preferences aren't hard requirements, so we should get a device, just not one that has the preferred tag
+			GetLegacyReservationResponse anotherPreferredReservation = ResultToValue(await DeviceController!.CreateDeviceReservationV1Async(preferredRequest));
+			Assert.AreEqual(1, anotherPreferredReservation.DeviceNames.Length);
+			Assert.AreNotEqual(expectedDeviceName, anotherPreferredReservation.DeviceNames[0], StringComparer.OrdinalIgnoreCase);
+
+			// Delete the reservations
+			await DeviceController.DeleteReservationV1Async(preferredReservation.Guid);
+			await DeviceController.DeleteReservationV1Async(anotherPreferredReservation.Guid);
+			await DeviceService.TickForTestingAsync();
+
+			// Attempt to reserve a device that is not undesired, which should give us the same device
+			GetLegacyReservationResponse undesiredReservation = ResultToValue(await DeviceController!.CreateDeviceReservationV1Async(undesiredRequest));
+			Assert.AreEqual(1, undesiredReservation.DeviceNames.Length);
+			Assert.AreEqual(expectedDeviceName, undesiredReservation.DeviceNames[0], StringComparer.OrdinalIgnoreCase);
+
+			// Now, while the non-undesired device is reserved, attempt to reserve it again. 
+			// Preferences aren't hard requirements, so we should get a device, just not one that has is missing the undesired tag
+			GetLegacyReservationResponse anotherUndesiredReservation = ResultToValue(await DeviceController!.CreateDeviceReservationV1Async(undesiredRequest));
+			Assert.AreEqual(1, anotherUndesiredReservation.DeviceNames.Length);
+			Assert.AreNotEqual(expectedDeviceName, anotherUndesiredReservation.DeviceNames[0], StringComparer.OrdinalIgnoreCase);
+		}
+
 		[TestMethod]
 		public async Task TestProblemDeviceAsync()
 		{
@@ -807,7 +1036,6 @@ namespace HordeServer.Tests.Devices
 			ActionResult deleted = (await DeviceController!.DeleteReservationAsync(reservation.Id));
 
 			Assert.AreEqual((deleted as OkResult)!.StatusCode, 200);
-
 		}
 	}
 }
diff --git a/Engine/Source/Programs/Horde/Plugins/Build/HordeServer.Build/Devices/DeviceCollection.cs b/Engine/Source/Programs/Horde/Plugins/Build/HordeServer.Build/Devices/DeviceCollection.cs
--- a/Engine/Source/Programs/Horde/Plugins/Build/HordeServer.Build/Devices/DeviceCollection.cs
+++ b/Engine/Source/Programs/Horde/Plugins/Build/HordeServer.Build/Devices/DeviceCollection.cs
@@ -197,12 +197,15 @@ namespace HordeServer.Devices
 			[BsonIgnoreIfNull]
 			public int? Version { get; set; }
 
+			[BsonIgnoreIfNull]
+			public List<string>? Tags { get; set; }
+
 			[BsonConstructor]
 			private DeviceDocument()
 			{
 			}
 
-			public DeviceDocument(DeviceId id, DevicePlatformId platformId, DevicePoolId poolId, string name, bool enabled, string? address, string? modelId, UserId? userId)
+			public DeviceDocument(DeviceId id, DevicePlatformId platformId, DevicePoolId poolId, string name, bool enabled, string? address, string? modelId, UserId? userId, List<string>? tags)
 			{
 				Id = id;
 				PlatformId = platformId;
@@ -213,6 +216,7 @@ namespace HordeServer.Devices
 				ModelId = modelId;
 				ModifiedByUser = userId?.ToString();
 				Version = CurrentVersion;
+				Tags = tags;
 			}
 		}
 
@@ -482,9 +486,9 @@ namespace HordeServer.Devices
 		}
 
 		/// <inheritdoc/>
-		public async Task<IDevice?> TryAddDeviceAsync(DeviceId id, string name, DevicePlatformId platformId, DevicePoolId poolId, bool? enabled, string? address, string? modelId, UserId? userId)
+		public async Task<IDevice?> TryAddDeviceAsync(DeviceId id, string name, DevicePlatformId platformId, DevicePoolId poolId, bool? enabled, string? address, string? modelId, UserId? userId, List<string>? tags = null)
 		{
-			DeviceDocument newDevice = new DeviceDocument(id, platformId, poolId, name, enabled ?? true, address, modelId, userId);
+			DeviceDocument newDevice = new DeviceDocument(id, platformId, poolId, name, enabled ?? true, address, modelId, userId, tags);
 			await _devices.InsertOneAsync(newDevice);
 			return newDevice;
 		}
@@ -550,7 +554,7 @@ namespace HordeServer.Devices
 		}
 
 		/// <inheritdoc/>
-		public async Task UpdateDeviceAsync(DeviceId deviceId, DevicePoolId? newPoolId, string? newName, string? newAddress, string? newModelId, string? newNotes, bool? newEnabled, bool? newProblem, bool? newMaintenance, UserId? modifiedByUserId = null, string? message = null)
+		public async Task UpdateDeviceAsync(DeviceId deviceId, DevicePoolId? newPoolId, string? newName, string? newAddress, string? newModelId, string? newNotes, bool? newEnabled, bool? newProblem, bool? newMaintenance, UserId? modifiedByUserId = null, string? message = null, List<string>? newTags = null)
 		{
 			UpdateDefinitionBuilder<DeviceDocument> updateBuilder = Builders<DeviceDocument>.Update;
 			List<UpdateDefinition<DeviceDocument>> updates = new List<UpdateDefinition<DeviceDocument>>();
@@ -740,6 +744,11 @@ namespace HordeServer.Devices
 				}
 			}
 
+			if (newTags != null)
+			{
+				updates.Add(updateBuilder.Set(x => x.Tags, newTags));
+			}
+
 			if (updates.Count > 0)
 			{
 				await _devices.FindOneAndUpdateAsync<DeviceDocument>(x => x.Id == deviceId, updateBuilder.Combine(updates));
@@ -923,10 +932,8 @@ namespace HordeServer.Devices
 			// filter out currently reserved devices
 			poolDevices = poolDevices.FindAll(x => poolReservations.FirstOrDefault(p => p.Devices.Contains(x.Id)) == null);
 
-			int availablePoolDevices = poolDevices.Count;
-
 			// sort to use last reserved first to cycle devices
-			poolDevices.Sort((a, b) =>
+			Comparison<DeviceDocument> rankByLastReservationTime = (a, b) =>
 			{
 				DateTime? aTime = a.ReservationTimeUtc;
 				DateTime? bTime = b.ReservationTimeUtc;
@@ -946,8 +953,10 @@ namespace HordeServer.Devices
 				}
 
 				return aTime < bTime ? -1 : 1;
+			};
 
-			});
+			
+			poolDevices.Sort(rankByLastReservationTime);
 
 			if (existing?.ProblemDevice != null)
 			{
@@ -956,9 +965,74 @@ namespace HordeServer.Devices
 
 			foreach (DeviceRequestData data in request)
 			{
-				DeviceDocument? device = poolDevices.FirstOrDefault(a =>
+				List<DeviceDocument> rankedDevices = new List<DeviceDocument>(poolDevices);
+				
+				if (data.Tags != null)
 				{
+					bool hasPreferences = data.Tags.PreferredTags != null && data.Tags.PreferredTags.Count != 0;
+					bool hasUndesired = data.Tags.UndesiredTags != null && data.Tags.UndesiredTags.Count != 0;
+
+					if (hasPreferences || hasUndesired)
+					{
+						// if there are preferred or undesired tags in the request, sort the devices by ranking preferred highest, undesired lowest.
+						// if a device happens to have both a preferred and undesired tag, it's sort order will be ignored
+						rankedDevices.Sort((a, b) =>
+						{
+							int aPreferred = hasPreferences && a.Tags?.Count > 0 ? data.Tags.PreferredTags!.Count(tag => a.Tags.Contains(tag, StringComparer.OrdinalIgnoreCase)) : 0;
+							int aUndesired = hasUndesired && a.Tags?.Count > 0 ? data.Tags.UndesiredTags!.Count(tag => a.Tags.Contains(tag, StringComparer.OrdinalIgnoreCase)) : 0;
+							
+							int bPreferred = hasPreferences && b.Tags?.Count > 0 ? data.Tags.PreferredTags!.Count(tag => b.Tags.Contains(tag, StringComparer.OrdinalIgnoreCase)) : 0;
+							int bUndesired = hasUndesired && b.Tags?.Count > 0 ? data.Tags.UndesiredTags!.Count(tag => b.Tags.Contains(tag, StringComparer.OrdinalIgnoreCase)) : 0;
+
+							bool aIsPreferred = aPreferred > 0 && aUndesired == 0;
+							bool bIsPreferred = bPreferred > 0 && bUndesired == 0;
+							
+							bool aIsUndesired = aUndesired > 0 && aPreferred == 0;
+							bool bIsUndesired = bUndesired > 0 && bPreferred == 0;
+
+							// Sort preferred up
+							if (aIsPreferred && !bIsPreferred)
+							{
+								return -1;
+							}
+							if (bIsPreferred && !aIsPreferred)
+							{
+								return 1;
+							}
+							if (aIsPreferred && bIsPreferred)
+							{
+								if (aPreferred == bPreferred)
+								{
+									return rankByLastReservationTime(a, b);
+								}
+								return aPreferred > bPreferred ? -1 : 1;
+							}
+
+							// Sort undesired down
+							if (aIsUndesired && !bIsUndesired)
+							{
+								return 1;
+							}
+							if (bIsUndesired && !aIsUndesired)
+							{
+								return -1;
+							}
+							if (aIsUndesired && bIsUndesired)
+							{
+								if (aUndesired == bUndesired)
+								{
+									return rankByLastReservationTime(a, b);
+								}
+								return aUndesired > bUndesired ? 1 : -1;
+							}
 
+							return rankByLastReservationTime(a, b);
+						});
+					}
+				}
+				
+				DeviceDocument? device = rankedDevices.FirstOrDefault(a =>
+				{
 					if (allocated.Contains(a.Id) || a.PlatformId != data.PlatformId)
 					{
 						return false;
@@ -979,6 +1053,25 @@ namespace HordeServer.Devices
 						return false;
 					}
 
+					if (data.Tags != null)
+					{
+						if (data.Tags.RequiredTags != null && data.Tags.RequiredTags.Count > 0)
+						{
+							if (!data.Tags.RequiredTags.All(tag => a.Tags != null && a.Tags.Contains(tag, StringComparer.OrdinalIgnoreCase)))
+							{
+								return false;
+							}
+						}
+
+						if (data.Tags.BlockedTags != null && data.Tags.BlockedTags.Count > 0)
+						{
+							if (data.Tags.BlockedTags.Any(tag => a.Tags != null && a.Tags.Contains(tag, StringComparer.OrdinalIgnoreCase)))
+							{
+								return false;
+							}
+						}
+					}
+					
 					return true;
 				});
 
diff --git a/Engine/Source/Programs/Horde/Plugins/Build/HordeServer.Build/Devices/DeviceService.cs b/Engine/Source/Programs/Horde/Plugins/Build/HordeServer.Build/Devices/DeviceService.cs
--- a/Engine/Source/Programs/Horde/Plugins/Build/HordeServer.Build/Devices/DeviceService.cs
+++ b/Engine/Source/Programs/Horde/Plugins/Build/HordeServer.Build/Devices/DeviceService.cs
@@ -397,15 +397,15 @@ namespace HordeServer.Devices
 		/// <param name="modelId">Vendor model id</param>
 		/// <param name="userId">User adding the device</param>
 		/// <returns></returns>
-		public async Task<IDevice?> TryCreateDeviceAsync(DeviceId id, string name, DevicePlatformId platformId, DevicePoolId poolId, bool? enabled, string? address, string? modelId, UserId? userId = null)
+		public async Task<IDevice?> TryCreateDeviceAsync(DeviceId id, string name, DevicePlatformId platformId, DevicePoolId poolId, bool? enabled, string? address, string? modelId, UserId? userId = null, List<string>? tags = null)
 		{
-			return await _devices.TryAddDeviceAsync(id, name, platformId, poolId, enabled, address, modelId, userId);
+			return await _devices.TryAddDeviceAsync(id, name, platformId, poolId, enabled, address, modelId, userId, tags);
 		}
 
 		/// <summary>
 		/// Update a device
 		/// </summary>
-		public async Task UpdateDeviceAsync(DeviceId deviceId, DevicePoolId? newPoolId = null, string? newName = null, string? newAddress = null, string? newModelId = null, string? newNotes = null, bool? newEnabled = null, bool? newProblem = null, bool? newMaintenance = null, UserId? modifiedByUserId = null, string? message = null)
+		public async Task UpdateDeviceAsync(DeviceId deviceId, DevicePoolId? newPoolId = null, string? newName = null, string? newAddress = null, string? newModelId = null, string? newNotes = null, bool? newEnabled = null, bool? newProblem = null, bool? newMaintenance = null, UserId? modifiedByUserId = null, string? message = null, List<string>? newTags = null)
 		{
 			if (newProblem == true)
 			{
@@ -415,7 +415,7 @@ namespace HordeServer.Devices
 					await TryUpdateReservationAsync(reservation.Id, problemDevice: deviceId);
 				}
 			}
-			await _devices.UpdateDeviceAsync(deviceId, newPoolId, newName, newAddress, newModelId, newNotes, newEnabled, newProblem, newMaintenance, modifiedByUserId, message);
+			await _devices.UpdateDeviceAsync(deviceId, newPoolId, newName, newAddress, newModelId, newNotes, newEnabled, newProblem, newMaintenance, modifiedByUserId, message, newTags);
 		}
 
 		/// <summary>
diff --git a/Engine/Source/Programs/Horde/Plugins/Build/HordeServer.Build/Devices/DevicesController.cs b/Engine/Source/Programs/Horde/Plugins/Build/HordeServer.Build/Devices/DevicesController.cs
--- a/Engine/Source/Programs/Horde/Plugins/Build/HordeServer.Build/Devices/DevicesController.cs
+++ b/Engine/Source/Programs/Horde/Plugins/Build/HordeServer.Build/Devices/DevicesController.cs
@@ -88,7 +88,7 @@ namespace HordeServer.Devices
 			string name = deviceRequest.Name.Trim();
 			string? address = deviceRequest.Address?.Trim();
 
-			IDevice? device = await _deviceService.TryCreateDeviceAsync(DeviceId.Sanitize(name), name, platform.Id, pool.Id, deviceRequest.Enabled, address, modelId, internalUser?.Id);
+			IDevice? device = await _deviceService.TryCreateDeviceAsync(DeviceId.Sanitize(name), name, platform.Id, pool.Id, deviceRequest.Enabled, address, modelId, internalUser?.Id, deviceRequest.Tags);
 
 			if (device == null)
 			{
@@ -136,7 +136,7 @@ namespace HordeServer.Devices
 					ReservationFinishUtc = u.ReservationFinishUtc
 				}).ToList();
 
-				responses.Add(new GetDeviceResponse(device.Id.ToString(), device.PlatformId.ToString(), device.PoolId.ToString(), device.Name, device.Enabled, device.Address, device.ModelId?.ToString(), device.ModifiedByUser, device.Notes, device.ProblemTimeUtc, device.MaintenanceTimeUtc, utilization, device.CheckedOutByUser, device.CheckOutTime, checkoutExpiration, device.CleanRequired, device.CleanStartTimeUtc, device.CleanFinishTimeUtc));
+				responses.Add(new GetDeviceResponse(device.Id.ToString(), device.PlatformId.ToString(), device.PoolId.ToString(), device.Name, device.Enabled, device.Address, device.ModelId?.ToString(), device.ModifiedByUser, device.Notes, device.ProblemTimeUtc, device.MaintenanceTimeUtc, utilization, device.CheckedOutByUser, device.CheckOutTime, checkoutExpiration, device.CleanRequired, device.CleanStartTimeUtc, device.CleanFinishTimeUtc, device.Tags));
 			}
 
 			return responses;
@@ -240,7 +240,7 @@ namespace HordeServer.Devices
 			string? name = update.Name?.Trim();
 			string? address = update.Address?.Trim();
 
-			await _deviceService.UpdateDeviceAsync(deviceIdValue, poolIdValue, name, address, modelIdValue, update.Notes, update.Enabled, update.Problem, update.Maintenance, internalUser?.Id, update.Message);
+			await _deviceService.UpdateDeviceAsync(deviceIdValue, poolIdValue, name, address, modelIdValue, update.Notes, update.Enabled, update.Problem, update.Maintenance, internalUser?.Id, update.Message, update.Tags);
 
 			return Ok();
 		}
@@ -541,7 +541,14 @@ namespace HordeServer.Devices
 					}
 				}
 
-				requestedDevices.Add(new DeviceRequestData(platformIdValue, platformIdValue.ToString(), deviceRequest.IncludeModels, deviceRequest.ExcludeModels));
+				HashSet<string>? duplicateTags = deviceRequest.Tags?.GetDuplicateEntries();
+				if (duplicateTags != null && duplicateTags.Count > 0)
+				{
+					string duplicateList = String.Join(", ", duplicateTags);
+					return BadRequest($"The following device tags were requested more than once: {duplicateList}");
+				}
+
+				requestedDevices.Add(new DeviceRequestData(platformIdValue, platformIdValue.ToString(), deviceRequest.IncludeModels, deviceRequest.ExcludeModels, tags: deviceRequest.Tags));
 			}
 
 			(IDeviceReservation? reservation, string? errorMessage, bool installRequired) = await _deviceService.TryCreateReservationAsync(poolIdValue, requestedDevices);
@@ -906,7 +913,14 @@ namespace HordeServer.Devices
 					}
 				}
 
-				requestedDevices.Add(new DeviceRequestData(platform.Id, platformName, includeModels, excludeModels, request.DeviceName));
+				requestedDevices.Add(new DeviceRequestData(platform.Id, platformName, includeModels, excludeModels, request.DeviceName, request.Tags));
+			}
+
+			HashSet<string>? duplicateTags = request.Tags?.GetDuplicateEntries();
+			if (duplicateTags != null && duplicateTags.Count > 0)
+			{
+				string duplicateList = String.Join(", ", duplicateTags);
+				return BadRequest($"The following device tags were requested more than once: {duplicateList}");
 			}
 
 			(IDeviceReservation? reservation, string? errorMessage, bool installRequired) = await _deviceService.TryCreateReservationAsync(poolIdValue, requestedDevices, request.Hostname, request.ReservationDetails, request.JobId == null ? null : JobId.Parse(request.JobId), request.StepId == null ? null : JobStepId.Parse(request.StepId));
diff --git a/Engine/Source/Programs/Horde/Plugins/Build/HordeServer.Build/Devices/IDevice.cs b/Engine/Source/Programs/Horde/Plugins/Build/HordeServer.Build/Devices/IDevice.cs
--- a/Engine/Source/Programs/Horde/Plugins/Build/HordeServer.Build/Devices/IDevice.cs
+++ b/Engine/Source/Programs/Horde/Plugins/Build/HordeServer.Build/Devices/IDevice.cs
@@ -332,6 +332,11 @@ namespace HordeServer.Devices
 		/// </summary>
 		public string? Notes { get; }
 
+		/// <summary>
+		/// Metadata tags on the device
+		/// </summary>
+		public List<string>? Tags { get; }
+
 		/// <summary>
 		/// Whether the device requires a clean
 		/// </summary>
diff --git a/Engine/Source/Programs/Horde/Plugins/Build/HordeServer.Build/Devices/IDeviceCollection.cs b/Engine/Source/Programs/Horde/Plugins/Build/HordeServer.Build/Devices/IDeviceCollection.cs
--- a/Engine/Source/Programs/Horde/Plugins/Build/HordeServer.Build/Devices/IDeviceCollection.cs
+++ b/Engine/Source/Programs/Horde/Plugins/Build/HordeServer.Build/Devices/IDeviceCollection.cs
@@ -37,16 +37,23 @@ namespace HordeServer.Devices
 		/// </summary>
 		public string? DeviceName { get; set; } = null;
 
+		/// <summary>
+		/// Desired Tags of the device
+		/// </summary>
+		public DeviceTagsRequest? Tags { get; set; } = null;
+
 		/// <summary>
 		/// Constructor
 		/// </summary>
-		public DeviceRequestData(DevicePlatformId platformId, string requestedPlatform, List<string>? includeModels = null, List<string>? excludeModels = null, string? deviceName = null)
+		public DeviceRequestData(DevicePlatformId platformId, string requestedPlatform, List<string>? includeModels = null, List<string>? excludeModels = null,
+			string? deviceName = null, DeviceTagsRequest? tags = null)
 		{
 			PlatformId = platformId;
 			IncludeModels = includeModels ?? new List<string>();
 			ExcludeModels = excludeModels ?? new List<string>();
 			RequestedPlatform = requestedPlatform;
 			DeviceName = deviceName;
+			Tags = tags;
 		}
 	}
 
@@ -86,7 +93,7 @@ namespace HordeServer.Devices
 		/// <param name="address">The network address or hostname the device can be reached at</param>
 		/// <param name="modelId">The vendor model id of the device</param>
 		/// <param name="userId">The user adding the device</param>
-		Task<IDevice?> TryAddDeviceAsync(DeviceId id, string name, DevicePlatformId platformId, DevicePoolId poolId, bool? enabled, string? address, string? modelId, UserId? userId);
+		Task<IDevice?> TryAddDeviceAsync(DeviceId id, string name, DevicePlatformId platformId, DevicePoolId poolId, bool? enabled, string? address, string? modelId, UserId? userId, List<string>? tags = null);
 
 		/// <summary>
 		/// Update a device
@@ -102,7 +109,7 @@ namespace HordeServer.Devices
 		/// <param name="newMaintenance">Whether to set or clear maintenance state</param>
 		/// <param name="modifiedByUserId">The user who is updating the device</param>
 		/// <param name="message">Attached a message to the log of the device</param>
-		Task UpdateDeviceAsync(DeviceId deviceId, DevicePoolId? newPoolId, string? newName, string? newAddress, string? newModelId, string? newNotes, bool? newEnabled, bool? newProblem, bool? newMaintenance, UserId? modifiedByUserId = null, string? message = null);
+		Task UpdateDeviceAsync(DeviceId deviceId, DevicePoolId? newPoolId, string? newName, string? newAddress, string? newModelId, string? newNotes, bool? newEnabled, bool? newProblem, bool? newMaintenance, UserId? modifiedByUserId = null, string? message = null, List<string>? newTags = null);
 
 		/// <summary>
 		/// Updates a device's cleaning status
diff --git a/Engine/Source/Programs/Shared/EpicGames.Horde/Devices/DeviceMessages.cs b/Engine/Source/Programs/Shared/EpicGames.Horde/Devices/DeviceMessages.cs
--- a/Engine/Source/Programs/Shared/EpicGames.Horde/Devices/DeviceMessages.cs
+++ b/Engine/Source/Programs/Shared/EpicGames.Horde/Devices/DeviceMessages.cs
@@ -233,6 +233,11 @@ namespace EpicGames.Horde.Devices
 		/// The vendor model id of the device
 		/// </summary>
 		public string? ModelId { get; set; }
+
+		/// <summary>
+		/// Any tags associated with the device
+		/// </summary>
+		public List<string>? Tags { get; set; } = null;
 	}
 
 	/// <summary>
@@ -327,6 +332,11 @@ namespace EpicGames.Horde.Devices
 		/// </summary>
 		public string? Notes { get; set; }
 
+		/// <summary>
+		/// Metadata tags on the device
+		/// </summary>
+		public List<string>? Tags { get; set; }
+
 		/// <summary>
 		/// If the device has a marked problem
 		/// </summary>
@@ -380,7 +390,7 @@ namespace EpicGames.Horde.Devices
 		/// <summary>
 		/// Device response constructor
 		/// </summary>
-		public GetDeviceResponse(string id, string platformId, string poolId, string name, bool enabled, string? address, string? modelId, string? modifiedByUser, string? notes, DateTime? problemTime, DateTime? maintenanceTime, List<GetDeviceUtilizationResponse>? utilization, string? checkedOutByUser = null, DateTime? checkOutTime = null, DateTime? checkOutExpirationTime = null, bool? cleanRequired = null, DateTime? cleanStartTime = null, DateTime? cleanFinishTime = null)
+		public GetDeviceResponse(string id, string platformId, string poolId, string name, bool enabled, string? address, string? modelId, string? modifiedByUser, string? notes, DateTime? problemTime, DateTime? maintenanceTime, List<GetDeviceUtilizationResponse>? utilization, string? checkedOutByUser = null, DateTime? checkOutTime = null, DateTime? checkOutExpirationTime = null, bool? cleanRequired = null, DateTime? cleanStartTime = null, DateTime? cleanFinishTime = null, List<string>? tags = null)
 		{
 			Id = id;
 			Name = name;
@@ -400,6 +410,7 @@ namespace EpicGames.Horde.Devices
 			CleanRequired = cleanRequired;
 			CleanStartTime = cleanStartTime;
 			CleanFinishTime = cleanFinishTime;
+			Tags = tags;
 		}
 	}
 
@@ -452,6 +463,11 @@ namespace EpicGames.Horde.Devices
 		/// Additional message to attach to this request
 		/// </summary>
 		public string? Message { get; set; }
+
+		/// <summary>
+		/// Tags on the device
+		/// </summary>
+		public List<string>? Tags { get; set; }
 	}
 
 	/// <summary>
@@ -487,6 +503,59 @@ namespace EpicGames.Horde.Devices
 		public bool Clean { get; set; }
 	}
 
+	/// <summary>
+	/// Device tags request object
+	/// </summary>
+	public class DeviceTagsRequest
+	{
+		/// <summary>
+		/// Only devices with these tags will be considered for reservation
+		/// </summary>
+		public List<string>? RequiredTags { get; set; }
+
+		/// <summary>
+		/// Devices with these tags will be preferred over devices without them, but are not explicitly required
+		/// </summary>
+		public List<string>? PreferredTags { get; set; }
+
+		/// <summary>
+		/// Devices with these tags will be considered, but only as a last resort
+		/// </summary>
+		public List<string>? UndesiredTags { get; set; }
+
+		/// <summary>
+		/// Devices with these tags will not be considered for reservation
+		/// </summary>
+		public List<string>? BlockedTags { get; set; }
+
+		/// <summary>
+		/// Returns any duplicate tag entries
+		/// </summary>
+		public HashSet<string> GetDuplicateEntries() 
+		{
+			List<List<string>?> tagSets = new List<List<string>?>() { RequiredTags, PreferredTags, UndesiredTags, BlockedTags };
+			HashSet<string> seenTags = new HashSet<string>();
+			HashSet<string> duplicateTags = new HashSet<string>();
+
+			foreach (List<string>? set in tagSets)
+			{
+				if (set == null)
+				{
+					continue;
+				}
+				foreach (string tag in set)
+				{
+					if (!seenTags.Add(tag))
+					{
+						duplicateTags.Add(tag);
+					}
+				}
+			}
+
+			return duplicateTags;
+		}
+	}
+
 	/// <summary>
 	/// Device reservation request object
 	/// </summary>
@@ -507,6 +576,11 @@ namespace EpicGames.Horde.Devices
 		/// The optional vendor model ids to exclude for this device
 		/// </summary>
 		public List<string>? ExcludeModels { get; set; }
+
+		/// <summary>
+		/// Desired device tags for this device
+		/// </summary>
+		public DeviceTagsRequest? Tags { get; set; } = null;
 	}
 
 	/// <summary>
@@ -881,6 +955,10 @@ namespace EpicGames.Horde.Devices
 		/// </summary>
 		public string? DeviceName { get; set; } = null;
 
+		/// <summary>
+		/// Desired metadata of the device
+		/// </summary>
+		public DeviceTagsRequest? Tags { get; set; } = null;
 	}
 
 	/// <summary>
