[Horde Dashboard] Adds support for viewing and editing device tags on the horde dashboard's device editor

diff --git a/Engine/Source/Programs/Horde/HordeDashboard/src/backend/Api.ts b/Engine/Source/Programs/Horde/HordeDashboard/src/backend/Api.ts
--- a/Engine/Source/Programs/Horde/HordeDashboard/src/backend/Api.ts
+++ b/Engine/Source/Programs/Horde/HordeDashboard/src/backend/Api.ts
@@ -3658,6 +3658,9 @@ export type CreateDeviceRequest = {
 
 	/**  The vendor model id of the device */
 	modelId?: string;
+
+	/** List of metadata tags on the device */
+	tags?: string[];
 }
 
 export type CreateDeviceResponse = {
@@ -3712,6 +3715,9 @@ export type GetDeviceResponse = {
 	/** Any notes provided for the device */
 	notes?: string;
 
+	/** List of metadata tags on the device */
+	tags?: string[];
+
 	/** The UTC time when a device problem was reported */
 	problemTime?: Date | string;
 
@@ -3779,6 +3785,9 @@ export type UpdateDeviceRequest = {
 
 	/** Markdown notes */
 	notes?: string;
+
+	/** List of metadata tags on the device */
+	tags?: string[];
 }
 
 
diff --git a/Engine/Source/Programs/Horde/HordeDashboard/src/backend/DeviceHandler.ts b/Engine/Source/Programs/Horde/HordeDashboard/src/backend/DeviceHandler.ts
--- a/Engine/Source/Programs/Horde/HordeDashboard/src/backend/DeviceHandler.ts
+++ b/Engine/Source/Programs/Horde/HordeDashboard/src/backend/DeviceHandler.ts
@@ -107,6 +107,10 @@ export class DeviceHandler extends PollBase {
             });
 
          }
+
+         const uniqueTags = new Set<string>(); 
+         this.devices.forEach(device => device.tags?.forEach(tag => uniqueTags.add(tag)))
+         this.tags = Array.from(uniqueTags);
          
          this.loaded = true;         
          this.setUpdated();
@@ -296,6 +300,8 @@ export class DeviceHandler extends PollBase {
 
    loaded = false;
 
+   tags: string[] = [];
+
    private devices: GetDeviceResponse[] = [];
 
    platforms: Map<string, GetDevicePlatformResponse> = new Map();
diff --git a/Engine/Source/Programs/Horde/HordeDashboard/src/components/devices/DeviceEditor.tsx b/Engine/Source/Programs/Horde/HordeDashboard/src/components/devices/DeviceEditor.tsx
--- a/Engine/Source/Programs/Horde/HordeDashboard/src/components/devices/DeviceEditor.tsx
+++ b/Engine/Source/Programs/Horde/HordeDashboard/src/components/devices/DeviceEditor.tsx
@@ -1,4 +1,4 @@
-import { SpinnerSize, DialogType, IDropdownOption, MessageBarType, Text, Checkbox, DefaultButton, Dialog, DialogFooter, Dropdown, MessageBar, Modal, PrimaryButton, Spinner, Stack, TextField } from "@fluentui/react";
+import { SpinnerSize, DialogType, IDropdownOption, MessageBarType, Text, Checkbox, DefaultButton, Dialog, DialogFooter, Dropdown, MessageBar, Modal, PrimaryButton, Spinner, Stack, TextField, Label, ITag, TagPicker, IPickerItemProps } from "@fluentui/react";
 import React, { useState } from "react";
 import backend from "../../backend";
 import { DashboardPreference, GetDeviceResponse } from "../../backend/Api";
@@ -18,8 +18,89 @@ type DeviceEditData = {
    enabled: boolean;
    maintenance: boolean;
    notes?: string;
+   tags?: string[];
 };
 
+const onRenderSuggestionsItem = (item: ITag) => {
+    return <Stack style={{height: 24, padding: 4}}>
+                <Text title={item.name} style={{textOverflow: 'ellipsis', overflow: 'hidden', whiteSpace: 'nowrap', maxWidth: 300, fontSize: 12}}>{item.name}</Text>
+            </Stack>
+}
+
+const onRenderItem = (props: IPickerItemProps<ITag>) => {
+    const item = props.item;
+    return <Stack style={{ marginTop: 2, marginLeft: 3 }} key={`picker_item_${item.name}`}>
+                <PrimaryButton
+                    iconProps={{ iconName: "Cancel", styles: { root: { fontSize: 12, margin: 0, padding: 0 } } }}
+                    styles={{label: { textOverflow: 'ellipsis', overflow: 'hidden', whiteSpace: 'nowrap', minWidth: 0, maxWidth: 150, margin: 0 }}}
+                    style={{ padding: 2, paddingLeft: 5, paddingRight: 5, fontSize: 12, height: "unset"}}
+                    text={item.name}
+                    title={item.name}
+                    onClick={props.onRemoveItem} />
+            </Stack>;
+}
+
+const formatTag = (tag: string) => {
+   const trimmed = tag.trim();
+   if (trimmed.length === 0) return ''; 
+   return trimmed
+      .split(' ')
+      .map(word => word.length > 1 
+         ? word[0].toUpperCase() + word.slice(1) 
+         : word.toUpperCase())
+      .join(' ')
+}
+
+const getTextFromTag = (item: ITag) => formatTag(item.name);
+
+const TagList: React.FC<{tags: string[], allTags: string[], onChange: (newTags: string[]) => void;}> = ({tags, allTags, onChange}) => {
+   const [currentTags, setTags] = useState<string[]>(tags);
+ 
+   const selectedItems: ITag[] = currentTags.map((key) => ({key: key.toLowerCase(), name: key} as ITag));
+
+   const filterSuggestedTags = (filterText: string, _: ITag[]): ITag[] => {
+    if (filterText.length < 1) return [];
+
+    const lowerCurrentTags = currentTags.map(currentTag => currentTag.toLowerCase());
+    const lowerText = filterText.trim().toLowerCase();
+    const lowerTags = allTags.map(tag => tag.toLowerCase()).filter(tag => !lowerCurrentTags.includes(tag));
+
+    let filteredTags: ITag[] = [];
+
+    if (!lowerTags.includes(lowerText) && !lowerCurrentTags.includes(lowerText)) {
+      filteredTags = [{key: lowerText, name: formatTag(filterText)}];
+    }
+
+    filteredTags = filteredTags.concat(lowerTags.filter(tag => tag.startsWith(lowerText)).map(t => ({key: t, name: formatTag(t)})))
+
+    return filteredTags;
+   }
+ 
+   return (
+      <Stack>
+            <Label>Tags</Label>
+            <TagPicker
+               styles={{ input: { height: 28, width: 40 }, itemsWrapper: { marginRight: 3, marginBottom: 2 } }}
+               onRenderItem={onRenderItem}
+               onRenderSuggestionsItem={onRenderSuggestionsItem}
+               removeButtonAriaLabel="Remove"
+               selectionAriaLabel="Selected tags"
+               selectedItems={selectedItems}
+               onResolveSuggestions={filterSuggestedTags}
+               getTextFromItem={getTextFromTag}
+               onChange={(tags) => {
+                  let newTagList : string[] = [];
+                  if (!!tags) {
+                        newTagList = tags!.map(getTextFromTag);
+                  }
+                  setTags(newTagList);
+                  onChange?.(newTagList);
+               }}
+            />                
+      </Stack>
+  );
+}
+
 export const DeviceEditor: React.FC<{ handler: DeviceHandler, deviceIn?: GetDeviceResponse | undefined, editNote?: boolean, onClose: (device?: GetDeviceResponse) => void }> = ({ handler, deviceIn, editNote, onClose }) => {
 
    const [state, setState] = useState<{ device?: DeviceEditData, title?: string }>({});
@@ -77,7 +158,8 @@ export const DeviceEditor: React.FC<{ handler: DeviceHandler, deviceIn?: GetDevi
             name: device.name,
             poolId: device.poolId,
             modelId: device.modelId,
-            notes: device.notes
+            notes: device.notes,
+            tags: device.tags
          }: undefined;
       onClose(updatedDevice);
    };
@@ -161,7 +243,8 @@ export const DeviceEditor: React.FC<{ handler: DeviceHandler, deviceIn?: GetDevi
                poolId: deviceIn.poolId,
                modelId: deviceIn.modelId,
                notes: deviceIn.notes,
-               maintenance: !!deviceIn.maintenanceTime
+               maintenance: !!deviceIn.maintenanceTime,
+               tags : deviceIn.tags
             }, title: `Edit ${deviceIn.name}`
          });
       } else {
@@ -172,7 +255,8 @@ export const DeviceEditor: React.FC<{ handler: DeviceHandler, deviceIn?: GetDevi
                address: "",
                name: "",
                platformId: platformOptions.length > 0 ? platformOptions[0].data : "",
-               poolId: poolOptions.length > 0 ? poolOptions[0].data : ""
+               poolId: poolOptions.length > 0 ? poolOptions[0].data : "",
+               tags: []
             },
             title: "Add New Device"
          });
@@ -243,7 +327,8 @@ export const DeviceEditor: React.FC<{ handler: DeviceHandler, deviceIn?: GetDevi
                poolId: device.poolId,
                platformId: device.platformId,
                modelId: device.modelId,
-               enabled: device.enabled
+               enabled: device.enabled,
+               tags: device.tags
             });
 
             if (!result) {
@@ -261,7 +346,8 @@ export const DeviceEditor: React.FC<{ handler: DeviceHandler, deviceIn?: GetDevi
                notes: device.notes,
                problem: false, // always clear automatically generated problem state when saving
                maintenance: device.maintenance,
-               enabled: device.enabled
+               enabled: device.enabled,
+               tags: device.tags
             }).then(() => {
                close();
             }).catch((reason) => {
@@ -356,6 +442,14 @@ export const DeviceEditor: React.FC<{ handler: DeviceHandler, deviceIn?: GetDevi
 
             }} />}
 
+            <Stack>
+               <Label></Label>
+               <TagList tags={device.tags ? device.tags : []} allTags={handler.tags} onChange={(newTags => {
+                  device.tags = newTags;
+               })}/>
+            </Stack>
+            
+
             <TextField label="Notes" defaultValue={device.notes} multiline rows={5} resizable={false} onChange={(ev, newValue) => {
                device.notes = newValue;
             }} />
diff --git a/Engine/Source/Programs/Horde/HordeDashboard/src/components/devices/DeviceInfoView.tsx b/Engine/Source/Programs/Horde/HordeDashboard/src/components/devices/DeviceInfoView.tsx
--- a/Engine/Source/Programs/Horde/HordeDashboard/src/components/devices/DeviceInfoView.tsx
+++ b/Engine/Source/Programs/Horde/HordeDashboard/src/components/devices/DeviceInfoView.tsx
@@ -233,6 +233,7 @@ export const DeviceInfoModal: React.FC<{ handler: DeviceHandler, deviceIn?: GetD
       return <Stack />;
    }
 
+   const tags = deviceIn.tags ? deviceIn.tags.join(',') : '';
    const failureRate = jobState.telemetryItems.length == 0? 0 : Math.floor(jobState.telemetryItems.filter(item => !!item.job.problemTimeUtc).length / jobState.telemetryItems.length * 100);
    const status = handler.getDeviceStatus(deviceIn);
 
@@ -272,6 +273,10 @@ export const DeviceInfoModal: React.FC<{ handler: DeviceHandler, deviceIn?: GetD
                         <TextField label="Status" value={DeviceStatusLabels.get(status)} disabled={true} readOnly={true} />
                      </Stack>
 
+                     <Stack>
+                        <TextField label="Tags" value={tags} disabled={true} readOnly={true} />
+                     </Stack>
+
                      <Stack>
                         <TextField label="Last Modified By" value={jobState.modifiedByUser ? jobState.modifiedByUser : ""} disabled={true} readOnly={true} />
                      </Stack>
