Horde: Simplify running of health checks in Perforce load balancer

diff --git a/Engine/Source/Programs/Horde/HordeServer/Startup.cs b/Engine/Source/Programs/Horde/HordeServer/Startup.cs
--- a/Engine/Source/Programs/Horde/HordeServer/Startup.cs
+++ b/Engine/Source/Programs/Horde/HordeServer/Startup.cs
@@ -35,7 +35,6 @@ using HordeServer.Auditing;
 using HordeServer.Authentication;
 using HordeServer.Configuration;
 using HordeServer.Dashboard;
-using HordeServer.VersionControl.Perforce;
 using HordeServer.Plugins;
 using HordeServer.Server;
 using HordeServer.Server.Notices;
@@ -941,8 +940,7 @@ namespace HordeServer
 			services.AddHealthChecks()
 				.AddCheck("self", () => HealthCheckResult.Healthy(), tags: new[] { "self" })
 				.AddCheck<RedisService>("redis")
-				.AddCheck<MongoService>("mongo")
-				.AddCheck<PerforceLoadBalancer>("p4lb");
+				.AddCheck<MongoService>("mongo");
 		}
 
 		// This method gets called by the runtime. Use this method to configure the HTTP request pipeline.
diff --git a/Engine/Source/Programs/Horde/Plugins/Build/HordeServer.Build.Tests/Perforce/PerforceLoadBalancerTests.cs b/Engine/Source/Programs/Horde/Plugins/Build/HordeServer.Build.Tests/Perforce/PerforceLoadBalancerTests.cs
--- a/Engine/Source/Programs/Horde/Plugins/Build/HordeServer.Build.Tests/Perforce/PerforceLoadBalancerTests.cs
+++ b/Engine/Source/Programs/Horde/Plugins/Build/HordeServer.Build.Tests/Perforce/PerforceLoadBalancerTests.cs
@@ -42,7 +42,7 @@ public class PerforceLoadBalancerTests : BuildTestSetup
 	public async Task StatusHealthyViaInfoAsync()
 	{
 		CheckPerforceServerIsSet();
-		GlobalConfig gc = GetConfig(PerforceHostname);
+		GlobalConfig gc = GetConfig([PerforceHostname]);
 		PerforceLoadBalancer plb = Create(gc, "");
 		await plb.UpdateHealthTestOnlyAsync("myCluster", PerforceHostname, null, CancellationToken.None);
 		IPerforceServer server = (await plb.GetServersAsync(CancellationToken.None)).First();
@@ -53,7 +53,7 @@ public class PerforceLoadBalancerTests : BuildTestSetup
 	public async Task StatusHealthyViaHttpAsync()
 	{
 		CheckPerforceServerIsSet();
-		GlobalConfig gc = GetConfig(PerforceHostname);
+		GlobalConfig gc = GetConfig([PerforceHostname]);
 		PerforceLoadBalancer plb = Create(gc, GetHealthCheckJsonForStatus(PerforceServerStatus.Healthy));
 		await plb.UpdateHealthTestOnlyAsync("myCluster", PerforceHostname, "http://stub-handler-is-used", CancellationToken.None);
 		IPerforceServer server = (await plb.GetServersAsync(CancellationToken.None)).First();
@@ -64,20 +64,60 @@ public class PerforceLoadBalancerTests : BuildTestSetup
 	public async Task StatusUnhealthyViaHttpAsync()
 	{
 		CheckPerforceServerIsSet();
-		GlobalConfig gc = GetConfig(PerforceHostname);
+		GlobalConfig gc = GetConfig([PerforceHostname]);
 		PerforceLoadBalancer plb = Create(gc, GetHealthCheckJsonForStatus(PerforceServerStatus.Unhealthy));
 		await plb.UpdateHealthTestOnlyAsync("myCluster", PerforceHostname, "http://stub-handler-is-used", CancellationToken.None);
 		IPerforceServer server = (await plb.GetServersAsync(CancellationToken.None)).First();
 		Assert.AreEqual(PerforceServerStatus.Unhealthy, server.Status);
 	}
+	
+	[TestMethod]
+	[Ignore] // Needs better stubbing of health checks to not run for an excessively long time
+	public async Task ServersAddRemoveUpAsync()
+	{
+		const string Server1 = "server1:1666";
+		const string Server2 = "server2:1666";
+		GlobalConfig gc = GetConfig([Server1]);
+		PerforceLoadBalancer plb = Create(gc, GetHealthCheckJsonForStatus(PerforceServerStatus.Healthy));
+		
+		// No servers before first tick
+		Assert.AreEqual(0, (await plb.GetServersAsync(CancellationToken.None)).Count);
+		
+		// Servers are updated after first tick
+		await plb.TickInternalTestOnlyAsync(CancellationToken.None);
+		List<IPerforceServer> servers = await plb.GetServersAsync(CancellationToken.None);
+		Assert.AreEqual(1, servers.Count);
+		Assert.AreEqual(Server1, servers[0].ServerAndPort);
+		
+		// Add another server
+		gc.Plugins.GetBuildConfig().PerforceClusters = [CreateCluster(Server1, Server2)];
+		await plb.TickInternalTestOnlyAsync(CancellationToken.None);
+		servers = await plb.GetServersAsync(CancellationToken.None);
+		Assert.AreEqual(2, servers.Count);
+		Assert.AreEqual(Server1, servers[0].ServerAndPort);
+		Assert.AreEqual(Server2, servers[1].ServerAndPort);
+		
+		// Remove a server
+		gc.Plugins.GetBuildConfig().PerforceClusters = [CreateCluster(Server2)];
+		await plb.TickInternalTestOnlyAsync(CancellationToken.None);
+		servers = await plb.GetServersAsync(CancellationToken.None);
+		Assert.AreEqual(1, servers.Count);
+		Assert.AreEqual(Server2, servers[0].ServerAndPort);
+	}
 
-	private static GlobalConfig GetConfig(string perforceHost)
+	private static GlobalConfig GetConfig(params string[] perforceHosts)
 	{
-		PerforceCluster perforceCluster = new() { Servers = [new PerforceServer { ServerAndPort = perforceHost }] };
+		PerforceCluster perforceCluster = CreateCluster(perforceHosts);
 		GlobalConfig gc = new();
+		gc.Plugins.AddBuildConfig(new BuildConfig());
 		gc.Plugins.GetBuildConfig().PerforceClusters = [perforceCluster];
 		return gc;
 	}
+	
+	private static PerforceCluster CreateCluster(params string[] perforceHosts)
+	{
+		return new PerforceCluster { Servers = perforceHosts.Select(x => new PerforceServer { ServerAndPort = x }).ToList() };
+	}
 
 	private static string GetHealthCheckJsonForStatus(PerforceServerStatus status)
 	{
diff --git a/Engine/Source/Programs/Horde/Plugins/Build/HordeServer.Build/VersionControl/Perforce/PerforceLoadBalancer.cs b/Engine/Source/Programs/Horde/Plugins/Build/HordeServer.Build/VersionControl/Perforce/PerforceLoadBalancer.cs
--- a/Engine/Source/Programs/Horde/Plugins/Build/HordeServer.Build/VersionControl/Perforce/PerforceLoadBalancer.cs
+++ b/Engine/Source/Programs/Horde/Plugins/Build/HordeServer.Build/VersionControl/Perforce/PerforceLoadBalancer.cs
@@ -1,6 +1,5 @@
 // Copyright Epic Games, Inc. All Rights Reserved.
 
-using System.Diagnostics;
 using System.Globalization;
 using System.Net;
 using System.Text;
@@ -101,14 +100,20 @@ namespace HordeServer.VersionControl.Perforce
 	/// <summary>
 	/// Load balancer for Perforce edge servers
 	/// </summary>
-	public sealed class PerforceLoadBalancer : IHostedService, IHealthCheck, IAsyncDisposable
+	public sealed class PerforceLoadBalancer : IHostedService, IAsyncDisposable
 	{
 		/// <summary>
 		/// Information about a resolved Perforce server
 		/// </summary>
 		class PerforceServerEntry : IPerforceServer
 		{
+			/// <summary>
+			/// Resolved host:port for this Perforce server.
+			/// Mapping this separately from <see cref="BaseServerAndPort"/> is useful when a DNS name has multiple IP lookups for load balancing
+			/// </summary>
 			public string ServerAndPort { get; set; } = String.Empty;
+			
+			/// <inheritdoc cref="PerforceServer.ServerAndPort" />
 			public string BaseServerAndPort { get; set; } = String.Empty;
 			public string? HealthCheckUrl { get; set; }
 			public string Cluster { get; set; } = String.Empty;
@@ -145,8 +150,6 @@ namespace HordeServer.VersionControl.Perforce
 			public List<PerforceServerEntry> Servers { get; set; } = new List<PerforceServerEntry>();
 		}
 
-		record class HealthTask(Task Task, Stopwatch Timer);
-
 		readonly IRedisService _redisService;
 		readonly ILeaseCollection _leaseCollection;
 		readonly SingletonDocument<PerforceServerList> _serverListSingleton;
@@ -157,9 +160,6 @@ namespace HordeServer.VersionControl.Perforce
 		readonly Tracer _tracer;
 		readonly ILogger _logger;
 		readonly ITicker _ticker;
-		readonly CancellationTokenSource _healthTaskCts = new CancellationTokenSource();
-		readonly SemaphoreSlim _healthTaskSema = new SemaphoreSlim(1);
-		readonly Dictionary<string, HealthTask> _healthTasks = new Dictionary<string, HealthTask>(StringComparer.OrdinalIgnoreCase);
 
 		/// <summary>
 		/// Constructor
@@ -197,36 +197,6 @@ namespace HordeServer.VersionControl.Perforce
 		public async ValueTask DisposeAsync()
 		{
 			await _ticker.DisposeAsync();
-
-			if (_healthTasks.Count > 0)
-			{
-				await _healthTaskCts.CancelAsync();
-				foreach (HealthTask healthTask in _healthTasks.Values)
-				{
-					await healthTask.Task;
-				}
-				_healthTasks.Clear();
-			}
-
-			_healthTaskCts.Dispose();
-			_healthTaskSema.Dispose();
-		}
-
-		/// <inheritdoc/>
-		public async Task<HealthCheckResult> CheckHealthAsync(HealthCheckContext context, CancellationToken cancellationToken = default)
-		{
-			using (await _healthTaskSema.WaitDisposableAsync(cancellationToken))
-			{
-				KeyValuePair<string, HealthTask> longestTask = _healthTasks.OrderByDescending(x => x.Value.Timer.Elapsed).FirstOrDefault(x => !x.Value.Task.IsCompleted);
-				if (longestTask.Key == null || longestTask.Value.Timer.Elapsed < TimeSpan.FromMinutes(5.0))
-				{
-					return HealthCheckResult.Healthy();
-				}
-				else
-				{
-					return HealthCheckResult.Unhealthy($"Health check for {longestTask.Key} has taken {(int)longestTask.Value.Timer.Elapsed.TotalSeconds}s");
-				}
-			}
 		}
 
 		/// <summary>
@@ -459,8 +429,12 @@ namespace HordeServer.VersionControl.Perforce
 			}
 			return baseWeight + (totalLeases - candidate.NumLeases);
 		}
-
-		/// <inheritdoc/>
+		
+		internal ValueTask TickInternalTestOnlyAsync(CancellationToken cancellationToken)
+		{
+			return TickInternalAsync(cancellationToken);
+		}
+		
 		async ValueTask TickInternalAsync(CancellationToken cancellationToken)
 		{
 			_logger.LogDebug("Ticking server health status");
@@ -469,7 +443,7 @@ namespace HordeServer.VersionControl.Perforce
 			BuildConfig buildConfig = _buildConfig.CurrentValue;
 
 			// Set of new server entries
-			List<PerforceServerEntry> newServers = new List<PerforceServerEntry>();
+			List<PerforceServerEntry> newServers = [];
 
 			// Update the state of all the valid servers
 			foreach (PerforceCluster cluster in buildConfig.PerforceClusters)
@@ -488,41 +462,15 @@ namespace HordeServer.VersionControl.Perforce
 			}
 
 			// Update the number of leases for each entry
-			List<PerforceServerEntry> newEntries = newServers.OrderBy(x => x.Cluster).ThenBy(x => x.BaseServerAndPort).ThenBy(x => x.ServerAndPort).ToList();
-			await UpdateLeaseCountsAsync(newEntries, cancellationToken);
-			PerforceServerList list = await _serverListSingleton.UpdateAsync(list => MergeServerList(list, newEntries), cancellationToken);
+			List<PerforceServerEntry> newServerEntries = newServers.OrderBy(x => x.Cluster).ThenBy(x => x.BaseServerAndPort).ThenBy(x => x.ServerAndPort).ToList();
+			await UpdateLeaseCountsAsync(newServerEntries, cancellationToken);
+			PerforceServerList list = await _serverListSingleton.UpdateAsync(currentServerList => MergeServerList(currentServerList, newServerEntries), cancellationToken);
 			span.SetAttribute("numServers", list.Servers.Count);
 
-			// Lock the task lookup and update server health
-			using (await _healthTaskSema.WaitDisposableAsync(cancellationToken))
-			{
-				// Remove any health check tasks that have finished
-				KeyValuePair<string, HealthTask>[] healthTasks = _healthTasks.ToArray();
-				foreach ((string serverAndPort, HealthTask healthTask) in healthTasks)
-				{
-					if (healthTask.Task.IsCompleted)
-					{
-						await healthTask.Task;
-						_healthTasks.Remove(serverAndPort);
-					}
-					else if (healthTask.Timer.Elapsed > TimeSpan.FromMinutes(1.0))
-					{
-						_logger.LogWarning("Health check for {ServerAndPort} is still running after {TimeSecs}s", serverAndPort, (int)healthTask.Timer.Elapsed.TotalSeconds);
-					}
-				}
-
-				// Now update the health of each entry in parallel
-				foreach (PerforceServerEntry entry in list.Servers)
-				{
-					if (!_healthTasks.ContainsKey(entry.ServerAndPort))
-					{
-						Task task = Task.Run(() => UpdateHealthGuardedAsync(entry, _healthTaskCts.Token), CancellationToken.None);
-						_healthTasks.Add(entry.ServerAndPort, new HealthTask(task, Stopwatch.StartNew()));
-					}
-				}
-			}
+			// Run health checks for all Perforce servers in parallel
+			await Task.WhenAll(list.Servers.Select(server => Task.Run(async () => { await UpdateHealthGuardedAsync(server, cancellationToken); }, cancellationToken)));
 
-			// Update Horde's health monitor for subsystems
+			// Update Horde's internal health monitor for Perforce subsystem
 			list = await _serverListSingleton.GetAsync(cancellationToken);
 			(HealthStatus health, string? message) = GetPerforceHealth(list.Servers);
 			span.SetAttribute("health.status", health.ToString());
@@ -562,15 +510,14 @@ namespace HordeServer.VersionControl.Perforce
 
 		static void MergeServerList(PerforceServerList serverList, List<PerforceServerEntry> newEntries)
 		{
-			Dictionary<string, PerforceServerEntry> existingEntries = new Dictionary<string, PerforceServerEntry>(StringComparer.OrdinalIgnoreCase);
+			Dictionary<string, PerforceServerEntry> existingEntries = new (StringComparer.OrdinalIgnoreCase);
 			foreach (PerforceServerEntry existingEntry in serverList.Servers)
 			{
 				existingEntries[existingEntry.ServerAndPort] = existingEntry;
 			}
 			foreach (PerforceServerEntry newEntry in newEntries)
 			{
-				PerforceServerEntry? existingEntry;
-				if (existingEntries.TryGetValue(newEntry.ServerAndPort, out existingEntry))
+				if (existingEntries.TryGetValue(newEntry.ServerAndPort, out PerforceServerEntry? existingEntry))
 				{
 					newEntry.Status = existingEntry.Status;
 					newEntry.Detail = existingEntry.Detail;
@@ -654,7 +601,7 @@ namespace HordeServer.VersionControl.Perforce
 				initialHostName = initialHostName.Substring(0, portIdx);
 			}
 
-			List<string> hostNames = new List<string>();
+			List<string> hostNames = [];
 			if (!server.ResolveDns || IPAddress.TryParse(initialHostName, out IPAddress? _))
 			{
 				hostNames.Add(initialHostName);
