[Horde] Update GetLegacyDeviceResponse to return device tags

diff --git a/Engine/Source/Programs/Horde/Plugins/Build/HordeServer.Build/Devices/DevicesController.cs b/Engine/Source/Programs/Horde/Plugins/Build/HordeServer.Build/Devices/DevicesController.cs
--- a/Engine/Source/Programs/Horde/Plugins/Build/HordeServer.Build/Devices/DevicesController.cs
+++ b/Engine/Source/Programs/Horde/Plugins/Build/HordeServer.Build/Devices/DevicesController.cs
@@ -1062,6 +1062,7 @@ namespace HordeServer.Devices
 			response.Model = device.ModelId ?? "Base";
 			response.DeviceData = "";
 			response.PerfSpec = "Minimum";
+			response.Tags = device.Tags ?? new List<string>();
 
 			if (device.ModelId != null)
 			{
diff --git a/Engine/Source/Programs/Shared/EpicGames.Horde/Devices/DeviceMessages.cs b/Engine/Source/Programs/Shared/EpicGames.Horde/Devices/DeviceMessages.cs
--- a/Engine/Source/Programs/Shared/EpicGames.Horde/Devices/DeviceMessages.cs
+++ b/Engine/Source/Programs/Shared/EpicGames.Horde/Devices/DeviceMessages.cs
@@ -1099,5 +1099,11 @@ namespace EpicGames.Horde.Devices
 		/// </summary>
 		[Required]
 		public string DeviceData { get; set; } = null!;
+
+		/// <summary>
+		/// Any tags associated with the device
+		/// </summary>
+		[Required]
+		public List<string> Tags { get; set; } = null!;
 	}
 }
