Horde: Fix missing XML docs warning

#rnx

diff --git a/Engine/Source/Programs/Horde/Plugins/Build/HordeServer.Build/Devices/DeviceService.cs b/Engine/Source/Programs/Horde/Plugins/Build/HordeServer.Build/Devices/DeviceService.cs
--- a/Engine/Source/Programs/Horde/Plugins/Build/HordeServer.Build/Devices/DeviceService.cs
+++ b/Engine/Source/Programs/Horde/Plugins/Build/HordeServer.Build/Devices/DeviceService.cs
@@ -396,6 +396,7 @@ namespace HordeServer.Devices
 		/// <param name="address">Address or hostname of device</param>
 		/// <param name="modelId">Vendor model id</param>
 		/// <param name="userId">User adding the device</param>
+		/// <param name="tags">The metadata tags of the device</param>
 		/// <returns></returns>
 		public async Task<IDevice?> TryCreateDeviceAsync(DeviceId id, string name, DevicePlatformId platformId, DevicePoolId poolId, bool? enabled, string? address, string? modelId, UserId? userId = null, List<string>? tags = null)
 		{
diff --git a/Engine/Source/Programs/Horde/Plugins/Build/HordeServer.Build/Devices/IDeviceCollection.cs b/Engine/Source/Programs/Horde/Plugins/Build/HordeServer.Build/Devices/IDeviceCollection.cs
--- a/Engine/Source/Programs/Horde/Plugins/Build/HordeServer.Build/Devices/IDeviceCollection.cs
+++ b/Engine/Source/Programs/Horde/Plugins/Build/HordeServer.Build/Devices/IDeviceCollection.cs
@@ -89,10 +89,11 @@ namespace HordeServer.Devices
 		/// <param name="name">The name of the device (unique)</param>
 		/// <param name="platformId">The platform of the device</param>
 		/// <param name="poolId">Which pool to add it to</param>
-		/// <param name="enabled">Whather the device is enabled by default</param>
+		/// <param name="enabled">Whether the device is enabled by default</param>
 		/// <param name="address">The network address or hostname the device can be reached at</param>
 		/// <param name="modelId">The vendor model id of the device</param>
 		/// <param name="userId">The user adding the device</param>
+		/// <param name="tags">The metadata tags of the device</param>
 		Task<IDevice?> TryAddDeviceAsync(DeviceId id, string name, DevicePlatformId platformId, DevicePoolId poolId, bool? enabled, string? address, string? modelId, UserId? userId, List<string>? tags = null);
 
 		/// <summary>
@@ -109,6 +110,7 @@ namespace HordeServer.Devices
 		/// <param name="newMaintenance">Whether to set or clear maintenance state</param>
 		/// <param name="modifiedByUserId">The user who is updating the device</param>
 		/// <param name="message">Attached a message to the log of the device</param>
+		/// <param name="newTags">The new metadata tags to assign</param>
 		Task UpdateDeviceAsync(DeviceId deviceId, DevicePoolId? newPoolId, string? newName, string? newAddress, string? newModelId, string? newNotes, bool? newEnabled, bool? newProblem, bool? newMaintenance, UserId? modifiedByUserId = null, string? message = null, List<string>? newTags = null);
 
 		/// <summary>
