Avoid issue when phase name is undefinied

#rnx

diff --git a/Engine/Source/Programs/Horde/HordeDashboard/plugins/test-automation/testDataV2Fetcher.ts b/Engine/Source/Programs/Horde/HordeDashboard/plugins/test-automation/testDataV2Fetcher.ts
--- a/Engine/Source/Programs/Horde/HordeDashboard/plugins/test-automation/testDataV2Fetcher.ts
+++ b/Engine/Source/Programs/Horde/HordeDashboard/plugins/test-automation/testDataV2Fetcher.ts
@@ -7,7 +7,7 @@ const minuteInWeek = 10080;
 class TestPhaseStatusV2 extends TestPhaseStatus {
 
     constructor(phase: PhaseSessionDetailsResponse, devices: DeviceRef[], parent: TestSessionDetails) {
-        super(phase.key, phase.name, new Date(phase.dateTime), phase.timeElapseSec, devices, phase.outcome, parent)
+        super(phase.key, phase.name ?? "<undefinied>", new Date(phase.dateTime), phase.timeElapseSec, devices, phase.outcome, parent)
         this._eventStreamPath = phase.eventStreamPath;
         this.hasWarning = phase.hasWarning;
     }
