Horde: Remove unused WithDatadog setting

diff --git a/Engine/Source/Programs/Horde/HordeServer/ServerSettings.cs b/Engine/Source/Programs/Horde/HordeServer/ServerSettings.cs
--- a/Engine/Source/Programs/Horde/HordeServer/ServerSettings.cs
+++ b/Engine/Source/Programs/Horde/HordeServer/ServerSettings.cs
@@ -435,11 +435,6 @@ namespace HordeServer
 		/// </summary>
 		public int? GlobalThreadPoolMinSize { get; set; }
 
-		/// <summary>
-		/// Whether to enable Datadog integration for tracing
-		/// </summary>
-		public bool WithDatadog { get; set; }
-
 		/// <summary>
 		/// Path to the root config file. Relative to the server.json file by default.
 		/// </summary>
