Horde Dashboard: repopulate adv add arg quotes as escaped quotes

diff --git a/Engine/Source/Programs/Horde/HordeDashboard/src/components/build/NewBuildV2.tsx b/Engine/Source/Programs/Horde/HordeDashboard/src/components/build/NewBuildV2.tsx
--- a/Engine/Source/Programs/Horde/HordeDashboard/src/components/build/NewBuildV2.tsx
+++ b/Engine/Source/Programs/Horde/HordeDashboard/src/components/build/NewBuildV2.tsx
@@ -71,6 +71,9 @@ class BuildOptions {
       // Also prevents overquoting incorrectly recorded args 
       // i.e., forgetting "" in -set:XXX=YYY ZZZ AAA BBB => -set:XXX="YYY" "ZZZ" "AAA" "BBB"
       if(!arg.startsWith("-set")) return arg;
+
+      // Replace any existing quotes with escaped quotes
+      arg = arg.replace(/"/g, '\\"');
       
       // Split argument on space to extract left side assignment
       const splitOnSpace = arg.split(" ");
