[Horde][Analytics] Initial proof of concept for Build Health Dashboard View.

diff --git a/Engine/Source/Programs/Horde/HordeDashboard/src/App.tsx b/Engine/Source/Programs/Horde/HordeDashboard/src/App.tsx
--- a/Engine/Source/Programs/Horde/HordeDashboard/src/App.tsx
+++ b/Engine/Source/Programs/Horde/HordeDashboard/src/App.tsx
@@ -46,6 +46,7 @@ import { preloadFonts } from './styles/Styles';
 import { darkTheme } from './styles/darkTheme';
 import { lightTheme } from './styles/lightTheme';
 import { ErrorBar } from './components/error/ErrorBar';
+import { BuildHealthView } from './components/buildhealth/BuildhealthView';
 
 let router: any;
 
@@ -165,6 +166,7 @@ const Main: React.FC = () => {
                { path: "audit/job/:jobId", element: <AuditLogView /> },
                { path: "audit/template/:streamId/:templateId", element: <AuditLogView /> },
                { path: "automation", element: <AutomationView /> },
+               { path: "buildhealth", element: <BuildHealthView/> },
                { path: "tools", element: <ToolView /> },
                { path: "lease/:leaseId", element: <DebugView /> },
                { path: "docs", element: <DocView /> },
diff --git a/Engine/Source/Programs/Horde/HordeDashboard/src/components/buildhealth/BuildHealthDataHandler.ts b/Engine/Source/Programs/Horde/HordeDashboard/src/components/buildhealth/BuildHealthDataHandler.ts
new file mode 100644
--- /dev/null
+++ b/Engine/Source/Programs/Horde/HordeDashboard/src/components/buildhealth/BuildHealthDataHandler.ts
@@ -0,0 +1,229 @@
+// Copyright Epic Games, Inc. All Rights Reserved.
+
+import backend from "horde/backend";
+import { JobData, JobStreamQuery, ProjectData, StreamData } from "horde/backend/Api";
+import { makeObservable, observable, runInAction } from "mobx";
+import { BuildHealthOptionsController, BuildHealthOptionsState } from "./BuildHealthOptions";
+import { StepRefData, TemplateRefData } from "./BuildHealthDataTypes";
+import { decodeTemplateKey } from "./BuildHealthUtilities";
+
+// #region -- API Data Types & Utilities --
+
+/**
+ * Request object that controls what items to refresh upon Data Handler refresh.
+ */
+export type DataHandlerRefreshRequest = {
+    projects?: boolean;
+    streams?: boolean;
+    jobs?: boolean;
+    steps?: boolean;
+}
+
+export const DEFAULT_REFRESH_REQUEST: DataHandlerRefreshRequest = { projects: true, streams: true, jobs: true, steps: true };
+
+/**
+ * Utility function that will expand a Data Handler Refresh Request to include all the dependents in it's refresh.
+ * @param request The request to expand.
+ * @returns The expanded request.
+ */
+export function expandRefreshRequest(request: DataHandlerRefreshRequest): DataHandlerRefreshRequest {
+    const expanded: DataHandlerRefreshRequest = { ...request };
+    if (request.projects) {
+        expanded.streams = true;
+    }
+
+    if (request.streams) {
+        expanded.jobs = true;
+    }
+
+    if (request.jobs) {
+        expanded.steps = true;
+    }
+
+    return expanded;
+}
+
+// #endregion -- API Data Types & Utilities --
+
+/**
+ * Data handler for all build health data used for dropdowns.
+ * @todo Data sharing & caching would ideally occur between the BuildHealth View & Step Outcome.
+ */
+export class BuildHealthDataHandler {
+    private readonly STREAMS_FILTER = "id,state,streamId,name,batches";
+    private readonly STEP_DATA_SAMPLE_SIZE = 5;
+
+    // #region -- Private Members --
+
+    private lastRefresh = -1;
+    private buildHealthOptions: BuildHealthOptionsState;
+    private buildHealthOptionsController: BuildHealthOptionsController;
+    private cachedStreamToTemplatesData: Map<string, TemplateRefData[]> = new Map<string, TemplateRefData[]>();
+
+    // #endregion -- Private Members --
+
+    // #region -- Public Members --
+
+    @observable.shallow
+    projectsData: ProjectData[] = [];
+
+    @observable.shallow
+    streamsData: StreamData[] = [];
+
+    @observable.shallow
+    templatesData: TemplateRefData[] = [];
+
+    @observable.shallow
+    stepData: StepRefData[] = [];
+
+    // #endregion -- Public Members --
+
+    // #region -- Constructor --
+
+    constructor(buildHealthOptionsState: BuildHealthOptionsState, buildHealthOptionsController: BuildHealthOptionsController) {
+        makeObservable(this);
+
+        this.buildHealthOptions = buildHealthOptionsState;
+        this.buildHealthOptionsController = buildHealthOptionsController;
+    }
+
+    //#endregion -- Constructor --
+
+    // #region -- Private API --
+
+    private async refreshData(request: DataHandlerRefreshRequest) {
+        return this.getProjectsData(request);
+    }
+
+    // #endregion -- Private API --
+
+    // #region -- Public API --
+
+    /**
+     * Requests a hierarchical data refresh. Refreshw ill occur if the options have changed since last refresh.
+     * @param request The data refresh request message to use in order to control refresh depth.
+     * @param force Whether to force the teh refresh or not. If false, will only refresh if the options have changed since last refresh.
+     */
+    async requestHierarchicalRefresh(request: DataHandlerRefreshRequest = DEFAULT_REFRESH_REQUEST, force?: boolean) {
+        if (this.lastRefresh < this.buildHealthOptionsController.optionsChangeVersion || force) {
+            this.lastRefresh = this.buildHealthOptionsController.optionsChangeVersion;
+            await this.refreshData(request);
+        }
+    }
+
+    /**
+     * Obtains the project data given the current buildHealthOptions.
+     * @param request The data refresh request message to use in order to control refresh depth.
+     */
+    async getProjectsData(request: DataHandlerRefreshRequest) {
+        if (this.projectsData.length === 0 || request.projects) {
+            const projectDataResult = (await backend.getProjects());
+
+            runInAction(() => {
+                this.projectsData = projectDataResult;
+            });
+        }
+
+        await this.getStreamData(request);
+    }
+
+    /**
+     * Obtains the stream data given the current buildHealthOptions.
+     * @param request The data refresh request message to use in order to control refresh depth.
+     */
+    async getStreamData(request: DataHandlerRefreshRequest) {
+        if (request.streams && Object.keys(this.buildHealthOptions.enabledProjects).length > 0) {
+            let streamsData: StreamData[] = [];
+            this.cachedStreamToTemplatesData.clear();
+
+            this.projectsData.forEach((project: ProjectData) => {
+                const isProjectEnabled = this.buildHealthOptions.enabledProjects.hasOwnProperty(project.id);
+                if (!isProjectEnabled) {
+                    return;
+                }
+
+                project.streams?.forEach((stream: StreamData) => {
+                    streamsData.push(stream);
+                    this.cachedStreamToTemplatesData.set(stream.id, stream.templates);
+                });
+            });
+
+            runInAction(() => {
+                this.streamsData = streamsData;
+            });
+        }
+
+        await this.getTemplateData(request);
+    }
+
+    /**
+     * Obtains the template data given the current buildHealthOptions.
+     * @param request The data refresh request message to use in order to control refresh depth.
+     */
+    async getTemplateData(request: DataHandlerRefreshRequest) {
+        if (request.jobs && Object.keys(this.buildHealthOptions.enabledStreams).length > 0) {
+            let templatesData: TemplateRefData[] = [];
+            for (const key of Object.keys(this.buildHealthOptions.enabledStreams)) {
+                if (this.cachedStreamToTemplatesData.has(key)) {
+                    const templates = this.cachedStreamToTemplatesData.get(key);
+                    if (templates) {
+                        templatesData.push(
+                            ...templates.map(t => ({ ...t, streamId: key }))
+                        );
+                    }
+                }
+            }
+
+            runInAction(() => {
+                this.templatesData = templatesData;
+            });
+        }
+
+        await this.getStepData(request);
+    }
+
+    /**
+     * Obtains the step data given the current buildHealthOptions.
+     * @param request The data refresh request message to use in order to control refresh depth.
+     */
+    async getStepData(request: DataHandlerRefreshRequest) {
+        const entries = Object.entries(this.buildHealthOptions.enabledTemplates);
+        if (request.steps && entries.length > 0) {
+            const results = await Promise.all(entries.map(async ([key]) => {
+                const templateData = decodeTemplateKey(key);
+
+                const query: JobStreamQuery = {
+                    template: [templateData.templateId],
+                    count: this.STEP_DATA_SAMPLE_SIZE,
+                    filter: this.STREAMS_FILTER
+                };
+
+                const jobDatas: JobData[] = await backend.getStreamJobs(templateData.streamId, query);
+                const recentJobs = jobDatas.filter(job => job.batches);
+                const stepMap = new Map<string, StepRefData>();
+
+                for (const job of recentJobs) {
+                    for (const batch of job.batches!) {
+                        for (const step of batch.steps) {
+                            const stepKey = step.name;
+                            if (!stepMap.has(stepKey)) {
+                                stepMap.set(stepKey, { ...step, templateId: decodeTemplateKey(key).templateId });
+                            }
+                        }
+                    }
+                }
+
+                const distinctSteps = [...stepMap.values()];
+                return distinctSteps;
+            }));
+
+            const allSteps = results.flat();
+
+            runInAction(() => {
+                this.stepData = [...allSteps.values()];
+            });
+        }
+    }
+
+    // #endregion -- Public API --
+}
\ No newline at end of file
diff --git a/Engine/Source/Programs/Horde/HordeDashboard/src/components/buildhealth/BuildHealthDataTypes.ts b/Engine/Source/Programs/Horde/HordeDashboard/src/components/buildhealth/BuildHealthDataTypes.ts
new file mode 100644
--- /dev/null
+++ b/Engine/Source/Programs/Horde/HordeDashboard/src/components/buildhealth/BuildHealthDataTypes.ts
@@ -0,0 +1,26 @@
+// Copyright Epic Games, Inc. All Rights Reserved.
+
+import { GetStepResponse, GetTemplateRefResponse } from "horde/backend/Api";
+
+export const PROJECTS_URL_PARAM_NAME: string = "projects";
+export const STREAMS_URL_PARAM_NAME: string = "streams";
+export const TEMPLATES_URL_PARAM_NAME: string = "templates";
+export const STEPS_URL_PARAM_NAME: string = "stepNames";
+export const JOBCOUNT_URL_PARAM_NAME: string = "jobCount";
+export const PARAMETER_KEY_PREFIX: string = "parameter_key";
+export const KEY_SEPARATOR: string = "::";
+
+/**
+ * Convenience type to fold in extra metadata relevant for the Template. 
+ */
+export type TemplateRefData = GetTemplateRefResponse & {
+    streamId?: string;
+    fullname?: string;
+}
+
+/**
+ * Convenience type to fold in extra metadata relevant for the Step. 
+ */
+export type StepRefData = GetStepResponse & {
+    templateId?: string;
+}
\ No newline at end of file
diff --git a/Engine/Source/Programs/Horde/HordeDashboard/src/components/buildhealth/BuildHealthOptions.ts b/Engine/Source/Programs/Horde/HordeDashboard/src/components/buildhealth/BuildHealthOptions.ts
new file mode 100644
--- /dev/null
+++ b/Engine/Source/Programs/Horde/HordeDashboard/src/components/buildhealth/BuildHealthOptions.ts
@@ -0,0 +1,387 @@
+// Copyright Epic Games, Inc. All Rights Reserved.
+
+import { action, makeObservable, observable, runInAction } from "mobx";
+import { JOBCOUNT_URL_PARAM_NAME, PROJECTS_URL_PARAM_NAME, STEPS_URL_PARAM_NAME, STREAMS_URL_PARAM_NAME, TemplateRefData, TEMPLATES_URL_PARAM_NAME } from "../buildhealth/BuildHealthDataTypes";
+import { BuildHealthDataHandler } from "./BuildHealthDataHandler";
+import { decodeStepKey, decodeTemplateKey, encodeStepKey, encodeTemplateKey } from "./BuildHealthUtilities";
+import { Location } from 'react-router-dom';
+
+// #region -- Options Utilities & Data Types --
+
+/**
+ * Data structure used to pass state from location query, to Build Health Options.
+ */
+type BuildHealthQueryParams = {
+    projectIds: string[];
+    streamIds: string[];
+    templateIds: string[];
+    stepNames: string[];
+    jobCount?: number;
+};
+
+/**
+ * Parses the Build Health options from the query params.
+ * @param location Location object.
+ * @returns Parsed Build Health Query Params.
+ */
+export function parseBuildHealthQueryParams(location: Location): BuildHealthQueryParams {
+    const searchParams = new URLSearchParams(location.search);
+    const getList = (key: string) => searchParams.get(key)?.split(",").filter(Boolean) ?? [];
+
+    const projectIds = getList(PROJECTS_URL_PARAM_NAME);
+    const streamIds = getList(STREAMS_URL_PARAM_NAME);
+    const templateIds = getList(TEMPLATES_URL_PARAM_NAME);
+    const stepNames = getList(STEPS_URL_PARAM_NAME);
+
+    let jobCount: number | undefined;
+    const jobCountRaw = searchParams.get(JOBCOUNT_URL_PARAM_NAME);
+
+    if (jobCountRaw !== null) {
+        const parsed = Number(jobCountRaw);
+        if (Number.isFinite(parsed) && parsed >= 0) {
+            jobCount = parsed;
+        }
+    }
+
+
+    return { projectIds, streamIds, templateIds: templateIds, stepNames, jobCount };
+}
+
+/**
+ * Processes provided Build Health Query Params for the Build Health View.
+ * @param params The params to process.
+ * @param options The Build Health Option controller to use to update options from the params.
+ * @param handler  THe Build Health Data Handler to use to obtain data based on the params.
+ */
+export async function loadBuildHealthOptionsFromParams(
+    params: BuildHealthQueryParams,
+    options: BuildHealthOptionsController,
+    handler: BuildHealthDataHandler) {
+    if (params.jobCount) {
+        runInAction(() => {
+            options.setJobCount(params.jobCount!);
+        });
+    }
+
+    // Get baseline project data first
+    await handler.requestHierarchicalRefresh();
+
+    runInAction(() => {
+        options.clearProjects();
+        params.projectIds.forEach(pid => {
+            const project = handler.projectsData.find(p => p.id === pid);
+            if (project) {
+                options.toggleProject(pid, project.name, true);
+            }
+            else {
+                console.warn(`Project parameter provided in URL query (${pid}) that could not be matched to underlying data - discarding.`);
+            }
+        });
+    });
+
+    await handler.getStreamData({ streams: true });
+
+    runInAction(() => {
+        options.clearStreams();
+        params.streamIds.forEach(sid => {
+            const stream = handler.streamsData.find(s => s.id === sid);
+            if (stream) {
+                options.toggleStream(sid, stream.name, true);
+            }
+            else {
+                console.warn(`Stream parameter provided in URL query (${sid}) that could not be matched to underlying data - discarding.`);
+            }
+        });
+    });
+
+    await handler.getTemplateData({ jobs: true });
+
+    runInAction(() => {
+        options.clearTemplates();
+        params.templateIds.forEach(tid => {
+            const template = handler.templatesData.find(t => t.id === tid);
+            if (template) {
+                options.toggleTemplate(encodeTemplateKey(template), template.name, true);
+            }
+            else {
+                console.warn(`Template parameter provided in URL query (${tid}) that could not be matched to underlying data - discarding.`);
+            }
+        });
+    });
+
+    await handler.getStepData({ steps: true });
+
+    runInAction(() => {
+        options.clearSteps();
+        params.stepNames.forEach(name => {
+            const step = handler.stepData.find(s => s.name.toLowerCase() === name.toLowerCase());
+            if (step) {
+                options.toggleStep(encodeStepKey(step), step.name, true);
+            }
+            else {
+                console.warn(`Step name parameter provided in URL query (${name}) that could not be matched to underlying data - discarding.`);
+            }
+        });
+
+        options.synchronizeDerivedKeys();
+    });
+}
+
+// #endregion -- Options Utilities & Data Types --
+
+/**
+ * Build Health View Options.
+ */
+export class BuildHealthOptionsState {
+
+    // #region -- Published Options --
+
+    @observable stepOutcomeEnabledStreamKeys: string[] = [];
+    @observable stepOutcomeEnabledJobKeys: string[] = [];
+    @observable stepOutcomeEnabledStepKeys: string[] = [];
+    @observable stepOutcomeTotalJobs: number = 25;
+
+    // #endregion -- Published Options --
+
+    // #region -- User Selected Options --
+
+    @observable enabledProjects: Record<string, string> = {};
+    @observable enabledStreams: Record<string, string> = {};
+    @observable enabledTemplates: Record<string, string> = {};
+    @observable enabledSteps: Record<string, string> = {};
+
+    // #endregion -- User Selected Options --
+
+    // #region -- Constructor --
+
+    constructor() {
+        makeObservable(this);
+    }
+
+    // #endregion -- Constructor --
+}
+
+/**
+ * Controller responsible for modifying the underlying BuildHealthOptionsState.
+ */
+export class BuildHealthOptionsController {
+    // #region -- Private Members --
+
+    readonly state: BuildHealthOptionsState;
+    private lastSynchronize = -1;
+
+    // #endregion -- Private Members --
+
+    // #region -- Public Members --
+
+    optionsChangeVersion = 0;
+
+    // #endregion -- Public Members --
+
+    // #region -- Constructor --
+
+    constructor(state: BuildHealthOptionsState) {
+        this.state = state;
+        makeObservable(this);
+    }
+
+    private setOptionsChanged() {
+        this.optionsChangeVersion++;
+    }
+
+    // #endregion -- Constructor --
+
+    // #region -- Public API --
+
+    /**
+     * Toggles the state of the provided project. Will clear all other set projects.
+     * @param id The id of the project to toggle.
+     * @param name The name of the project.
+     * @param enabled Whether it's enabled or not.
+     */
+    @action
+    toggleSingleProject(id: string, name: string) {
+        this.state.enabledProjects = { [id]: name };
+        this.setOptionsChanged();
+    }
+
+    /**
+     * Toggles the state of the provided project.
+     * @param id The id of the project to toggle.
+     * @param name The name of the project.
+     * @param enabled Whether it's enabled or not.
+     */
+    @action
+    toggleProject(id: string, name: string, enabled: boolean) {
+        if (enabled) {
+            this.state.enabledProjects[id] = name;
+        } else {
+            delete this.state.enabledProjects[id];
+            this.clearStreams(true);
+        }
+        this.setOptionsChanged();
+    }
+
+    /**
+     * Toggles the state of the provided stream. Will clear all other set streams.
+     * @param id The id of the stream to toggle.
+     * @param name The name of the stream.
+     * @param enabled Whether it's enabled or not.
+     */
+    @action
+    toggleSingleStream(id: string, name: string) {
+        this.clearStreams(true);
+        this.state.enabledStreams[id] = name;
+        this.setOptionsChanged();
+    }
+
+    /**
+     * Toggles the state of the provided stream.
+     * @param id The id of the stream to toggle.
+     * @param name The name of the stream.
+     * @param enabled Whether it's enabled or not.
+     */
+    @action
+    toggleStream(id: string, name: string, enabled: boolean) {
+        if (enabled) {
+            this.state.enabledStreams[id] = name;
+        } else {
+            delete this.state.enabledStreams[id];
+            this.clearTemplates(true);
+        }
+        this.setOptionsChanged();
+    }
+
+    /**
+     * Toggles the state of the provided template.
+     * @param id The id of the template to toggle.
+     * @param name The name of the template.
+     * @param enabled Whether it's enabled or not.
+     */
+    @action
+    toggleTemplate(id: string, name: string, enabled: boolean) {
+        if (enabled) {
+            this.state.enabledTemplates[id] = name;
+        } else {
+            delete this.state.enabledTemplates[id];
+            this.clearSteps();
+        }
+        this.setOptionsChanged();
+    }
+
+    /**
+     * Toggles the state of the provided step.
+     * @param id The id of the step to toggle.
+     * @param name The name of the step.
+     * @param enabled Whether it's enabled or not.
+     */
+    @action
+    toggleStep(id: string, name: string, enabled: boolean) {
+        if (enabled) {
+            this.state.enabledSteps[id] = name;
+        } else {
+            delete this.state.enabledSteps[id];
+        }
+        this.setOptionsChanged();
+    }
+
+    /**
+     * Changes the Job Count on the options.
+     * @param newJobCount The new job count.
+     */
+    @action
+    setJobCount(newJobCount: number) {
+        this.state.stepOutcomeTotalJobs = newJobCount;
+        this.setOptionsChanged();
+    }
+
+    /**
+      * Clear all the projects tracked by the options.
+      * @param cascade Whether to cascade this down to dependent options.
+      */
+    clearProjects(cascade?: boolean) {
+        Object.keys(this.state.enabledProjects).forEach(key => {
+            delete this.state.enabledProjects[key];
+        });
+
+        if (cascade) {
+            this.clearStreams(cascade);
+        }
+    }
+
+    /**
+      * Clear all the streams tracked by the options.
+      * @param cascade Whether to cascade this down to dependent options.
+      */
+    clearStreams(cascade?: boolean) {
+        Object.keys(this.state.enabledStreams).forEach(key => {
+            delete this.state.enabledStreams[key];
+        });
+        if (cascade) {
+            this.clearTemplates(cascade);
+        }
+    }
+
+    /**
+      * Clear all the templates tracked by the options.
+      * @param cascade Whether to cascade this down to dependent options.
+      */
+    clearTemplates(cascade?: boolean) {
+        Object.keys(this.state.enabledTemplates).forEach(key => {
+            delete this.state.enabledTemplates[key];
+        });
+        if (cascade) {
+            this.clearSteps();
+        }
+    }
+
+    /**
+     * Clear all the steps tracked by the options.
+     * @param cascade Whether to cascade this down to dependent options.
+     */
+    clearSteps() {
+        Object.keys(this.state.enabledSteps).forEach(key => {
+            delete this.state.enabledSteps[key];
+        });
+    }
+
+    /**
+     * Sycnrhonizes the options that are bound to UI elements, with the data observed by consumers of the option set.
+     * This is important to keep separate as it allows us to control when we signal/flush a finalized set of options to the consumers.
+     */
+    @action
+    synchronizeDerivedKeys() {
+        if (this.lastSynchronize < this.optionsChangeVersion) {
+            this.state.stepOutcomeEnabledStreamKeys = Object.keys(this.state.enabledStreams);
+            this.state.stepOutcomeEnabledJobKeys = Object.keys(this.state.enabledTemplates).map(key => decodeTemplateKey(key).templateId);
+            this.state.stepOutcomeEnabledStepKeys = Object.keys(this.state.enabledSteps).map(key => decodeStepKey(key).stepName);
+            this.lastSynchronize = this.optionsChangeVersion;
+        }
+    }
+
+    /**
+     * Obtains a URLSearchParams representation of the current options.
+     * @returns The URLSearchParams representation of the current options.
+     */
+    toNavigationQuery(): URLSearchParams {
+        const params = new URLSearchParams(location.search);
+
+        const projects = Object.keys(this.state.enabledProjects).join(",");
+        const streams = this.state.stepOutcomeEnabledStreamKeys.join(",");
+        const jobs = this.state.stepOutcomeEnabledJobKeys.join(",");
+        const steps = this.state.stepOutcomeEnabledStepKeys.join(",");
+        const jobCount = this.state.stepOutcomeTotalJobs.toString();
+
+        if (!projects.length) return params;
+
+        params.set(PROJECTS_URL_PARAM_NAME, projects);
+        params.set(STREAMS_URL_PARAM_NAME, streams);
+        params.set(TEMPLATES_URL_PARAM_NAME, jobs);
+        params.set(STEPS_URL_PARAM_NAME, steps);
+        params.set(JOBCOUNT_URL_PARAM_NAME, jobCount);
+
+        return params;
+    }
+
+    // #endregion -- Public API --
+}
\ No newline at end of file
diff --git a/Engine/Source/Programs/Horde/HordeDashboard/src/components/buildhealth/BuildHealthUtilities.ts b/Engine/Source/Programs/Horde/HordeDashboard/src/components/buildhealth/BuildHealthUtilities.ts
new file mode 100644
--- /dev/null
+++ b/Engine/Source/Programs/Horde/HordeDashboard/src/components/buildhealth/BuildHealthUtilities.ts
@@ -0,0 +1,41 @@
+// Copyright Epic Games, Inc. All Rights Reserved.
+
+import { KEY_SEPARATOR, StepRefData, TemplateRefData } from "./BuildHealthDataTypes";
+
+/**
+ * Produces an encoded template key for unique identification.
+ * @param template The template to produce the key for.
+ * @returns The key representing the template.
+ */
+export function encodeTemplateKey(template: TemplateRefData): string {
+    return `${template.streamId}${KEY_SEPARATOR}${template.id}`;
+}
+
+/**
+ * Decodes an encoded template key, providing it's constituent parts.
+ * @param templateKey The template key to decode.
+ * @returns The constituent parts.
+ */
+export function decodeTemplateKey(templateKey: string): { streamId: string; templateId: string } {
+    const [streamId, templateId] = templateKey.split(KEY_SEPARATOR);
+    return { streamId, templateId };
+}
+
+/**
+ * Produces an encoded step key for unique identification.
+ * @param step The step to produce the key for.
+ * @returns The key representing the step.
+ */
+export function encodeStepKey(step: StepRefData) {
+    return `${step.templateId}${KEY_SEPARATOR}${step.name}`;
+}
+
+/**
+ * Decodes an encoded step key, providing it's constituent parts.
+ * @param templateKey The step key to decode.
+ * @returns The constituent parts.
+ */
+export function decodeStepKey(stepKey: string): { templateId: string, stepName: string } {
+    const [templateId, stepName] = stepKey.split(KEY_SEPARATOR);
+    return { templateId, stepName };
+}
\ No newline at end of file
diff --git a/Engine/Source/Programs/Horde/HordeDashboard/src/components/buildhealth/BuildhealthView.tsx b/Engine/Source/Programs/Horde/HordeDashboard/src/components/buildhealth/BuildhealthView.tsx
new file mode 100644
--- /dev/null
+++ b/Engine/Source/Programs/Horde/HordeDashboard/src/components/buildhealth/BuildhealthView.tsx
@@ -0,0 +1,516 @@
+// Copyright Epic Games, Inc. All Rights Reserved.
+
+import { Dropdown, DropdownMenuItemType, IDropdownOption, Label, Slider, Stack, TooltipHost } from "@fluentui/react";
+import { useWindowSize } from "../../base/utilities/hooks";
+import { getHordeStyling } from "../../styles/Styles";
+import { Breadcrumbs } from "../Breadcrumbs";
+import { TopNav } from "../TopNav";
+import { StepOutcomeView } from "../buildhealth/stepoutcome/StepOutcome";
+import { useEffect, useMemo, useState } from "react";
+import { GetProjectResponse, GetStreamResponse } from "horde/backend/Api";
+import { observer } from "mobx-react-lite";
+import { useLocation, useNavigate } from "react-router";
+import { Text } from '@fluentui/react';
+import { BuildHealthDataHandler, DataHandlerRefreshRequest, expandRefreshRequest } from "./BuildHealthDataHandler";
+import { BuildHealthOptionsController, BuildHealthOptionsState, loadBuildHealthOptionsFromParams, parseBuildHealthQueryParams } from "./BuildHealthOptions";
+import { encodeStepKey, encodeTemplateKey } from "./BuildHealthUtilities";
+import { KEY_SEPARATOR, PARAMETER_KEY_PREFIX, StepRefData, TemplateRefData } from "./BuildHealthDataTypes";
+
+// #region -- Dropdown Components --
+
+// #region -- Dropdown Data Types --
+
+type BuildHealthOptionData = {
+    id: string,
+    text: string,
+    group?: string
+}
+
+type BuildHealthOptionList = {
+    items: BuildHealthOptionData[]
+    label: string,
+    toolTip: string
+}
+
+// #endregion -- Dropdown Data Types --
+
+
+// #region -- Base Dropdown Components --
+
+/**
+ * React Component that is a single select drop down.
+ * @returns React Component.
+ */
+const BasicListParameter: React.FC<{ handler: BuildHealthDataHandler, buildHealthOptions: BuildHealthOptionsController, param: BuildHealthOptionList, disabled?: boolean, selectedKey?: string, request: DataHandlerRefreshRequest, onChange: (option: IDropdownOption<BuildHealthOptionData>) => void }> = function ConstructBasicListParamater({ param, buildHealthOptions, selectedKey, request, onChange }) {
+    const key = `${PARAMETER_KEY_PREFIX}${KEY_SEPARATOR}${param.label}`;
+    const changeVersion = buildHealthOptions.optionsChangeVersion;
+    const doptions: IDropdownOption<BuildHealthOptionData>[] = [];
+
+    param.items.forEach((item) => {
+        doptions.push({
+            key: item.id,
+            text: item.text,
+        });
+    });
+
+    return (<Dropdown id={changeVersion.toString() + key}
+        key={key}
+        label={param.label}
+        options={doptions}
+        dropdownWidth={200}
+        selectedKey={selectedKey}
+        placeholder={"Select option"}
+        onChange={(_ev, option) => {
+            let castedOption = option as IDropdownOption<BuildHealthOptionData>;
+            if (!castedOption) {
+                return;
+            }
+            onChange(castedOption as IDropdownOption<BuildHealthOptionData>);
+        }}
+        onDismiss={() => {
+            handler.requestHierarchicalRefresh(request);
+            buildHealthOptions.synchronizeDerivedKeys();
+        }}
+    />);
+};
+
+/**
+ * React Component that is a multi select drop down.
+ * @returns React Component.
+ */
+const MultiListParameter: React.FC<{ handler: BuildHealthDataHandler, buildHealthOptions: BuildHealthOptionsController, param: BuildHealthOptionList, disabled?: boolean, selectedKeys: string[], request: DataHandlerRefreshRequest, onChange: (option: IDropdownOption<BuildHealthOptionData>) => void }> = function ConstructMultilistParamter({ param, buildHealthOptions, selectedKeys, request, onChange }) {
+    const changeVersion = buildHealthOptions.optionsChangeVersion;
+    const { modeColors } = getHordeStyling();
+    const key = `${PARAMETER_KEY_PREFIX}${KEY_SEPARATOR}${param.label}`;// + `_${options.currentRenderKey}`;
+    const gset: Set<string> = new Set();
+
+    param.items.forEach(item => {
+        if (!item.group) {
+            item.group = "__nogroup";
+        }
+        if (item.group) {
+            gset.add(item.group);
+        }
+    });
+
+    const groups = Array.from(gset).sort((a, b) => a.localeCompare(b));
+
+    const doptions: IDropdownOption[] = [];
+
+    groups.forEach(group => {
+
+        if (group !== "__nogroup") {
+            doptions.push({
+                key: `group_${group}`,
+                text: group,
+                itemType: DropdownMenuItemType.Header
+            });
+        }
+
+        const dupes = new Map<string, number>();
+
+        param.items.forEach(item => {
+            const v = dupes.get(item.text);
+            dupes.set(item.text, !v ? 1 : v + 1);
+        });
+
+        param.items.forEach(item => {
+            if (item.group === group) {
+                const key = item.id;
+                doptions.push({
+                    key: key,
+                    data: item,
+                    text: (item.group !== "__nogroup" && dupes.get(item.text)! > 1) ? `${item.text} - ${item.group}` : item.text,
+                });
+            }
+        });
+    });
+
+    return (
+        <Dropdown id={changeVersion.toString() + key}
+            key={key}
+            placeholder={"Select options"}
+            dropdownWidth={300}
+            styles={{
+                callout: {
+                    selectors: {
+                        ".ms-Callout-main": {
+                            padding: "4px 4px 12px 12px",
+                            overflow: "auto",
+                        }
+                    }
+                },
+                dropdownItemHeader: { fontSize: 12, color: modeColors.text },
+                dropdownOptionText: { fontSize: 12 },
+                dropdownItem: {
+                    minHeight: 18, lineHeight: 18, selectors: {
+                        '.ms-Checkbox-checkbox': {
+                            width: 14,
+                            height: 14,
+                            fontSize: 11
+                        }
+                    }
+                },
+                dropdownItemSelected: {
+                    minHeight: 18, lineHeight: 18, backgroundColor: "inherit",
+                    selectors: {
+                        '.ms-Checkbox-checkbox': {
+                            width: 14,
+                            height: 14,
+                            fontSize: 11
+                        }
+                    }
+                }
+            }}
+            label={param.label}
+            selectedKeys={selectedKeys}
+            onChange={(_ev, option) => {
+                let castedOption = option as IDropdownOption<BuildHealthOptionData>;
+                if (!castedOption) {
+                    return;
+                }
+                onChange(castedOption);
+            }}
+            onDismiss={() => {
+                handler.requestHierarchicalRefresh(request);
+                buildHealthOptions.synchronizeDerivedKeys();
+            }}
+            multiSelect
+            options={doptions}
+        />);
+};
+
+// #endregion -- Base Dropdown Components --
+
+// #region -- Option Specific Dropdown Components --
+
+const ProjectDropdownSingle: React.FC<{ handler: BuildHealthDataHandler, buildHealthOptions: BuildHealthOptionsController }> = observer(function ConstructProjectDropdownSingle({ handler, buildHealthOptions }) {
+    const projectDropdownParams: BuildHealthOptionList = {
+        label: "Select Project",
+        items: [
+        ],
+        toolTip: "Select an option from the dropdown."
+    };
+    projectDropdownParams.items = handler.projectsData.map((project: GetProjectResponse) => {
+        let listParam: BuildHealthOptionData = {
+            id: `${project.id}`,
+            text: project.name,
+        };
+
+        return listParam;
+    }).sort((a, b) => a.text.localeCompare(b.text));
+
+    let selectedProjectKey = Object.keys(buildHealthOptions.state.enabledProjects)[0] ?? undefined;
+
+    // Note: This may look odd, but there is really no reason to force a full hierarchy of project refresh; we just want to pulse streams
+    let request: DataHandlerRefreshRequest = expandRefreshRequest({ streams: true });
+
+    return (
+        <Stack>
+            <BasicListParameter handler={handler} buildHealthOptions={buildHealthOptions} param={projectDropdownParams} selectedKey={selectedProjectKey} request={request} onChange={(option: IDropdownOption<BuildHealthOptionData>) => {
+                const key = option.key as string;
+                buildHealthOptions.toggleSingleProject(key, option.text);
+            }} />
+        </Stack>
+    );
+});
+
+const ProjectDropdown: React.FC<{ handler: BuildHealthDataHandler, buildHealthOptions: BuildHealthOptionsController }> = observer(function ConstructProjectDropdown({ handler, buildHealthOptions }) {
+    const projectDropdownParams: BuildHealthOptionList = {
+        label: "Select Project",
+        items: [
+        ],
+        toolTip: "Select an option from the dropdown."
+    };
+
+    projectDropdownParams.items = handler.projectsData.map((project: GetProjectResponse) => {
+        let listParam: BuildHealthOptionData = {
+            id: `${PARAMETER_KEY_PREFIX}${KEY_SEPARATOR}${projectDropdownParams.label}${KEY_SEPARATOR}${project.id}`,
+            text: project.name,
+        };
+
+        return listParam;
+    }).sort((a, b) => a.text.localeCompare(b.text));
+
+    let selectedProjectKeys: string[] = Object.keys(buildHealthOptions.state.enabledProjects);
+    let request: DataHandlerRefreshRequest = expandRefreshRequest({ streams: true });
+
+    return (
+        <Stack>
+            <MultiListParameter handler={handler} buildHealthOptions={buildHealthOptions} param={projectDropdownParams} selectedKeys={selectedProjectKeys} request={request} onChange={(option: IDropdownOption<BuildHealthOptionData>) => {
+                const key = option.key as string;
+                const isSelected = !!option.selected;
+                buildHealthOptions.toggleProject(key, option.text, isSelected);
+            }} />
+        </Stack>
+    );
+});
+
+const StreamDropdownSingle: React.FC<{ handler: BuildHealthDataHandler, buildHealthOptions: BuildHealthOptionsController }> = observer(function ConstructProjectDropdownSingle({ handler, buildHealthOptions }) {
+    const streamDropdownParams: BuildHealthOptionList = {
+        label: "Select Stream",
+        items: [
+        ],
+        toolTip: "Select an option from the dropdown."
+    };
+
+    streamDropdownParams.items = handler.streamsData.map((stream: GetStreamResponse) => {
+        let listParam: BuildHealthOptionData = {
+            id: stream.id,
+            text: stream.name,
+            group: stream.projectId,
+        };
+        return listParam;
+    }).sort((a, b) => a.text.localeCompare(b.text));
+
+    let selectedStreamKey: string = Object.keys(buildHealthOptions.state.enabledStreams)[0]!;
+    let request: DataHandlerRefreshRequest = expandRefreshRequest({ streams: true });
+    {
+        return streamDropdownParams.items.length > 0 && (
+            <Stack >
+                <BasicListParameter handler={handler} buildHealthOptions={buildHealthOptions} param={streamDropdownParams} selectedKey={selectedStreamKey} request={request} onChange={(option: IDropdownOption<BuildHealthOptionData>) => {
+                    const key = option.key as string;
+                    buildHealthOptions.toggleSingleStream(key, option.text);
+                }} />
+            </Stack>
+        );
+    }
+});
+
+const StreamDropdown: React.FC<{ handler: BuildHealthDataHandler, buildHealthOptions: BuildHealthOptionsController }> = observer(function ConstructProjectDropdown({ handler, buildHealthOptions }) {
+    const streamDropdownParams: BuildHealthOptionList = {
+        label: "Select Stream(s)",
+        items: [
+        ],
+        toolTip: "Select an option from the dropdown."
+    };
+
+    streamDropdownParams.items = handler.streamsData.map((stream: GetStreamResponse) => {
+        let listParam: BuildHealthOptionData = {
+            id: stream.id,
+            text: stream.name,
+            group: stream.projectId,
+        };
+        return listParam;
+    }).sort((a, b) => a.text.localeCompare(b.text));
+
+    let selectedStreamKeys: string[] = Object.keys(buildHealthOptions.state.enabledStreams);
+    let request: DataHandlerRefreshRequest = expandRefreshRequest({ streams: true });
+    {
+        return streamDropdownParams.items.length > 0 && (
+            <Stack>
+                <MultiListParameter handler={handler} buildHealthOptions={buildHealthOptions} param={streamDropdownParams} selectedKeys={selectedStreamKeys} request={request} onChange={(option: IDropdownOption<BuildHealthOptionData>) => {
+                    const key = option.key as string;
+                    const isSelected = !!option.selected;
+
+                    buildHealthOptions.toggleStream(key, option.text, isSelected);
+                }} />
+            </Stack>
+        );
+    }
+});
+
+const TemplateDropdown: React.FC<{ handler: BuildHealthDataHandler, buildHealthOptions: BuildHealthOptionsController }> = observer(function ConstructProjectDropdown({ handler, buildHealthOptions }) {
+    const templateDropdownParams: BuildHealthOptionList = {
+        label: "Select Template(s)",
+        items: [
+        ],
+        toolTip: "Select an option from the dropdown."
+    };
+
+    templateDropdownParams.items = handler.templatesData.map((template: TemplateRefData) => {
+        let listParam: BuildHealthOptionData = {
+            id: encodeTemplateKey(template),
+            text: template.name,
+            group: template.streamId,
+        };
+        return listParam;
+    }).sort((a, b) => a.text.localeCompare(b.text));
+
+    let selectedTemplateKeys: string[] = Object.keys(buildHealthOptions.state.enabledTemplates);
+    let request: DataHandlerRefreshRequest = expandRefreshRequest({ jobs: true });
+    {
+        return templateDropdownParams.items.length > 0 && (
+            <Stack id={`job_dropdown`}>
+                <MultiListParameter handler={handler} buildHealthOptions={buildHealthOptions} param={templateDropdownParams} selectedKeys={selectedTemplateKeys} request={request} onChange={(option: IDropdownOption<BuildHealthOptionData>) => {
+                    const key = option.key as string;
+                    const isSelected = !!option.selected;
+
+                    buildHealthOptions.toggleTemplate(key, option.text, isSelected);
+                }} />
+            </Stack>
+        );
+    }
+});
+
+const StepDropdown: React.FC<{ handler: BuildHealthDataHandler, buildHealthOptions: BuildHealthOptionsController }> = observer(function ConstructProjectDropdown({ handler, buildHealthOptions }) {
+    const stepDropdownParams: BuildHealthOptionList = {
+        label: "Select Step(s)",
+        items: [
+        ],
+        toolTip: "Select an option from the dropdown."
+    };
+
+    stepDropdownParams.items = handler.stepData.map((step: StepRefData) => {
+        let listParam: BuildHealthOptionData = {
+            id: encodeStepKey(step),
+            text: step.name,
+            group: step.templateId
+        };
+        return listParam;
+    }).sort((a, b) => a.text.localeCompare(b.text));
+
+    let selectedStepKeys: string[] = Object.keys(buildHealthOptions.state.enabledSteps);
+    let request: DataHandlerRefreshRequest = { steps: true };
+    {
+        return stepDropdownParams.items.length > 0 && (
+            <Stack id={`step_dropdown`}>
+                <MultiListParameter handler={handler} buildHealthOptions={buildHealthOptions} param={stepDropdownParams} selectedKeys={selectedStepKeys} request={request} onChange={(option: IDropdownOption<BuildHealthOptionData>) => {
+                    const key = option.key as string;
+                    const isSelected = !!option.selected;
+
+                    buildHealthOptions.toggleStep(key, option.text, isSelected);
+                }} />
+            </Stack>
+        );
+    }
+});
+
+// #endregion -- Option Specific Dropdown Components --
+
+// #endregion -- Dropdown Components --
+
+// #region -- Primary View Components --
+
+/**
+ * React Component used to manage the URL query params, and refresh it based on build health options. This is separate from the main component to scope rerenders to solely this component.
+ * @returns React Component.
+ */
+const BuildHealthUrlSync: React.FC<{ buildHealthController: BuildHealthOptionsController }> = function ConstrucBuildHealthSidebar({ buildHealthController }) {
+    const location = useLocation();
+    const navigate = useNavigate();
+
+    useEffect(() => {
+        const initializeFromUrl = async () => {
+            const params = parseBuildHealthQueryParams(location);
+            await loadBuildHealthOptionsFromParams(params, buildHealthController, handler);
+        };
+
+        initializeFromUrl();
+    }, []);
+    useEffect(() => {
+        const desiredSearch = buildHealthController.toNavigationQuery().toString();
+        if (desiredSearch && location.search !== `?${desiredSearch}`) {
+            navigate(`${location.pathname}?${desiredSearch}`, { replace: true });
+        }
+    }, [navigate, buildHealthController.optionsChangeVersion, location.pathname, location.search]);
+
+    return null;
+};
+
+/**
+ * React Component that represents the Build Health sidebar. All filters are present here.
+ * @returns React Component.
+ */
+const BuildHealthSidebar: React.FC<{ handler: BuildHealthDataHandler, buildHealthOptions: BuildHealthOptionsController }> = function ConstrucBuildHealthSidebar({ handler, buildHealthOptions }) {
+    const { hordeClasses } = getHordeStyling();
+    const [jobCountFilter, setJobCountFilter] = useState<BuildCountFilter>({ numberBuilds: buildHealthOptions.state.stepOutcomeTotalJobs });
+
+    const updateJobCount = (value: number) => {
+        setJobCountFilter(prev => ({ ...prev, numberBuilds: value }));
+    };
+
+    const updateJobCountCommitted = (_event: any, value: number, _range?) => {
+        buildHealthOptions.setJobCount(value);
+    };
+
+    useEffect(() => {
+        if (buildHealthOptions.state.stepOutcomeTotalJobs) {
+            setJobCountFilter(prev => ({
+                ...prev,
+                numberBuilds: buildHealthOptions.state.stepOutcomeTotalJobs
+            }));
+        }
+    }, [buildHealthOptions.state.stepOutcomeTotalJobs]);
+
+    return (
+        <Stack style={{ minWidth: 300, maxWidth: 450, paddingRight: 18 }}>
+            <Stack className={hordeClasses.modal}>
+                <Stack>
+                    <Stack style={{ paddingTop: 0, paddingBottom: 4 }}>
+                        <Label>Step Outcome Filter</Label>
+                    </Stack>
+                    <ProjectDropdownSingle handler={handler} buildHealthOptions={buildHealthOptions} />
+                    <StreamDropdownSingle handler={handler} buildHealthOptions={buildHealthOptions} />
+                    <TemplateDropdown handler={handler} buildHealthOptions={buildHealthOptions} />
+                    <StepDropdown handler={handler} buildHealthOptions={buildHealthOptions} />
+                    <Stack>
+                        <Stack horizontal verticalAlign="center" tokens={{ childrenGap: 8 }}>
+                            <TooltipHost content="Adjust the maximum number of jobs per template to show.">
+                                <Label>Max Jobs:</Label>
+                            </TooltipHost>
+                            <Text>{jobCountFilter.numberBuilds}</Text>
+                        </Stack>
+                        <Slider min={0} max={90} value={jobCountFilter.numberBuilds} showValue={false} onChange={updateJobCount} onChanged={updateJobCountCommitted} styles={{ root: { width: 300 }, slideBox: { paddingLeft: 2 } }} />
+                    </Stack>
+                </Stack>
+            </Stack>
+        </Stack>
+    );
+};
+
+type BuildCountFilter = {
+    numberBuilds: number;
+}
+
+/**
+ * React Component representing the entire Build Health view.
+ * @returns React Component.
+ */
+export const BuildHealthView: React.FC = observer(function ConstructBuildHealthView() {
+    const windowSize = useWindowSize();
+    const { hordeClasses, modeColors } = getHordeStyling();
+
+    const filter = useMemo(() => ({
+        streams: buildHealthState.stepOutcomeEnabledStreamKeys,
+        jobs: buildHealthState.stepOutcomeEnabledJobKeys,
+        steps: buildHealthState.stepOutcomeEnabledStepKeys,
+        changelistCount: buildHealthState.stepOutcomeTotalJobs,
+    }), [
+        buildHealthState.stepOutcomeEnabledStreamKeys,
+        buildHealthState.stepOutcomeEnabledJobKeys,
+        buildHealthState.stepOutcomeEnabledStepKeys,
+        buildHealthState.stepOutcomeTotalJobs,
+    ]);
+
+    return (
+        <Stack className={hordeClasses.horde}>
+            <BuildHealthUrlSync buildHealthController={buildHealthController} />
+            <TopNav />
+            <Breadcrumbs items={[{ text: 'Build Health' }]} />
+            <Stack horizontal>
+                <div key={`windowsize_automationview_${windowSize.width}_${windowSize.height}`} style={{ width: 0, flexShrink: 0, backgroundColor: modeColors.background }} />
+                <Stack horizontalAlign="center" grow styles={{ root: { width: "100%", padding: 12, backgroundColor: modeColors.background } }}>
+                    <Stack styles={{ root: { width: "100%" } }}>
+                        <Stack horizontal styles={{ root: { minHeight: '85vh' } }} >
+                            <BuildHealthSidebar handler={handler} buildHealthOptions={buildHealthController} />
+                            <Stack id="parent-buildhealth-stepoutcome" grow style={{ overflowX: "auto", overflowY: "visible", minWidth: "1128px", position: "relative" }}>
+                                <StepOutcomeView filter={filter} />
+                            </Stack>
+                        </Stack>
+                    </Stack>
+                </Stack>
+            </Stack>
+        </Stack>
+    );
+});
+
+// #endregion -- Primary View Components --
+
+// #region -- Script --
+
+const buildHealthState = new BuildHealthOptionsState();
+const buildHealthController = new BuildHealthOptionsController(buildHealthState);
+const handler = new BuildHealthDataHandler(buildHealthState, buildHealthController);
+
+// #endregion -- Script --
\ No newline at end of file
diff --git a/Engine/Source/Programs/Horde/HordeDashboard/src/components/buildhealth/stepoutcome/StepOutcome.tsx b/Engine/Source/Programs/Horde/HordeDashboard/src/components/buildhealth/stepoutcome/StepOutcome.tsx
new file mode 100644
--- /dev/null
+++ b/Engine/Source/Programs/Horde/HordeDashboard/src/components/buildhealth/stepoutcome/StepOutcome.tsx
@@ -0,0 +1,722 @@
+// Copyright Epic Games, Inc. All Rights Reserved.
+
+import { Stack, DetailsList, IColumn, mergeStyleSets, SelectionMode, TooltipHost, DetailsListLayoutMode, ConstrainMode, Dialog, DialogType } from "@fluentui/react";
+import { observer } from "mobx-react-lite";
+import { getHordeStyling } from 'horde/styles/Styles';
+import React, { useEffect, useRef, useState } from 'react';
+import { Text } from '@fluentui/react/lib/Text';
+import { StepOutcomeTable, StepOutcomeTableEntry } from "./StepOutcomeDataTypes";
+import { StepOutcomeFilters, StepOutcomeDataHandler } from "./StepOutcomeDataHandler";
+
+// #region -- UI Helper Methods --
+
+/**
+ * Produces Fluent UI compatible columns from the provided Step Outcome Table.
+ * @param stepOutcomeTable The step outcome table to use in producing Fluent UI columns.
+ * @returns The Fluent UI compatible columns.
+ */
+function getFluentColumns(stepOutcomeTable: StepOutcomeTable): IColumn[] {
+    const columns: IColumn[] = [
+        {
+            key: 'stepName',
+            name: '',
+            fieldName: 'stepName',
+            minWidth: 300,
+            onRender: (item) => {
+                return (
+                    <TooltipHost content={`Step: ${item.stepName}`}>
+                        <div style={{ whiteSpace: 'normal', wordWrap: 'break-word', minHeight: 20, maxHeight: 20 }}>
+                            {item.stepName}
+                        </div>
+                    </TooltipHost>
+                );
+            },
+        }
+    ];
+
+    for (let i: number = 0; i < stepOutcomeTable.changeColHeaders.length; ++i) {
+        let change = stepOutcomeTable.changeColHeaders[i];
+        columns.push({
+            key: change.toString(),
+            name: change.toString(),
+            fieldName: change.toString(),
+            minWidth: 10,
+            onRender: (item: any) => {
+                const entry = item[change.toString()] as StepOutcomeTableEntry | null;
+                let bgColor = 'transparent';
+
+                if (entry) {
+                    let retrievedCellStype = getCellStyle(entry);
+                    bgColor = retrievedCellStype.bgColor;
+                }
+
+                return (
+                    <div
+                        title={
+                            `Change: ${change}\nJob: ${entry?.jobName}\nStep: ${entry?.stepResponse.name}`}
+                        className={stepOutcomeTableClasses.fillLink}
+                        style={{
+                            backgroundColor: bgColor,
+                            minHeight: 20, maxHeight: 20,
+                            minWidth: 20, maxWidth: 20,
+                        }}
+                    >
+                        &nbsp;
+                    </div>
+                );
+            }
+        });
+    }
+
+    return columns;
+}
+
+/**
+ * Produces Fluent UI compatible rows from the provided Step Outcome Table.
+ * @param stepOutcomeTable The step outcome table to use in producing Fluent UI rows.
+ * @returns The Fluent UI compatible rows.
+ */
+function getFluentRows(stepOutcomeTable: StepOutcomeTable): any[] {
+    let items: StepOutcomeFluentRow[] = [];
+
+    for (let i: number = 0; i < stepOutcomeTable.stepNameRowHeaders.length; ++i) {
+        let rows: StepOutcomeTableEntry[] = [...handler.stepOutcomeTableData.getStepOutputRow(stepOutcomeTable.stepNameRowHeaders[i].toLocaleLowerCase())];
+        let fluentRowEntry: StepOutcomeFluentRow = { 'stepName': stepOutcomeTable.stepNameRowHeaders[i] };
+        for (let y: number = 0; y < rows.length; ++y) {
+            let tableEntry: StepOutcomeTableEntry = rows[y];
+            if (tableEntry.change === undefined) {
+                continue;
+            }
+            fluentRowEntry[tableEntry.change!.toString()] = tableEntry;
+        }
+        items.push(fluentRowEntry);
+    }
+
+    return items;
+}
+
+// #endregion -- UI Helper Methods --
+
+// #region -- View Components --
+
+type StepOutcomeFluentRow = {
+    stepName: string;
+    [change: string]: StepOutcomeTableEntry | string | undefined;
+};
+
+type StepOutcomeTableProps = {
+    filters: StepOutcomeFilters;
+    handler: StepOutcomeDataHandler;
+    onCellSelected: (item: StepOutcomeTableEntry) => void;
+};
+
+interface StepOutcomeViewProps {
+    filter: StepOutcomeFilters;
+    refreshTime?: number;
+}
+
+// #region -- Styles  --
+
+/**
+ * Defines the style for column headers.
+ */
+
+const stepOutcomeTableClasses = mergeStyleSets({
+    stepOutcomeTableHeaders: {
+        selectors: {
+            '.ms-DetailsList-headerWrapper': {
+                position: 'sticky'
+            },
+            '.ms-DetailsHeader-cellTitle': {
+                writingMode: 'vertical-rl',
+                transform: 'rotate(190deg)',
+                whiteSpace: 'nowrap',
+                textAlign: 'center',
+                padding: '4px 2px',
+                paddingBottom: '10px',
+                width: '36px !important',
+                userSelect: 'text',
+            },
+            '.ms-DetailsHeader': {
+                minHeight: '80px',
+                lineHeight: 'normal',
+            },
+            '.ms-DetailsHeader-cell': {
+                verticalAlign: 'bottom'
+            },
+            '.ms-DetailsHeader-cell span': {
+                overflow: 'visible',
+            },
+            '.ms-List-cell': {
+                minHeight: 40,
+                height: 40,
+            },
+            '.ms-SelectionZone': {
+                display: 'none',
+            },
+
+        },
+    },
+    removeCellContentPadding: {
+        selectors: {
+            '.ms-DetailsRow-cell': {
+                paddingLeft: '0 !important'
+            }
+        }
+    },
+    hideGroupCollapseButton: {
+        selectors: {
+            'button.ms-GroupHeader-expand': {
+                display: 'none',
+            },
+            '.ms-GroupHeader': {
+                height: '35px'
+            }
+        },
+    },
+    hideGroupHeaderTitle: {
+        selectors: {
+            '.ms-GroupHeader': {
+                height: '35px'
+            },
+            '.ms-GroupHeader-title': {
+                display: 'none',
+            },
+            'button.ms-GroupHeader-expand': {
+                display: 'none',
+            },
+            '.ms-GroupSpacer': {
+                width: '0 !important',
+            },
+            '.ms-DetailsRow-cell': {
+                paddingLeft: '0 !important'
+            }
+        },
+    },
+    disableGroupCollapseButton: {
+        selectors: {
+            '.ms-GroupHeader-dropIcon': {
+                display: 'none',
+                pointerEvents: 'none',
+                opacity: 0.5,
+            },
+        },
+    },
+    stickyHeaderOverride: {
+        selectors: {
+            '.ms-DetailsHeader-cellTitle': {
+                writingMode: 'vertical-rl',
+                transform: 'rotate(190deg)',
+                whiteSpace: 'nowrap',
+                textAlign: 'center',
+                padding: '4px 2px',
+                paddingBottom: '10px',
+                width: '36px !important',
+            },
+            '.ms-DetailsHeader': {
+                minHeight: '100px',
+            },
+        },
+    },
+    fillLink: {
+        color: 'inherit',
+        padding: '2px',
+        boxSizing: 'border-box',
+    },
+});
+
+// #endregion -- Styles  --
+
+
+// #region -- Modal Popup View --
+
+type StepOutcomeFocusModalProps = {
+    item: StepOutcomeTableEntry;
+    onDismiss?: (_ev?: React.MouseEvent<HTMLButtonElement>) => any;
+};
+
+/**
+ * Normalizes an ambiguous input (Date | string) to a Date object. 
+ * @param input 
+ * @returns 
+ */
+function toDate(input?: Date | string): Date | undefined {
+    if (!input) return undefined;
+
+    if (input instanceof Date) {
+        return input;
+    }
+
+    const parsed = new Date(input);
+    return isNaN(parsed.getTime()) ? undefined : parsed;
+}
+
+/**
+ * Generates a formatted duration provided a start and finish time.
+ * @param startISO The start time, in string representation.
+ * @param finishISO The finish time, in string representation.
+ * @returns The formated duration.
+ */
+function formatDuration(startISO: string, finishISO: string): string {
+    const start = new Date(startISO);
+    const finish = new Date(finishISO);
+    const diffMs = finish.getTime() - start.getTime();
+    const diffSeconds = Math.floor(diffMs / 1000);
+
+    const hours = Math.floor(diffSeconds / 3600);
+    const minutes = Math.floor((diffSeconds % 3600) / 60);
+    const seconds = diffSeconds % 60;
+
+    return `${hours}h ${minutes}m ${seconds}s`;
+}
+
+/**
+ * React Component representing the Step Outcome focus.
+ */
+const StepOutcomeFocusModal: React.FC<StepOutcomeFocusModalProps> = ({ item, onDismiss }) => {
+    const startDate = toDate(item.stepResponse.startTime);
+    const finishDate = toDate(item.stepResponse.finishTime);
+    let duration: string = "Ongoing";
+
+    if (startDate && finishDate) {
+        duration = formatDuration(startDate.toISOString(), finishDate.toISOString());
+    }
+
+    return (
+        <Dialog
+            modalProps={{
+                isBlocking: false,
+                topOffsetFixed: true,
+                styles: {
+                    root: {
+                        selectors: {
+                            ".ms-Dialog-title": {
+                                paddingTop: '24px',
+                                paddingLeft: '32px'
+                            }
+                        }
+                    }
+                }
+            }}
+            onDismiss={onDismiss}
+            hidden={false}
+            minWidth={1200}
+            dialogContentProps={{
+                type: DialogType.close,
+                onDismiss: onDismiss,
+                title: item.stepResponse.name,
+            }}>
+            <Stack tokens={{ childrenGap: 8 }}>
+                {[
+                    ['Date', startDate?.toDateString()],
+                    ['Stream', item.streamId],
+                    ['Code Change', item.change],
+                    ['Job', item.jobName],
+                    ['Step', item.stepResponse.name],
+                    ['Duration', duration],
+                    ['Outcome', item.stepResponse.outcome],
+                    ['Job Link', <a href={`/job/${item.jobId}`} target="_blank" rel="noopener noreferrer">View Job</a>],
+                    ['Step Link', <a href={`/job/${item.jobId}?step=${item.stepResponse.id}`} target="_blank" rel="noopener noreferrer">View Step</a>]
+                ].map(([label, value]) => (
+                    <Stack horizontal tokens={{ childrenGap: 12 }}>
+                        <Text styles={{ root: { fontWeight: '600', minWidth: 100 } }}>{label}:</Text>
+                        <Text>{value}</Text>
+                    </Stack>
+                ))}
+            </Stack>
+        </Dialog>
+    );
+};
+
+// #endregion -- Modal Popup View --
+
+/**
+ * Gets the cell style given the provided StepOutcomeTableEntry.
+ * @param entry The entry to base the style off of.
+ * @returns The matching style given the state of the entry. 
+ */
+function getCellStyle(entry: StepOutcomeTableEntry | null): { bgColor: string } {
+    if (!entry) return { bgColor: "grey" };
+
+    const { stepResponse } = entry;
+    const { outcome, state } = stepResponse;
+
+    let bgColor = "grey";
+    if (["Aborted", "Completed", "Skipped"].includes(state)) {
+        if (outcome === "Failure") bgColor = state === "Skipped" ? "black" : "red";
+        else if (outcome === "Warnings") bgColor = "yellow";
+        else if (outcome === "Success") bgColor = "lightgreen";
+    }
+
+    return { bgColor };
+}
+
+/**
+ * React Component that represents the provided stepOutcomeTable. This is a standard HTML table.
+ * @param stepOutcomeTable The data to base the table off of.
+ * @returns React Component.
+ */
+function GenerateClassicTable(stepOutcomeTable: StepOutcomeTable): JSX.Element {
+    const { modeColors } = getHordeStyling();
+    return (
+        <div style={{ overflow: 'auto', maxHeight: '600px', maxWidth: '100%' }}>
+            <table
+                style={{
+                    borderCollapse: 'separate',
+                    borderSpacing: '2px',
+                    width: '100%',
+                    tableLayout: 'fixed', // Fix column widths strictly
+                }}
+            >
+                <thead>
+                    <tr>
+                        <th
+                            style={{
+                                position: 'sticky',
+                                top: 0,
+                                left: 0,
+                                zIndex: 3,
+                                backgroundColor: modeColors.background,
+                                maxWidth: 200,
+                                minWidth: 180,
+                                width: 180,
+                                whiteSpace: 'nowrap',
+                                overflow: 'hidden',
+                                textOverflow: 'ellipsis',
+                                height: '100px',
+                            }}
+                        >
+                            Step
+                        </th>
+                        {stepOutcomeTable.changeColHeaders.slice(0, 20).map((col, idx) => (
+                            <th
+                                key={idx}
+                                style={{
+                                    position: 'sticky',
+                                    top: 0,
+                                    zIndex: 2,
+                                    backgroundColor: modeColors.background,
+                                    writingMode: 'vertical-rl',
+                                    textAlign: 'left',
+                                    whiteSpace: 'nowrap',
+                                    height: '40px',
+                                    width: '40px',
+                                }}
+                            >
+                                <div
+                                    style={{
+                                        transform: 'rotate(180deg)',
+                                        display: 'inline-block',
+                                        padding: '8px 0',
+                                    }}
+                                >
+                                    {col}
+                                </div>
+                            </th>
+                        ))}
+                    </tr>
+                </thead>
+                <tbody>
+                    {stepOutcomeTable.stepNameRowHeaders.slice(0, 20).map((key, rIdx) => (
+                        <tr key={rIdx} style={{ height: '40px' }}>
+                            <td
+                                title={key}
+                                style={{
+                                    position: 'sticky',
+                                    left: 0,
+                                    zIndex: 1,
+                                    backgroundColor: modeColors.background,
+                                    whiteSpace: 'nowrap',
+                                    overflow: 'hidden',
+                                    textOverflow: 'ellipsis',
+                                    maxWidth: 200,
+                                    minWidth: 180,
+                                    width: 180,
+                                    height: '40px',
+                                }}
+                            >
+                                {key}
+                            </td>
+                            {[...stepOutcomeTable.getStepOutputRow(key.toLocaleLowerCase())]
+                                .slice(0, 20)
+                                .map((entry, cIdx) => {
+                                    const { bgColor } = getCellStyle(entry);
+                                    return (
+                                        <td
+                                            key={cIdx}
+                                            style={{
+                                                backgroundColor: bgColor,
+                                                padding: 0,
+                                                margin: 0,
+                                                width: '40px',
+                                                height: '40px',
+                                                textAlign: 'center',
+                                            }}
+                                        >
+                                            <div
+                                                style={{
+                                                    display: 'block',
+                                                    width: '100%',
+                                                    height: '100%',
+                                                    textDecoration: 'none',
+                                                }}
+                                            >
+                                                &nbsp;
+                                            </div>
+                                        </td>
+                                    );
+                                })}
+                        </tr>
+                    ))}
+                </tbody>
+            </table>
+        </div>
+
+    );
+}
+
+/**
+ * React Component that represents the provided stepOutcomeTable. This is a Fluent UI table.
+ * @param stepOutcomeTable The data to base the table off of.
+ * @returns React Component.
+ */
+function GenerateFluentUITable(stepOutcomeTable: StepOutcomeTable, onCellClick: any): JSX.Element {
+    /**
+     * This is a fairly "virtual" (aka hacky) table in the following respects:
+     * - We need to have a sticky header, and first column
+     * - Sticky component works well for header, but not at all for first column
+     * - As a result, we have opted to take the following strategy:
+     *      - Create a dual "panel" view:
+     *          - Left panel is the 1st column; the Step Names
+     *          - Right panel - Top Stack is the Headers for all the changes
+     *          - Right panel - Bottom Stack is the cell data for all the steps
+     *      - Synchronize the Left panel vertical scroll to the Right Panel - Bottom Stack vertical scroll
+     *      - Synchronize the Rigth Panel Top Stack horizontal scroll to the Right Panel - Bottom Stack horizontal scroll
+     * 
+     * Whilst it's a tad hacky, it's predictable and easy to reason.
+     */
+
+    const leftPaneRef = useRef<HTMLDivElement>(null);
+    const rightTopPaneRef = useRef<HTMLDivElement>(null);
+    const rightBottomPaneRef = useRef<HTMLDivElement>(null);
+
+    const isSyncingVertical = useRef(false);
+    const isSyncingHorizontal = useRef(false);
+
+    const allColumns = getFluentColumns(stepOutcomeTable);
+    const items = getFluentRows(stepOutcomeTable);
+
+    const firstColumn = allColumns.slice(0, 1);
+    const restColumns = allColumns.slice(1);
+
+    const [hasScrollbar, setHasScrollbar] = useState(false);
+
+    useEffect(() => {
+        const checkScroll = () => {
+            const el = rightBottomPaneRef.current;
+            if (el) setHasScrollbar(el.scrollWidth > el.clientWidth);
+        };
+
+        checkScroll();
+        const resizeObserver = new ResizeObserver(checkScroll);
+        rightBottomPaneRef.current && resizeObserver.observe(rightBottomPaneRef.current);
+        return () => resizeObserver.disconnect();
+    }, []);
+
+    const handleScroll = (direction: 'vertical' | 'horizontal', source: 'left' | 'rightTop' | 'rightBottom') =>
+        (e: React.UIEvent<HTMLDivElement>) => {
+            const scrollValue = direction === 'vertical' ? e.currentTarget.scrollTop : e.currentTarget.scrollLeft;
+            const syncRef = direction === 'vertical' ? isSyncingVertical : isSyncingHorizontal;
+
+            if (syncRef.current) return;
+            syncRef.current = true;
+
+            const targetMap = {
+                vertical: {
+                    left: rightBottomPaneRef,
+                    rightBottom: leftPaneRef,
+                },
+                horizontal: {
+                    rightTop: rightBottomPaneRef,
+                    rightBottom: rightTopPaneRef,
+                },
+            };
+
+            const targetRef = targetMap[direction][source];
+            if (targetRef?.current) {
+                if (direction === 'vertical') {
+                    targetRef.current.scrollTop = scrollValue;
+                } else {
+                    targetRef.current.scrollLeft = scrollValue;
+                }
+            }
+
+            requestAnimationFrame(() => {
+                syncRef.current = false;
+            });
+        };
+
+    const gridStyle = {
+        display: 'grid',
+        gridTemplateRows: 'auto 1fr',
+        gridTemplateColumns: 'auto 1fr',
+        maxHeight: '85vh',
+        overflow: 'hidden',
+    };
+
+    return (
+        <div style={gridStyle}>
+            {/* Left Panel - Left Column - Step Name - (sticky) */}
+            <div
+                ref={leftPaneRef}
+                onScroll={handleScroll('vertical', 'left')}
+                className="no-scrollbar"
+                style={{
+                    maxHeight: '100%',
+                    overflowY: 'auto',
+                    position: 'relative',
+                    gridRow: 2,
+                    gridColumn: 1,
+                    overflowX: 'scroll',
+                    paddingBottom: hasScrollbar ? '17px' : '0',
+                    zIndex: 1,
+                    scrollbarWidth: 'none',
+                    msOverflowStyle: 'none',
+                }}
+            >
+                <DetailsList
+                    items={items}
+                    columns={firstColumn}
+                    selectionMode={SelectionMode.none}
+                    layoutMode={DetailsListLayoutMode.fixedColumns}
+                    constrainMode={ConstrainMode.unconstrained}
+                    isHeaderVisible={false}
+                    className={stepOutcomeTableClasses.hideGroupCollapseButton}
+                />
+            </div>
+
+            {/* Right Panel - Top Stack - Change Header Row (sticky) */}
+            <div
+                ref={rightTopPaneRef}
+                onScroll={handleScroll('horizontal', 'rightTop')}
+                className="no-scrollbar"
+                style={{
+                    gridRow: 1,
+                    gridColumn: 2,
+                    overflowX: 'auto',
+                    minHeight: '80px',
+                    maxHeight: '400px',
+                    position: 'relative',
+                    zIndex: 2,
+                }}
+            >
+                <DetailsList
+                    items={[]} // header only
+                    columns={restColumns}
+                    isHeaderVisible={true}
+                    selectionMode={SelectionMode.none}
+                    className={stepOutcomeTableClasses.stepOutcomeTableHeaders}
+                    layoutMode={DetailsListLayoutMode.fixedColumns}
+                    constrainMode={ConstrainMode.unconstrained}
+                />
+            </div>
+
+            {/* Right Panel - Bottom Stack - Step Data Cells (scroll driver) */}
+            <div
+                ref={rightBottomPaneRef}
+                onScroll={(e) => {
+                    handleScroll('vertical', 'rightBottom')(e);
+                    handleScroll('horizontal', 'rightBottom')(e);
+                }}
+                style={{
+                    maxHeight: '100%',
+                    overflowY: 'auto',
+                    position: 'relative',
+                    gridRow: 2,
+                    gridColumn: 2,
+                    overflow: 'auto',
+                }}
+            >
+                <DetailsList
+                    items={items}
+                    columns={restColumns}
+                    selectionMode={SelectionMode.none}
+                    layoutMode={DetailsListLayoutMode.fixedColumns}
+                    constrainMode={ConstrainMode.unconstrained}
+                    isHeaderVisible={false}
+                    className={stepOutcomeTableClasses.hideGroupHeaderTitle}
+                    onRenderField={(fieldProps, defaultRender) => {
+
+                        if (!fieldProps) {
+                            return null;
+                        }
+                        const { item, column, itemIndex } = fieldProps;
+                        const cellValue = column.fieldName ? item[column.fieldName] : undefined;
+                        const stepData: StepOutcomeTableEntry = cellValue as StepOutcomeTableEntry;
+                        const id = `cell-${itemIndex}-${column.key}`;
+
+                        return (
+                            <div id={id} key={id} style={{ paddingLeft: 0, cursor: 'pointer' }} onClick={() => onCellClick?.(stepData)}>
+                                {defaultRender?.(fieldProps)}
+                            </div>
+                        );
+                    }}
+                />
+            </div>
+        </div>
+    );
+};
+
+/**
+ * React Component that acts as a observable binder for data handler updates, and allows for easy swapping of other table types.
+ */
+const StepOutcomeTableComponent: React.FC<StepOutcomeTableProps> = observer(({ handler, onCellSelected }) => {
+    handler.updated;
+    return GenerateFluentUITable(handler.stepOutcomeTableData, onCellSelected);
+});
+
+/**
+ * React Component that represents the overall Step Outcome Table view.
+ * @param StepOutcomeViewProps The filter to use for the step outcome, and the refresh cadence of the data.
+ * @returns React Component.
+ */
+export const StepOutcomeView: React.FC<StepOutcomeViewProps> = ({ filter, refreshTime }) => {
+    useEffect(() => {
+        handler.start();
+
+        return () => {
+            handler.stop();
+        };
+
+    }, []);
+
+    useEffect(() => {
+        handler.setFilter(filter);
+    }, [filter]);
+
+    useEffect(() => {
+        if (refreshTime) {
+            handler.setRefreshTime(refreshTime);
+        }
+    }, [refreshTime]);
+
+    const [selectedItem, setSelectedItem] = useState<StepOutcomeTableEntry | null>(null);
+
+    return (
+        <Stack>
+            <Stack>
+                {selectedItem && (<StepOutcomeFocusModal item={selectedItem} onDismiss={() => setSelectedItem(null)} />)}
+                <StepOutcomeTableComponent handler={handler} filters={filter} onCellSelected={setSelectedItem} />
+            </Stack>
+        </Stack>
+    );
+};
+
+// #endregion -- View Components --
+
+// #region -- Script --
+
+const DEFAULT_POLL_TIME = 60000;
+const handler = new StepOutcomeDataHandler(DEFAULT_POLL_TIME);
+
+// #endregion -- Script --
\ No newline at end of file
diff --git a/Engine/Source/Programs/Horde/HordeDashboard/src/components/buildhealth/stepoutcome/StepOutcomeDataHandler.ts b/Engine/Source/Programs/Horde/HordeDashboard/src/components/buildhealth/stepoutcome/StepOutcomeDataHandler.ts
new file mode 100644
--- /dev/null
+++ b/Engine/Source/Programs/Horde/HordeDashboard/src/components/buildhealth/stepoutcome/StepOutcomeDataHandler.ts
@@ -0,0 +1,134 @@
+// Copyright Epic Games, Inc. All Rights Reserved.
+
+import { GetJobResponse, JobStreamQuery } from "horde/backend/Api";
+import { PollBase } from "horde/backend/PollBase";
+import { StepOutcomeTable, StepOutcomeTableEntry } from "./StepOutcomeDataTypes";
+import backend from "horde/backend";
+
+// #region -- API Data Types & Utilities --
+
+/**
+ * Data structure that controls the filtering & queries of the Data Handler.
+ */
+export interface StepOutcomeFilters {
+    streams: string[];
+    jobs?: string[];
+    steps?: string[];
+    dateRange?: { start: Date; end: Date };
+    changelistCount?: number;
+}
+
+// #endregion -- API Data Types & Utilities --
+
+/**
+ * Data handler for all step outcome data used for step outcomes.
+ * @todo Data sharing & caching would ideally occur between the BuildHealth View & Step Outcome.
+ */
+export class StepOutcomeDataHandler extends PollBase {
+    private readonly STREAM_JOBS_FILTER = "id,streamId,name,change,templateId,state,createTime,updateTime,batches,preflightDescription,cancellationReason";
+
+    // #region -- Private Members -- 
+
+    activeFilter: StepOutcomeFilters;
+    jobData: GetJobResponse[] = [];
+    stepOutcomeTableData: StepOutcomeTable = new StepOutcomeTable();
+
+    //#endregion -- Private Members
+
+    // #region -- Constructor --
+
+    /**
+     * Cosntructs a StepOutcomeDataHandler.
+     * @param pollTime The initial poll refresh time.
+     */
+    constructor(pollTime: number) {
+        super(pollTime);
+    }
+
+    //#endregion -- Constructor --
+
+    // #region -- Interface --
+
+    async poll(): Promise<void> {
+        if (
+            this.activeFilter === undefined ||
+            !this.activeFilter.streams?.length ||
+            !this.activeFilter.jobs?.length
+        ) {
+            this.stepOutcomeTableData.reset();
+            return;
+        }
+
+        this.stepOutcomeTableData.reset();
+
+        const stepSet = new Set(this.activeFilter.steps?.map(s => s.toLowerCase()));
+        const jobList = this.activeFilter.jobs ?? [];
+        const streamList = this.activeFilter.streams ?? [];
+        const changelistCount = this.activeFilter.changelistCount ?? 90;
+        const queries: { streamId: string; query: JobStreamQuery }[] = [];
+
+        for (const streamId of streamList) {
+            for (const job of jobList) {
+                const query: JobStreamQuery = {
+                    template: [job],
+                    count: changelistCount,
+                    filter: this.STREAM_JOBS_FILTER
+                };
+                queries.push({ streamId, query });
+            }
+        }
+
+        // Parallelize backend calls
+        const results = await Promise.all(
+            queries.map(({ streamId, query }) => backend.getStreamJobs(streamId, query))
+        );
+
+        // Flatten all jobData
+        this.jobData = results.flat();
+
+        // Process all job responses
+        for (const jobResponse of this.jobData) {
+            const batchResponses = jobResponse.batches ?? [];
+            for (const batch of batchResponses) {
+                for (const step of batch.steps) {
+                    if (
+                        stepSet.size === 0 ||
+                        stepSet.has(step.name.toLowerCase())
+                    ) {
+                        this.stepOutcomeTableData.addEntry(
+                            new StepOutcomeTableEntry(step, jobResponse.id, jobResponse.name, jobResponse.streamId, jobResponse.change)
+                        );
+                    }
+                }
+            }
+        }
+
+        this.stepOutcomeTableData.orderTableDataByChange(true);
+        this.setUpdated();
+
+        return Promise.resolve();
+    }
+
+    // #endregion -- Interface --
+
+    // #region -- Public API  --
+
+    /**
+     * Sets a new filter for the handler, which will control the underlying dataset.
+     * @param filter The new filter to use.
+     */
+    setFilter(filter: StepOutcomeFilters) {
+        this.activeFilter = filter;
+        this.poll();
+    }
+
+    /**
+     *  Sets a new refresh time for the handler, in ms.
+     * @param newRefreshTime The new refresh rate.
+     */
+    setRefreshTime(newRefreshTime: number) {
+        this.pollTime = newRefreshTime;
+    }
+
+    // #endregion -- Public API  --
+}
\ No newline at end of file
diff --git a/Engine/Source/Programs/Horde/HordeDashboard/src/components/buildhealth/stepoutcome/StepOutcomeDataTypes.ts b/Engine/Source/Programs/Horde/HordeDashboard/src/components/buildhealth/stepoutcome/StepOutcomeDataTypes.ts
new file mode 100644
--- /dev/null
+++ b/Engine/Source/Programs/Horde/HordeDashboard/src/components/buildhealth/stepoutcome/StepOutcomeDataTypes.ts
@@ -0,0 +1,218 @@
+// Copyright Epic Games, Inc. All Rights Reserved.
+
+import { GetStepResponse } from "horde/backend/Api";
+
+/**
+ * Model for a table entry.
+ */
+export class StepOutcomeTableEntry {
+    // #region -- Public Members --
+
+    change?: number;
+    stepResponse: GetStepResponse;
+    jobId: string;
+    jobName: string;
+    streamId: string;
+
+    // #endregion -- Public Members --
+
+    // #region -- Constructor --
+
+    /**
+     * Constructor.
+     * @param stepResponse The step response object for the table entry.
+     * @param jobId The owning job's Id.
+     * @param jobName The owning job's name.
+     * @param streamId The owning streams's Id.
+     * @param change The change associated with the stepResponse.
+     */
+    constructor(stepResponse: GetStepResponse, jobId: string, jobName: string, streamId: string, change?: number) {
+        this.stepResponse = stepResponse;
+        if (change !== undefined) {
+            this.change = change;
+        }
+        this.jobId = jobId;
+        this.jobName = jobName;
+        this.streamId = streamId;
+    }
+
+    // #endregion -- Constructor --
+}
+
+/**
+ * Model for the StepOutcomeTable.
+ */
+export class StepOutcomeTable {
+    // #region -- Public Members --
+
+    /**
+     * Change column headers.
+     */
+    changeColLookup: Map<number, number> = new Map<number, number>();
+    changeColHeaders: number[] = [];
+
+    /**
+     * Step name row headers.
+     */
+    stepRowLookup: Map<string, number> = new Map<string, number>();
+    stepNameRowHeaders: string[] = [];
+
+    /**
+     * Table data structure
+     */
+    tableEntries: (StepOutcomeTableEntry | null)[][] = [];
+
+    // #endregion --  Public Members --
+
+    // #region -- Private API --
+
+    private ensureRow(source: (StepOutcomeTableEntry | null)[], targetSize: number) {
+        while (targetSize > source.length) {
+            source.push(null);
+        }
+    }
+
+    // #endregion -- Private API --
+
+    // #region -- Public API --
+
+    /**
+     * Resets the table data. This will remove all data from the container.
+     */
+    reset(): void {
+        this.tableEntries = [];
+        this.stepNameRowHeaders = [];
+        this.stepRowLookup.clear();
+        this.changeColHeaders = [];
+        this.changeColLookup.clear();
+    }
+
+    /**
+     * Adds a new entry to the table data.
+     * @param tableEntry The table entry to add to the data table.
+     * @returns If tableEntry parameter is invalid, returns.
+     */
+    addEntry(tableEntry: StepOutcomeTableEntry): void {
+        if (tableEntry.change === undefined) {
+            return;
+        }
+
+        const normalizedStepName = tableEntry.stepResponse.name.toLocaleLowerCase();
+
+        // Resolve or create row index
+        let insertionRow = this.stepRowLookup.get(normalizedStepName);
+        if (insertionRow === undefined) {
+            insertionRow = this.stepNameRowHeaders.length;
+            this.stepRowLookup.set(normalizedStepName, insertionRow);
+            this.stepNameRowHeaders.push(tableEntry.stepResponse.name);
+        }
+
+        // Resolve or create column index
+        let insertionCol = this.changeColLookup.get(tableEntry.change);
+        if (insertionCol === undefined) {
+            insertionCol = this.changeColHeaders.length;
+            this.changeColLookup.set(tableEntry.change, insertionCol);
+            this.changeColHeaders.push(tableEntry.change);
+        }
+
+        // Ensure the row exists and is padded
+        if (this.tableEntries[insertionRow] === undefined) {
+            this.tableEntries[insertionRow] = Array(this.changeColHeaders.length).fill(null);
+        } else {
+            this.ensureRow(this.tableEntries[insertionRow], this.changeColHeaders.length);
+        }
+
+        const existingEntry = this.tableEntries[insertionRow][insertionCol];
+        if (existingEntry !== null) {
+            console.warn(`Duplicate entry for step "${tableEntry.stepResponse.name}" and change "${tableEntry.change}". Ambiguous selector needed.`, {
+                existingEntry,
+                newEntry: tableEntry
+            });
+            return;
+        }
+
+        this.tableEntries[insertionRow][insertionCol] = tableEntry;
+    }
+
+    /**
+     * Gets the Table Entries for a given step name.
+     * @param rowName The step name to obtain.
+     * @returns * Generator of TableEntries.
+     */
+    getStepOutputRow(rowName: string): Generator<StepOutcomeTableEntry>;
+    getStepOutputRow(rowIndex: number): Generator<StepOutcomeTableEntry>;
+    *getStepOutputRow(row: string | number): Generator<StepOutcomeTableEntry> {
+        let rowIdx: number | undefined;
+
+        if (typeof row === 'string') {
+            rowIdx = this.stepRowLookup.get(row);
+        } else {
+            rowIdx = row;
+        }
+
+        if (rowIdx === undefined || rowIdx < 0 || rowIdx >= this.tableEntries.length) {
+            return;
+        }
+
+        let rowEntries = this.tableEntries[rowIdx];
+
+        for (const entry of this.changeColHeaders) {
+            let colIdx = this.changeColLookup.get(entry);
+            if (colIdx === undefined) {
+                continue;
+            }
+
+            const value = rowEntries[colIdx];
+            if (value !== null && value !== undefined) {
+                yield value;
+            }
+        }
+    }
+
+    /**
+     * Gets the Table Entries for the given change.
+     * @param changeNumber The change number.
+     * @returns All Table Entries for the specific change.
+     */
+    *getChangeColumn(changeNumber: number): Generator<StepOutcomeTableEntry> {
+
+        let colIdx: (number | undefined) = this.changeColLookup.get(changeNumber);
+        if (colIdx === undefined) {
+            return;
+        }
+
+        for (const entry of this.stepNameRowHeaders) {
+            let rowIdx: (number | undefined) = this.stepRowLookup.get(entry);
+            if (rowIdx === undefined) {
+                continue;
+            }
+
+            let tableEntry: (StepOutcomeTableEntry | null) = this.tableEntries[rowIdx][colIdx];
+            if (tableEntry !== null) {
+                yield tableEntry;
+            }
+        }
+    }
+
+    /**
+     * Sets the table data's change (column) order.
+     * @param ascending True for ascending order; false for descending order.
+     */
+    orderTableDataByChange(ascending: boolean): void {
+        this.changeColHeaders.sort((a, b) => ascending ? a - b : b - a);
+    }
+
+    /**
+     * Sets the table data's step (row) order.
+     * @param ascending True for ascending order; false for descending order.
+     */
+    orderTableDataByStepName(ascending: boolean): void {
+        this.stepNameRowHeaders.sort((a, b) => {
+            if (a < b) return ascending ? -1 : 1;
+            if (a > b) return ascending ? 1 : -1;
+            return 0;
+        });
+    }
+
+    // #endregion -- Public API --
+}
\ No newline at end of file
