Horde: Remove unused setting for compute queue AWS metric in PoolConfig

This has since long moved into a separate strategy - ComputeQueueAwsMetricStrategy.

diff --git a/Engine/Source/Programs/Horde/Plugins/Compute/HordeServer.Compute/Agents/Pools/PoolCollection.cs b/Engine/Source/Programs/Horde/Plugins/Compute/HordeServer.Compute/Agents/Pools/PoolCollection.cs
--- a/Engine/Source/Programs/Horde/Plugins/Compute/HordeServer.Compute/Agents/Pools/PoolCollection.cs
+++ b/Engine/Source/Programs/Horde/Plugins/Compute/HordeServer.Compute/Agents/Pools/PoolCollection.cs
@@ -132,7 +132,6 @@ namespace HordeServer.Agents.Pools
 				}
 				LeaseUtilizationSettings = other.LeaseUtilizationSettings;
 				JobQueueSettings = other.JobQueueSettings;
-				ComputeQueueAwsMetricSettings = other.ComputeQueueAwsMetricSettings;
 				Revision = otherRevision;
 			}
 		}
@@ -204,7 +203,6 @@ namespace HordeServer.Agents.Pools
 			public IReadOnlyList<FleetManagerInfo>? FleetManagers => _config.FleetManagers;
 			public LeaseUtilizationSettings? LeaseUtilizationSettings => _config.LeaseUtilizationSettings;
 			public JobQueueSettings? JobQueueSettings => _config.JobQueueSettings;
-			public ComputeQueueAwsMetricSettings? ComputeQueueAwsMetricSettings => _config.ComputeQueueAwsMetricSettings;
 
 			#endregion
 
diff --git a/Engine/Source/Programs/Horde/Plugins/Compute/HordeServer.Compute/Agents/Pools/PoolConfig.cs b/Engine/Source/Programs/Horde/Plugins/Compute/HordeServer.Compute/Agents/Pools/PoolConfig.cs
--- a/Engine/Source/Programs/Horde/Plugins/Compute/HordeServer.Compute/Agents/Pools/PoolConfig.cs
+++ b/Engine/Source/Programs/Horde/Plugins/Compute/HordeServer.Compute/Agents/Pools/PoolConfig.cs
@@ -127,11 +127,6 @@ namespace HordeServer.Agents.Pools
 		/// Settings for job queue pool sizing strategy (if used)
 		/// </summary>
 		JobQueueSettings? JobQueueSettings { get; }
-
-		/// <summary>
-		/// Settings for job queue pool sizing strategy (if used)
-		/// </summary>
-		ComputeQueueAwsMetricSettings? ComputeQueueAwsMetricSettings { get; }
 	}
 
 	/// <summary>
diff --git a/Engine/Source/Programs/Horde/Plugins/Compute/HordeServer.Compute/Agents/Pools/PoolMessages.cs b/Engine/Source/Programs/Horde/Plugins/Compute/HordeServer.Compute/Agents/Pools/PoolMessages.cs
--- a/Engine/Source/Programs/Horde/Plugins/Compute/HordeServer.Compute/Agents/Pools/PoolMessages.cs
+++ b/Engine/Source/Programs/Horde/Plugins/Compute/HordeServer.Compute/Agents/Pools/PoolMessages.cs
@@ -553,7 +553,6 @@ namespace HordeServer.Agents.Pools
 #pragma warning restore CS0618 // Type or member is obsolete
 			LeaseUtilizationSettings = pool.LeaseUtilizationSettings == null ? null : new LeaseUtilizationSettingsMessage(pool.LeaseUtilizationSettings);
 			JobQueueSettings = pool.JobQueueSettings == null ? null : new JobQueueSettingsMessage(pool.JobQueueSettings);
-			ComputeQueueAwsMetricSettings = pool.ComputeQueueAwsMetricSettings == null ? null : new ComputeQueueAwsMetricSettingsMessage(pool.ComputeQueueAwsMetricSettings);
 			MinAgents = pool.MinAgents;
 			NumReserveAgents = pool.NumReserveAgents;
 			ShutdownIfDisabledGracePeriod = pool.ShutdownIfDisabledGracePeriod;
diff --git a/Engine/Source/Programs/Horde/Plugins/Compute/HordeServer.Compute/Server/PoolsServerController.cs b/Engine/Source/Programs/Horde/Plugins/Compute/HordeServer.Compute/Server/PoolsServerController.cs
--- a/Engine/Source/Programs/Horde/Plugins/Compute/HordeServer.Compute/Server/PoolsServerController.cs
+++ b/Engine/Source/Programs/Horde/Plugins/Compute/HordeServer.Compute/Server/PoolsServerController.cs
@@ -108,7 +108,6 @@ namespace HordeServer.Server
 				}
 				config.LeaseUtilizationSettings = currentConfig.LeaseUtilizationSettings;
 				config.JobQueueSettings = currentConfig.JobQueueSettings;
-				config.ComputeQueueAwsMetricSettings = currentConfig.ComputeQueueAwsMetricSettings;
 				configs.Add(config);
 			}
 
