Encode artifact path before querying backend

#rnx

diff --git a/Engine/Source/Programs/Horde/HordeDashboard/plugins/test-automation/testAutomationCommon.tsx b/Engine/Source/Programs/Horde/HordeDashboard/plugins/test-automation/testAutomationCommon.tsx
--- a/Engine/Source/Programs/Horde/HordeDashboard/plugins/test-automation/testAutomationCommon.tsx
+++ b/Engine/Source/Programs/Horde/HordeDashboard/plugins/test-automation/testAutomationCommon.tsx
@@ -687,17 +687,19 @@ export const PhaseArtifactsModal: React.FC<{ phase: TestPhaseStatus, artifactPat
         if (column.name === "Name") {
 
             const path = (item.name as string).split("/");
+            const start_index = path.length > 5? path.length - 5 : 0;
 
-            const pathElements = path.map((t, index) => {
-                const last = index === (path.length - 1);
+            const pathElements = path.slice(start_index).map((t, index) => {
+                const last = (index + start_index) === (path.length - 1);
                 let color = last ? (dashboard.darktheme ? "#FFFFFF" : "#605E5C") : undefined;
                 const font = last ? undefined : "Horde Open Sans Light";
                 const sep = last ? undefined : "/"
                 return <Text key={`p-${item.name}-${index}`} styles={{ root: { fontFamily: font } }} style={{ color: color }}>{t}{sep}</Text>
-            })
+            });
+            if (start_index !== 0) pathElements.splice(0, 0, <Text key={`p-${item.name}`} >... /</Text>);
 
             return <Stack verticalFill verticalAlign="center" style={{ cursor: "pointer" }}>
-                        <Stack horizontal tokens={{ childrenGap: 8 }}>
+                        <Stack horizontal tokens={{ childrenGap: 8 }} title={item.name}>
                             <Stack>
                                 <FontIcon style={{ paddingTop: 1, fontSize: 16 }} iconName="Document" />
                             </Stack>
diff --git a/Engine/Source/Programs/Horde/HordeDashboard/plugins/test-automation/testData.ts b/Engine/Source/Programs/Horde/HordeDashboard/plugins/test-automation/testData.ts
--- a/Engine/Source/Programs/Horde/HordeDashboard/plugins/test-automation/testData.ts
+++ b/Engine/Source/Programs/Horde/HordeDashboard/plugins/test-automation/testData.ts
@@ -1528,7 +1528,7 @@ export class ArtifactFactory {
 
         await this.getJobStepSavedArtifacts();
         
-        return `${backend.serverUrl}/api/v2/artifacts/${this.artifactId}/file?path=${referencePath}&inline=true`;
+        return `${backend.serverUrl}/api/v2/artifacts/${this.artifactId}/file?path=${encodeURIComponent(referencePath)}&inline=true`;
     }
 
     async fetch(referencePath: string): Promise<any> {
@@ -1554,12 +1554,12 @@ export class ArtifactFactory {
 
     async download(referencePath: string): Promise<void> {
         await this.getJobStepSavedArtifacts();
-        if (!!this.artifactId) backend.downloadArtifactV2(this.artifactId, decodeURIComponent(referencePath));
+        if (!!this.artifactId) backend.downloadArtifactV2(this.artifactId, referencePath);
     }
 
     async downloadZip(referencePaths: string[]) : Promise<void> {
         await this.getJobStepSavedArtifacts();
-        if (!!this.artifactId) backend.downloadArtifactZipV2(this.artifactId, {filter: referencePaths.map(r => decodeURIComponent(r))})
+        if (!!this.artifactId) backend.downloadArtifactZipV2(this.artifactId, {filter: referencePaths})
     }
 
     private jobId: string;
