Update test and phase names when updating LastSeen field so it reflects latest revision

#rnx

diff --git a/Engine/Source/Programs/Horde/Plugins/Build/HordeServer.Build/Jobs/TestData/TestDataCollectionV2.cs b/Engine/Source/Programs/Horde/Plugins/Build/HordeServer.Build/Jobs/TestData/TestDataCollectionV2.cs
--- a/Engine/Source/Programs/Horde/Plugins/Build/HordeServer.Build/Jobs/TestData/TestDataCollectionV2.cs
+++ b/Engine/Source/Programs/Horde/Plugins/Build/HordeServer.Build/Jobs/TestData/TestDataCollectionV2.cs
@@ -917,6 +917,7 @@ namespace HordeServer.Jobs.TestData
 				UpdateDefinitionBuilder<TestNameDocument> updateBuilder = Builders<TestNameDocument>.Update;
 				List<UpdateDefinition<TestNameDocument>> updates = new List<UpdateDefinition<TestNameDocument>>();
 				updates.Add(updateBuilder.Set(x => x.LastSeenUtc, DateTime.UtcNow));
+				updates.Add(updateBuilder.Set(x => x.Name, name));
 				await _tests.FindOneAndUpdateAsync(x => x.Id == test.Id, updateBuilder.Combine(updates), null, cancellationToken);
 			}
 
@@ -997,6 +998,7 @@ namespace HordeServer.Jobs.TestData
 				UpdateDefinitionBuilder<TestPhaseDocument> updateBuilder = Builders<TestPhaseDocument>.Update;
 				List<UpdateDefinition<TestPhaseDocument>> updates = new List<UpdateDefinition<TestPhaseDocument>>();
 				updates.Add(updateBuilder.Set(x => x.LastSeenUtc, DateTime.UtcNow));
+				updates.Add(updateBuilder.Set(x => x.Name, name));
 				await _testPhases.FindOneAndUpdateAsync(x => x.Id == phase.Id, updateBuilder.Combine(updates), null, cancellationToken);
 			}
 
