Horde: Add logging to shutdown of disabled agents

diff --git a/Engine/Source/Programs/Horde/Plugins/Build/HordeServer.Build/Agents/Pools/PoolUpdateService.cs b/Engine/Source/Programs/Horde/Plugins/Build/HordeServer.Build/Agents/Pools/PoolUpdateService.cs
--- a/Engine/Source/Programs/Horde/Plugins/Build/HordeServer.Build/Agents/Pools/PoolUpdateService.cs
+++ b/Engine/Source/Programs/Horde/Plugins/Build/HordeServer.Build/Agents/Pools/PoolUpdateService.cs
@@ -84,6 +84,7 @@ namespace HordeServer.Agents.Pools
 			int c = 0;
 			foreach (IAgent agent in disabledAgents)
 			{
+				_logger.LogInformation("Checking if disabled agent {AgentId} should be shut down", agent.Id.ToString());
 				if (HasGracePeriodExpired(agent, pools, _buildConfig.CurrentValue.AgentShutdownIfDisabledGracePeriod))
 				{
 					await agent.TryUpdateAsync(new UpdateAgentOptions { RequestShutdown = true }, cancellationToken: cancellationToken);
