[Horde][BuildHealth][Dashboard] Job History time span filter & removal of job count.

diff --git a/Engine/Source/Programs/Horde/HordeDashboard/src/components/buildhealth/BuildHealthDataTypes.ts b/Engine/Source/Programs/Horde/HordeDashboard/src/components/buildhealth/BuildHealthDataTypes.ts
--- a/Engine/Source/Programs/Horde/HordeDashboard/src/components/buildhealth/BuildHealthDataTypes.ts
+++ b/Engine/Source/Programs/Horde/HordeDashboard/src/components/buildhealth/BuildHealthDataTypes.ts
@@ -6,8 +6,8 @@ export const PROJECTS_URL_PARAM_NAME: string = "projects";
 export const STREAMS_URL_PARAM_NAME: string = "streams";
 export const TEMPLATES_URL_PARAM_NAME: string = "templates";
 export const STEPS_URL_PARAM_NAME: string = "stepNames";
-export const JOBCOUNT_URL_PARAM_NAME: string = "jobCount";
 export const PARAMETER_KEY_PREFIX: string = "parameter_key";
+export const TIME_SPAN_URL_PARAM: string = "lastTimeRange";
 export const KEY_SEPARATOR: string = "::";
 
 /**
diff --git a/Engine/Source/Programs/Horde/HordeDashboard/src/components/buildhealth/BuildHealthOptions.ts b/Engine/Source/Programs/Horde/HordeDashboard/src/components/buildhealth/BuildHealthOptions.ts
--- a/Engine/Source/Programs/Horde/HordeDashboard/src/components/buildhealth/BuildHealthOptions.ts
+++ b/Engine/Source/Programs/Horde/HordeDashboard/src/components/buildhealth/BuildHealthOptions.ts
@@ -1,7 +1,7 @@
 // Copyright Epic Games, Inc. All Rights Reserved.
 
 import { action, makeObservable, observable, runInAction } from "mobx";
-import { JOBCOUNT_URL_PARAM_NAME, PROJECTS_URL_PARAM_NAME, STEPS_URL_PARAM_NAME, STREAMS_URL_PARAM_NAME, TemplateRefData, TEMPLATES_URL_PARAM_NAME } from "../buildhealth/BuildHealthDataTypes";
+import { PROJECTS_URL_PARAM_NAME, STEPS_URL_PARAM_NAME, STREAMS_URL_PARAM_NAME, TemplateRefData, TEMPLATES_URL_PARAM_NAME, TIME_SPAN_URL_PARAM } from "../buildhealth/BuildHealthDataTypes";
 import { BuildHealthDataHandler } from "./BuildHealthDataHandler";
 import { decodeStepKey, decodeTemplateKey, encodeStepKey, encodeTemplateKey } from "./BuildHealthUtilities";
 import { Location } from 'react-router-dom';
@@ -16,9 +16,47 @@ type BuildHealthQueryParams = {
     streamIds: string[];
     templateIds: string[];
     stepNames: string[];
-    jobCount?: number;
+    timeSpanKey: string;
 };
 
+/**
+ * Data type used to specify time span from now to consider for jobs.
+ */
+export type TimeSpan = {
+    text: string;
+    key: string;
+    minutes: number;
+}
+
+/**
+* List of currently supported time spans.
+*/ 
+export const JobHistoryTimeSpans: TimeSpan[] = [
+    {
+        text: "Past 15 Minutes", key: "time_15_minutes", minutes: 15
+    },
+    {
+        text: "Past 1 Hour", key: "time_1_hour", minutes: 60
+    },
+    {
+        text: "Past 4 Hours", key: "time_4_hours", minutes: 60 * 4
+    },
+    {
+        text: "Past 1 Day", key: "time_1_day", minutes: 60 * 24
+    },
+    {
+        text: "Past 2 Days", key: "time_2_days", minutes: 60 * 24 * 2
+    },
+    {
+        text: "Past 1 Week", key: "time_1_week", minutes: 60 * 24 * 7
+    },
+    {
+        text: "Past 1 Month", key: "time_1_month", minutes: 60 * 24 * 31
+    }
+]
+
+export const DEFAULT_JOB_HISTORY_TIME_SPAN: TimeSpan = JobHistoryTimeSpans[1];
+
 /**
  * Parses the Build Health options from the query params.
  * @param location Location object.
@@ -33,18 +71,12 @@ export function parseBuildHealthQueryParams(location: Location): BuildHealthQuer
     const templateIds = getList(TEMPLATES_URL_PARAM_NAME);
     const stepNames = getList(STEPS_URL_PARAM_NAME);
 
-    let jobCount: number | undefined;
-    const jobCountRaw = searchParams.get(JOBCOUNT_URL_PARAM_NAME);
-
-    if (jobCountRaw !== null) {
-        const parsed = Number(jobCountRaw);
-        if (Number.isFinite(parsed) && parsed >= 0) {
-            jobCount = parsed;
-        }
+    let timeSpanKey: string = searchParams.get(TIME_SPAN_URL_PARAM) ?? "";
+    if (!JobHistoryTimeSpans.some(ts => ts.key === timeSpanKey)) {
+        timeSpanKey = DEFAULT_JOB_HISTORY_TIME_SPAN.key;
     }
 
-
-    return { projectIds, streamIds, templateIds: templateIds, stepNames, jobCount };
+    return { projectIds, streamIds, templateIds: templateIds, stepNames, timeSpanKey };
 }
 
 /**
@@ -57,9 +89,10 @@ export async function loadBuildHealthOptionsFromParams(
     params: BuildHealthQueryParams,
     options: BuildHealthOptionsController,
     handler: BuildHealthDataHandler) {
-    if (params.jobCount) {
+
+    if (params.timeSpanKey) {
         runInAction(() => {
-            options.setJobCount(params.jobCount!);
+            options.setJobHistoryTimeSpan(JobHistoryTimeSpans.find(ts => ts.key === params.timeSpanKey)!);
         });
     }
 
@@ -139,7 +172,8 @@ export class BuildHealthOptionsState {
     @observable stepOutcomeEnabledStreamKeys: string[] = [];
     @observable stepOutcomeEnabledJobKeys: string[] = [];
     @observable stepOutcomeEnabledStepKeys: string[] = [];
-    @observable stepOutcomeTotalJobs: number = 25;
+    @observable startDate!: Date;
+    @observable endDate?: Date;
 
     // #endregion -- Published Options --
 
@@ -149,6 +183,7 @@ export class BuildHealthOptionsState {
     @observable enabledStreams: Record<string, string> = {};
     @observable enabledTemplates: Record<string, string> = {};
     @observable enabledSteps: Record<string, string> = {};
+    @observable jobHistoryTimeSpan: TimeSpan = DEFAULT_JOB_HISTORY_TIME_SPAN;
 
     // #endregion -- User Selected Options --
 
@@ -156,6 +191,11 @@ export class BuildHealthOptionsState {
 
     constructor() {
         makeObservable(this);
+        this.endDate = new Date();
+
+        const d = new Date();
+        d.setDate(d.getDate() - 1);
+        this.startDate = d;
     }
 
     // #endregion -- Constructor --
@@ -286,12 +326,12 @@ export class BuildHealthOptionsController {
     }
 
     /**
-     * Changes the Job Count on the options.
-     * @param newJobCount The new job count.
+     * Sets the current job history time span to utilize.
+     * @param time The time span to utilize.
      */
     @action
-    setJobCount(newJobCount: number) {
-        this.state.stepOutcomeTotalJobs = newJobCount;
+    setJobHistoryTimeSpan(time: TimeSpan) {
+        this.state.jobHistoryTimeSpan = time;
         this.setOptionsChanged();
     }
 
@@ -355,6 +395,9 @@ export class BuildHealthOptionsController {
             this.state.stepOutcomeEnabledStreamKeys = Object.keys(this.state.enabledStreams);
             this.state.stepOutcomeEnabledJobKeys = Object.keys(this.state.enabledTemplates).map(key => decodeTemplateKey(key).templateId);
             this.state.stepOutcomeEnabledStepKeys = Object.keys(this.state.enabledSteps).map(key => decodeStepKey(key).stepName);
+            this.state.startDate = new Date(new Date().valueOf() - (this.state.jobHistoryTimeSpan.minutes * 60000));
+            this.state.endDate = new Date();
+
             this.lastSynchronize = this.optionsChangeVersion;
         }
     }
@@ -370,7 +413,7 @@ export class BuildHealthOptionsController {
         const streams = this.state.stepOutcomeEnabledStreamKeys.join(",");
         const jobs = this.state.stepOutcomeEnabledJobKeys.join(",");
         const steps = this.state.stepOutcomeEnabledStepKeys.join(",");
-        const jobCount = this.state.stepOutcomeTotalJobs.toString();
+        const lastTimeRange = this.state.jobHistoryTimeSpan.key.toString();
 
         if (!projects.length) return params;
 
@@ -378,7 +421,7 @@ export class BuildHealthOptionsController {
         params.set(STREAMS_URL_PARAM_NAME, streams);
         params.set(TEMPLATES_URL_PARAM_NAME, jobs);
         params.set(STEPS_URL_PARAM_NAME, steps);
-        params.set(JOBCOUNT_URL_PARAM_NAME, jobCount);
+        params.set(TIME_SPAN_URL_PARAM, lastTimeRange);
 
         return params;
     }
diff --git a/Engine/Source/Programs/Horde/HordeDashboard/src/components/buildhealth/BuildhealthView.tsx b/Engine/Source/Programs/Horde/HordeDashboard/src/components/buildhealth/BuildhealthView.tsx
--- a/Engine/Source/Programs/Horde/HordeDashboard/src/components/buildhealth/BuildhealthView.tsx
+++ b/Engine/Source/Programs/Horde/HordeDashboard/src/components/buildhealth/BuildhealthView.tsx
@@ -1,6 +1,6 @@
 // Copyright Epic Games, Inc. All Rights Reserved.
 
-import { Dropdown, DropdownMenuItemType, IDropdownOption, Label, Slider, Stack, TooltipHost } from "@fluentui/react";
+import { ComboBox, Dropdown, DropdownMenuItemType, IDropdownOption, Label, Slider, Stack, TooltipHost } from "@fluentui/react";
 import { useWindowSize } from "../../base/utilities/hooks";
 import { getHordeStyling } from "../../styles/Styles";
 import { Breadcrumbs } from "../Breadcrumbs";
@@ -10,9 +10,8 @@ import { useEffect, useMemo, useState } from "react";
 import { GetProjectResponse, GetStreamResponse } from "horde/backend/Api";
 import { observer } from "mobx-react-lite";
 import { useLocation, useNavigate } from "react-router";
-import { Text } from '@fluentui/react';
 import { BuildHealthDataHandler, DataHandlerRefreshRequest, expandRefreshRequest } from "./BuildHealthDataHandler";
-import { BuildHealthOptionsController, BuildHealthOptionsState, loadBuildHealthOptionsFromParams, parseBuildHealthQueryParams } from "./BuildHealthOptions";
+import { BuildHealthOptionsController, BuildHealthOptionsState, loadBuildHealthOptionsFromParams, parseBuildHealthQueryParams, TimeSpan, JobHistoryTimeSpans } from "./BuildHealthOptions";
 import { encodeStepKey, encodeTemplateKey } from "./BuildHealthUtilities";
 import { KEY_SEPARATOR, PARAMETER_KEY_PREFIX, StepRefData, TemplateRefData } from "./BuildHealthDataTypes";
 
@@ -34,7 +33,6 @@ type BuildHealthOptionList = {
 
 // #endregion -- Dropdown Data Types --
 
-
 // #region -- Base Dropdown Components --
 
 /**
@@ -414,24 +412,19 @@ const BuildHealthUrlSync: React.FC<{ buildHealthController: BuildHealthOptionsCo
  */
 const BuildHealthSidebar: React.FC<{ handler: BuildHealthDataHandler, buildHealthOptions: BuildHealthOptionsController }> = function ConstrucBuildHealthSidebar({ handler, buildHealthOptions }) {
     const { hordeClasses } = getHordeStyling();
-    const [jobCountFilter, setJobCountFilter] = useState<BuildCountFilter>({ numberBuilds: buildHealthOptions.state.stepOutcomeTotalJobs });
-
-    const updateJobCount = (value: number) => {
-        setJobCountFilter(prev => ({ ...prev, numberBuilds: value }));
-    };
+    const [timeSpanKey, setTimeSpanKey] = useState(() => buildHealthOptions.state.jobHistoryTimeSpan.key);
 
-    const updateJobCountCommitted = (_event: any, value: number, _range?) => {
-        buildHealthOptions.setJobCount(value);
+    const updateJobHistoryTimeSpan = (_ev, option) => {
+        const select = option as TimeSpan;
+        setTimeSpanKey(select.key);
+        buildHealthOptions.setJobHistoryTimeSpan(select);
     };
 
     useEffect(() => {
-        if (buildHealthOptions.state.stepOutcomeTotalJobs) {
-            setJobCountFilter(prev => ({
-                ...prev,
-                numberBuilds: buildHealthOptions.state.stepOutcomeTotalJobs
-            }));
-        }
-    }, [buildHealthOptions.state.stepOutcomeTotalJobs]);
+        setTimeSpanKey(buildHealthOptions.state.jobHistoryTimeSpan.key);
+    }, [buildHealthOptions.state.jobHistoryTimeSpan.key]);
+
+    let timeComboWidth = 180;
 
     return (
         <Stack style={{ minWidth: 300, maxWidth: 450, paddingRight: 18 }}>
@@ -445,13 +438,18 @@ const BuildHealthSidebar: React.FC<{ handler: BuildHealthDataHandler, buildHealt
                     <TemplateDropdown handler={handler} buildHealthOptions={buildHealthOptions} />
                     <StepDropdown handler={handler} buildHealthOptions={buildHealthOptions} />
                     <Stack>
-                        <Stack horizontal verticalAlign="center" tokens={{ childrenGap: 8 }}>
-                            <TooltipHost content="Adjust the maximum number of jobs per template to show.">
-                                <Label>Max Jobs:</Label>
-                            </TooltipHost>
-                            <Text>{jobCountFilter.numberBuilds}</Text>
-                        </Stack>
-                        <Slider min={0} max={90} value={jobCountFilter.numberBuilds} showValue={false} onChange={updateJobCount} onChanged={updateJobCountCommitted} styles={{ root: { width: 300 }, slideBox: { paddingLeft: 2 } }} />
+                        <TooltipHost content="Adjust the historical job history time span.">
+                            <Label>Jobs Since:</Label>
+                        </TooltipHost>
+                        <ComboBox
+                            styles={{ root: { width: timeComboWidth } }}
+                            options={JobHistoryTimeSpans}
+                            selectedKey={timeSpanKey}
+                            onChange={(ev, option, index, value) => {
+                                updateJobHistoryTimeSpan(ev, option);
+                                buildHealthOptions.synchronizeDerivedKeys();
+                            }}
+                        />
                     </Stack>
                 </Stack>
             </Stack>
@@ -459,10 +457,6 @@ const BuildHealthSidebar: React.FC<{ handler: BuildHealthDataHandler, buildHealt
     );
 };
 
-type BuildCountFilter = {
-    numberBuilds: number;
-}
-
 /**
  * React Component representing the entire Build Health view.
  * @returns React Component.
@@ -475,12 +469,12 @@ export const BuildHealthView: React.FC = observer(function ConstructBuildHealthV
         streams: buildHealthState.stepOutcomeEnabledStreamKeys,
         jobs: buildHealthState.stepOutcomeEnabledJobKeys,
         steps: buildHealthState.stepOutcomeEnabledStepKeys,
-        changelistCount: buildHealthState.stepOutcomeTotalJobs,
+        jobHistorySpan: { start: buildHealthState.startDate, end: buildHealthState.endDate }
     }), [
         buildHealthState.stepOutcomeEnabledStreamKeys,
         buildHealthState.stepOutcomeEnabledJobKeys,
         buildHealthState.stepOutcomeEnabledStepKeys,
-        buildHealthState.stepOutcomeTotalJobs,
+        { start: buildHealthState.startDate, end: buildHealthState.endDate }
     ]);
 
     return (
diff --git a/Engine/Source/Programs/Horde/HordeDashboard/src/components/buildhealth/stepoutcome/StepOutcomeDataHandler.ts b/Engine/Source/Programs/Horde/HordeDashboard/src/components/buildhealth/stepoutcome/StepOutcomeDataHandler.ts
--- a/Engine/Source/Programs/Horde/HordeDashboard/src/components/buildhealth/stepoutcome/StepOutcomeDataHandler.ts
+++ b/Engine/Source/Programs/Horde/HordeDashboard/src/components/buildhealth/stepoutcome/StepOutcomeDataHandler.ts
@@ -12,10 +12,10 @@ import backend from "horde/backend";
  */
 export interface StepOutcomeFilters {
     streams: string[];
+    jobHistorySpan: { start: Date; end?: Date };
     jobs?: string[];
     steps?: string[];
-    dateRange?: { start: Date; end: Date };
-    changelistCount?: number;
+    maxJobCount? : number;
 }
 
 // #endregion -- API Data Types & Utilities --
@@ -26,9 +26,11 @@ export interface StepOutcomeFilters {
  */
 export class StepOutcomeDataHandler extends PollBase {
     private readonly STREAM_JOBS_FILTER = "id,streamId,name,change,templateId,state,createTime,updateTime,batches,preflightDescription,cancellationReason";
+    private readonly MAX_JOB_COUNT_DEFAULT = 150;
 
     // #region -- Private Members -- 
 
+    activeRefresh : boolean = false;
     activeFilter: StepOutcomeFilters;
     jobData: GetJobResponse[] = [];
     stepOutcomeTableData: StepOutcomeTable = new StepOutcomeTable();
@@ -50,29 +52,33 @@ export class StepOutcomeDataHandler extends PollBase {
     // #region -- Interface --
 
     async poll(): Promise<void> {
-        if (
+        if (this.activeRefresh || 
             this.activeFilter === undefined ||
             !this.activeFilter.streams?.length ||
             !this.activeFilter.jobs?.length
         ) {
             this.stepOutcomeTableData.reset();
+            this.setUpdated();
             return;
         }
 
+        this.activeRefresh = true;
         this.stepOutcomeTableData.reset();
 
         const stepSet = new Set(this.activeFilter.steps?.map(s => s.toLowerCase()));
         const jobList = this.activeFilter.jobs ?? [];
         const streamList = this.activeFilter.streams ?? [];
-        const changelistCount = this.activeFilter.changelistCount ?? 90;
         const queries: { streamId: string; query: JobStreamQuery }[] = [];
+        const modifiedAfter = this.activeFilter.jobHistorySpan.start.toISOString();
+        const jobCount = this.activeFilter.maxJobCount ?? this.MAX_JOB_COUNT_DEFAULT;
 
         for (const streamId of streamList) {
             for (const job of jobList) {
                 const query: JobStreamQuery = {
                     template: [job],
-                    count: changelistCount,
-                    filter: this.STREAM_JOBS_FILTER
+                    count: jobCount,
+                    filter: this.STREAM_JOBS_FILTER,
+                    modifiedAfter: modifiedAfter
                 };
                 queries.push({ streamId, query });
             }
@@ -103,9 +109,12 @@ export class StepOutcomeDataHandler extends PollBase {
             }
         }
 
+        console.info(`Retrieved ${this.jobData.length} results for StepOutcome query.`);
+
         this.stepOutcomeTableData.orderTableDataByChange(true);
         this.setUpdated();
-
+        
+        this.activeRefresh = false;
         return Promise.resolve();
     }
 
diff --git a/Engine/Source/Programs/Horde/HordeDashboard/src/components/buildhealth/stepoutcome/StepOutcomeDataTypes.ts b/Engine/Source/Programs/Horde/HordeDashboard/src/components/buildhealth/stepoutcome/StepOutcomeDataTypes.ts
--- a/Engine/Source/Programs/Horde/HordeDashboard/src/components/buildhealth/stepoutcome/StepOutcomeDataTypes.ts
+++ b/Engine/Source/Programs/Horde/HordeDashboard/src/components/buildhealth/stepoutcome/StepOutcomeDataTypes.ts
@@ -113,6 +113,12 @@ export class StepOutcomeTable {
             insertionCol = this.changeColHeaders.length;
             this.changeColLookup.set(tableEntry.change, insertionCol);
             this.changeColHeaders.push(tableEntry.change);
+
+            for (const row of this.tableEntries) {
+                if (row) {
+                    this.ensureRow(row, this.changeColHeaders.length);
+                }
+            }
         }
 
         // Ensure the row exists and is padded
@@ -124,7 +130,7 @@ export class StepOutcomeTable {
 
         const existingEntry = this.tableEntries[insertionRow][insertionCol];
         if (existingEntry !== null) {
-            console.warn(`Duplicate entry for step "${tableEntry.stepResponse.name}" and change "${tableEntry.change}". Ambiguous selector needed.`, {
+            console.warn(`Duplicate entry for step new (${tableEntry.jobName}::${tableEntry.jobId}::${tableEntry.stepResponse.name} (pre-existing: ${existingEntry.jobName}::${existingEntry.jobId}::${existingEntry.stepResponse.name})) and change ("${tableEntry.change}" (${existingEntry.change})). Ambiguous selector needed.`, {
                 existingEntry,
                 newEntry: tableEntry
             });
