Horde: Disable logging of UAT exit code in JobExecutor

Reverts 44661300.

diff --git a/Engine/Source/Programs/Horde/Drivers/JobDriver/Execution/JobExecutor.cs b/Engine/Source/Programs/Horde/Drivers/JobDriver/Execution/JobExecutor.cs
--- a/Engine/Source/Programs/Horde/Drivers/JobDriver/Execution/JobExecutor.cs
+++ b/Engine/Source/Programs/Horde/Drivers/JobDriver/Execution/JobExecutor.cs
@@ -809,7 +809,9 @@ namespace JobDriver.Execution
 			int exitCode = await ExecuteAutomationToolAsync(step, workspaceDir, arguments.ToString(), useP4, logger, cancellationToken);
 			if (exitCode != 0)
 			{
-				logger.LogError("BuildGraph setup terminated with non-zero exit code: {ExitCode}", exitCode);
+				// Disable this log line until an appropriate structured/systemtic log event can be assigned
+				// for correct grouping with other events in Build Health.
+				// logger.LogError("BuildGraph setup terminated with non-zero exit code: {ExitCode}", exitCode);
 				return false;
 			}
 
@@ -1256,7 +1258,9 @@ namespace JobDriver.Execution
 			int exitCode = await ExecuteAutomationToolAsync(step, workspaceDir, arguments, useP4, logger, cancellationToken);
 			if (exitCode != 0)
 			{
-				logger.LogError("AutomationTool (UAT) terminated with non-zero exit code: {ExitCode}", exitCode);
+				// Disable this log line until an appropriate structured/systemtic log event can be assigned
+				// for correct grouping with other events in Build Health.
+				// logger.LogError("AutomationTool (UAT) terminated with non-zero exit code: {ExitCode}", exitCode);
 				return false;
 			}
 
