Horde: Support new EC2 tag for identifying agent's pool ID

Older tag "Horde_Autoscale_Pool" identifying by name will still work but should be considered deprecated.

diff --git a/Engine/Source/Programs/Horde/Plugins/Build/HordeServer.Build.Tests/Fleet/AwsRecyclingFleetManagerTest.cs b/Engine/Source/Programs/Horde/Plugins/Build/HordeServer.Build.Tests/Fleet/AwsRecyclingFleetManagerTest.cs
--- a/Engine/Source/Programs/Horde/Plugins/Build/HordeServer.Build.Tests/Fleet/AwsRecyclingFleetManagerTest.cs
+++ b/Engine/Source/Programs/Horde/Plugins/Build/HordeServer.Build.Tests/Fleet/AwsRecyclingFleetManagerTest.cs
@@ -2,6 +2,7 @@
 
 using Amazon.EC2;
 using Amazon.EC2.Model;
+using EpicGames.Horde.Agents.Pools;
 using HordeCommon;
 using HordeServer.Agents;
 using HordeServer.Agents.Fleet;
@@ -15,13 +16,15 @@ namespace HordeServer.Tests.Fleet
 	[TestClass]
 	public class AwsRecyclingFleetManagerTest : BuildTestSetup
 	{
+		private const string PoolId = "test-pool";
+		private const string PoolName = "testPool";
 		private readonly FakeClock _clock = new();
 
 		[TestMethod]
 		public async Task ExpandAmazonEc2ExceptionsArePropagatedAsync()
 		{
 			FakeAmazonEc2 ec2 = new();
-			ec2.AddInstance(FakeAmazonEc2.StateStopped, InstanceType.M5Large, FakeAmazonEc2.AzUsEast1A);
+			ec2.AddInstance(FakeAmazonEc2.StateStopped, InstanceType.M5Large, FakeAmazonEc2.AzUsEast1A, tagPoolId: PoolId);
 			ec2.SetCapacity(FakeAmazonEc2.AzUsEast1A, InstanceType.M5Large, 1);
 			ec2.OnStartInstances = (_, _) => throw new AmazonEC2Exception("Something bad happened");
 			await Assert.ThrowsExceptionAsync<AmazonEC2Exception>(() => ExpandPoolAsync(ec2.Get(), 1, new()));
@@ -31,8 +34,8 @@ namespace HordeServer.Tests.Fleet
 		public async Task ExpandOneAzPartialCapacityWithFallbackAsync()
 		{
 			FakeAmazonEc2 ec2 = new();
-			Instance i1 = ec2.AddInstance(FakeAmazonEc2.StateStopped, InstanceType.M5Large, FakeAmazonEc2.AzUsEast1A);
-			Instance i2 = ec2.AddInstance(FakeAmazonEc2.StateStopped, InstanceType.M5Large, FakeAmazonEc2.AzUsEast1A);
+			Instance i1 = ec2.AddInstance(FakeAmazonEc2.StateStopped, InstanceType.M5Large, FakeAmazonEc2.AzUsEast1A, tagPoolName: PoolName, tagPoolId: PoolId);
+			Instance i2 = ec2.AddInstance(FakeAmazonEc2.StateStopped, InstanceType.M5Large, FakeAmazonEc2.AzUsEast1A, tagPoolName: PoolName, tagPoolId: PoolId);
 			ec2.SetCapacity(FakeAmazonEc2.AzUsEast1A, InstanceType.M5Large, 1);
 			ec2.SetCapacity(FakeAmazonEc2.AzUsEast1A, InstanceType.C5Large, 2);
 
@@ -47,8 +50,8 @@ namespace HordeServer.Tests.Fleet
 		public async Task ExpandOneAzPartialCapacityWithMultipleFallbacksAsync()
 		{
 			FakeAmazonEc2 ec2 = new();
-			Instance i1 = ec2.AddInstance(FakeAmazonEc2.StateStopped, InstanceType.M5Large, FakeAmazonEc2.AzUsEast1A);
-			Instance i2 = ec2.AddInstance(FakeAmazonEc2.StateStopped, InstanceType.M5Large, FakeAmazonEc2.AzUsEast1A);
+			Instance i1 = ec2.AddInstance(FakeAmazonEc2.StateStopped, InstanceType.M5Large, FakeAmazonEc2.AzUsEast1A, tagPoolName: PoolName);
+			Instance i2 = ec2.AddInstance(FakeAmazonEc2.StateStopped, InstanceType.M5Large, FakeAmazonEc2.AzUsEast1A, tagPoolName: PoolName);
 			ec2.SetCapacity(FakeAmazonEc2.AzUsEast1A, InstanceType.M5Large, 1);
 			ec2.SetCapacity(FakeAmazonEc2.AzUsEast1A, InstanceType.C5Large, 2);
 
@@ -63,8 +66,8 @@ namespace HordeServer.Tests.Fleet
 		public async Task ExpandOneAzPartialCapacityNoFallbackAsync()
 		{
 			FakeAmazonEc2 ec2 = new();
-			ec2.AddInstance(FakeAmazonEc2.StateStopped, InstanceType.M5Large, FakeAmazonEc2.AzUsEast1A);
-			ec2.AddInstance(FakeAmazonEc2.StateStopped, InstanceType.M5Large, FakeAmazonEc2.AzUsEast1A);
+			ec2.AddInstance(FakeAmazonEc2.StateStopped, InstanceType.M5Large, FakeAmazonEc2.AzUsEast1A, tagPoolName: PoolName);
+			ec2.AddInstance(FakeAmazonEc2.StateStopped, InstanceType.M5Large, FakeAmazonEc2.AzUsEast1A, tagPoolName: PoolName);
 			ec2.SetCapacity(FakeAmazonEc2.AzUsEast1A, InstanceType.M5Large, 1);
 
 			await ExpandPoolAsync(ec2.Get(), 2, new());
@@ -76,8 +79,8 @@ namespace HordeServer.Tests.Fleet
 		public async Task ExpandTwoAzsPartialCapacityNoFallbackAsync()
 		{
 			FakeAmazonEc2 ec2 = new();
-			ec2.AddInstance(FakeAmazonEc2.StateStopped, InstanceType.M5Large, FakeAmazonEc2.AzUsEast1A);
-			ec2.AddInstance(FakeAmazonEc2.StateStopped, InstanceType.M5Large, FakeAmazonEc2.AzUsEast1B);
+			ec2.AddInstance(FakeAmazonEc2.StateStopped, InstanceType.M5Large, FakeAmazonEc2.AzUsEast1A, tagPoolName: PoolName);
+			ec2.AddInstance(FakeAmazonEc2.StateStopped, InstanceType.M5Large, FakeAmazonEc2.AzUsEast1B, tagPoolName: PoolName);
 			ec2.SetCapacity(FakeAmazonEc2.AzUsEast1A, InstanceType.M5Large, 1);
 			ec2.SetCapacity(FakeAmazonEc2.AzUsEast1B, InstanceType.M5Large, 0);
 
@@ -90,8 +93,8 @@ namespace HordeServer.Tests.Fleet
 		public async Task ExpandMoreThanNumStoppedInstancesAsync()
 		{
 			FakeAmazonEc2 ec2 = new();
-			ec2.AddInstance(FakeAmazonEc2.StateStopped, InstanceType.M5Large, FakeAmazonEc2.AzUsEast1A);
-			ec2.AddInstance(FakeAmazonEc2.StateStopped, InstanceType.M5Large, FakeAmazonEc2.AzUsEast1B);
+			ec2.AddInstance(FakeAmazonEc2.StateStopped, InstanceType.M5Large, FakeAmazonEc2.AzUsEast1A, tagPoolName: PoolName);
+			ec2.AddInstance(FakeAmazonEc2.StateStopped, InstanceType.M5Large, FakeAmazonEc2.AzUsEast1B, tagPoolName: PoolName);
 			ec2.SetCapacity(FakeAmazonEc2.AzUsEast1A, InstanceType.M5Large, 1);
 			ec2.SetCapacity(FakeAmazonEc2.AzUsEast1B, InstanceType.M5Large, 1);
 
@@ -107,8 +110,8 @@ namespace HordeServer.Tests.Fleet
 		public async Task ExpandTwoAzsFullCapacityNoFallbackAsync()
 		{
 			FakeAmazonEc2 ec2 = new();
-			ec2.AddInstance(FakeAmazonEc2.StateStopped, InstanceType.M5Large, FakeAmazonEc2.AzUsEast1A);
-			ec2.AddInstance(FakeAmazonEc2.StateStopped, InstanceType.M5Large, FakeAmazonEc2.AzUsEast1B);
+			ec2.AddInstance(FakeAmazonEc2.StateStopped, InstanceType.M5Large, FakeAmazonEc2.AzUsEast1A, tagPoolName: PoolName);
+			ec2.AddInstance(FakeAmazonEc2.StateStopped, InstanceType.M5Large, FakeAmazonEc2.AzUsEast1B, tagPoolName: PoolName);
 			ec2.SetCapacity(FakeAmazonEc2.AzUsEast1A, InstanceType.M5Large, 1);
 			ec2.SetCapacity(FakeAmazonEc2.AzUsEast1B, InstanceType.M5Large, 1);
 
@@ -121,8 +124,8 @@ namespace HordeServer.Tests.Fleet
 		public async Task ExpandTwoAzsFullCapacityEmptyFallbackAsync()
 		{
 			FakeAmazonEc2 ec2 = new();
-			ec2.AddInstance(FakeAmazonEc2.StateStopped, InstanceType.M5Large, FakeAmazonEc2.AzUsEast1A);
-			ec2.AddInstance(FakeAmazonEc2.StateStopped, InstanceType.M5Large, FakeAmazonEc2.AzUsEast1B);
+			ec2.AddInstance(FakeAmazonEc2.StateStopped, InstanceType.M5Large, FakeAmazonEc2.AzUsEast1A, tagPoolName: PoolName);
+			ec2.AddInstance(FakeAmazonEc2.StateStopped, InstanceType.M5Large, FakeAmazonEc2.AzUsEast1B, tagPoolName: PoolName);
 			ec2.SetCapacity(FakeAmazonEc2.AzUsEast1A, InstanceType.M5Large, 1);
 			ec2.SetCapacity(FakeAmazonEc2.AzUsEast1B, InstanceType.M5Large, 1);
 
@@ -135,11 +138,11 @@ namespace HordeServer.Tests.Fleet
 		public async Task ExpandLaunchesLastUsedInstancesFirstAsync()
 		{
 			FakeAmazonEc2 ec2 = new();
-			Instance i1 = ec2.AddInstance(FakeAmazonEc2.StateStopped, InstanceType.M5Large, launchTime: DateTime.Parse("Aug 18, 2018"));
-			Instance i2 = ec2.AddInstance(FakeAmazonEc2.StateStopped, InstanceType.M5Large, launchTime: DateTime.Parse("Aug 22, 2018"));
-			Instance i3 = ec2.AddInstance(FakeAmazonEc2.StateStopped, InstanceType.M5Large, launchTime: DateTime.Parse("Aug 5, 2018"));
-			Instance i4 = ec2.AddInstance(FakeAmazonEc2.StateStopped, InstanceType.M5Large, launchTime: DateTime.Parse("Aug 28, 2018"));
-			Instance i5 = ec2.AddInstance(FakeAmazonEc2.StateStopped, InstanceType.M5Large, launchTime: DateTime.Parse("Aug 1, 2018"));
+			Instance i1 = ec2.AddInstance(FakeAmazonEc2.StateStopped, InstanceType.M5Large, launchTime: DateTime.Parse("Aug 18, 2018"), tagPoolName: PoolName);
+			Instance i2 = ec2.AddInstance(FakeAmazonEc2.StateStopped, InstanceType.M5Large, launchTime: DateTime.Parse("Aug 22, 2018"), tagPoolName: PoolName);
+			Instance i3 = ec2.AddInstance(FakeAmazonEc2.StateStopped, InstanceType.M5Large, launchTime: DateTime.Parse("Aug 5, 2018"), tagPoolName: PoolName);
+			Instance i4 = ec2.AddInstance(FakeAmazonEc2.StateStopped, InstanceType.M5Large, launchTime: DateTime.Parse("Aug 28, 2018"), tagPoolName: PoolName);
+			Instance i5 = ec2.AddInstance(FakeAmazonEc2.StateStopped, InstanceType.M5Large, launchTime: DateTime.Parse("Aug 1, 2018"), tagPoolName: PoolName);
 			ec2.SetCapacity(FakeAmazonEc2.AzUsEast1A, InstanceType.M5Large, 2);
 
 			_ = i1;
@@ -157,10 +160,10 @@ namespace HordeServer.Tests.Fleet
 		public async Task StopInstancesStuckInPendingAsync()
 		{
 			FakeAmazonEc2 ec2 = new();
-			Instance i1 = ec2.AddInstance(FakeAmazonEc2.StatePending, InstanceType.M5Large, launchTime: _clock.UtcNow.Subtract(TimeSpan.FromMinutes(50)));
-			Instance i2 = ec2.AddInstance(FakeAmazonEc2.StatePending, InstanceType.M5Large, launchTime: _clock.UtcNow.Subtract(TimeSpan.FromMinutes(3)));
-			Instance i3 = ec2.AddInstance(FakeAmazonEc2.StateStopped, InstanceType.M5Large);
-			Instance i4 = ec2.AddInstance(FakeAmazonEc2.StateStopped, InstanceType.M5Large);
+			Instance i1 = ec2.AddInstance(FakeAmazonEc2.StatePending, InstanceType.M5Large, launchTime: _clock.UtcNow.Subtract(TimeSpan.FromMinutes(50)), tagPoolName: PoolName);
+			Instance i2 = ec2.AddInstance(FakeAmazonEc2.StatePending, InstanceType.M5Large, launchTime: _clock.UtcNow.Subtract(TimeSpan.FromMinutes(3)), tagPoolName: PoolName);
+			Instance i3 = ec2.AddInstance(FakeAmazonEc2.StateStopped, InstanceType.M5Large, tagPoolName: PoolName);
+			Instance i4 = ec2.AddInstance(FakeAmazonEc2.StateStopped, InstanceType.M5Large, tagPoolName: PoolName);
 			ec2.SetCapacity(FakeAmazonEc2.AzUsEast1A, InstanceType.M5Large, 2);
 
 			ScaleResult result = await ExpandPoolAsync(ec2.Get(), 2, new(new List<string>()));
@@ -215,7 +218,7 @@ namespace HordeServer.Tests.Fleet
 
 			ILogger<AwsRecyclingFleetManager> logger = loggerFactory.CreateLogger<AwsRecyclingFleetManager>();
 #pragma warning disable CS0618 // Type or member is obsolete
-			IPool pool = await CreatePoolAsync(new PoolConfig { Name = "testPool", EnableAutoscaling = true, MinAgents = 0, NumReserveAgents = 0, SizeStrategy = PoolSizeStrategy.NoOp });
+			IPool pool = await CreatePoolAsync(new PoolConfig { Id = new PoolId(PoolId), Name = PoolName, EnableAutoscaling = true, MinAgents = 0, NumReserveAgents = 0, SizeStrategy = PoolSizeStrategy.NoOp });
 #pragma warning restore CS0618 // Type or member is obsolete
 			AwsRecyclingFleetManager manager = new(ec2, AgentCollection, Meter, _clock, settings, Tracer, logger);
 			return await manager.ExpandPoolAsync(pool, new List<IAgent>(), numRequestedInstances, CancellationToken.None);
diff --git a/Engine/Source/Programs/Horde/Plugins/Build/HordeServer.Build.Tests/Fleet/AwsReuseFleetManagerTest.cs b/Engine/Source/Programs/Horde/Plugins/Build/HordeServer.Build.Tests/Fleet/AwsReuseFleetManagerTest.cs
--- a/Engine/Source/Programs/Horde/Plugins/Build/HordeServer.Build.Tests/Fleet/AwsReuseFleetManagerTest.cs
+++ b/Engine/Source/Programs/Horde/Plugins/Build/HordeServer.Build.Tests/Fleet/AwsReuseFleetManagerTest.cs
@@ -2,6 +2,7 @@
 
 using Amazon.EC2;
 using Amazon.EC2.Model;
+using EpicGames.Horde.Agents.Pools;
 using HordeServer.Agents;
 using HordeServer.Agents.Fleet;
 using HordeServer.Agents.Fleet.Providers;
@@ -14,11 +15,14 @@ namespace HordeServer.Tests.Fleet
 	[TestClass]
 	public class AwsReuseFleetManagerTest : BuildTestSetup
 	{
+		private const string PoolId = "test-pool";
+		private const string PoolName = "testPool";
+		
 		[TestMethod]
 		public async Task ExpandOneAgentAsync()
 		{
 			FakeAmazonEc2 ec2 = new();
-			Instance i = ec2.AddInstance(FakeAmazonEc2.StateStopped, InstanceType.M5Large);
+			Instance i = ec2.AddInstance(FakeAmazonEc2.StateStopped, InstanceType.M5Large, tagPoolId: PoolId);
 			ec2.SetCapacity(FakeAmazonEc2.AzUsEast1A, InstanceType.M5Large, 1);
 
 			await ExpandPoolAsync(ec2.Get(), 1, new());
@@ -30,7 +34,7 @@ namespace HordeServer.Tests.Fleet
 		public async Task ExpandWithInstanceTypeChangeAsync()
 		{
 			FakeAmazonEc2 ec2 = new();
-			Instance i = ec2.AddInstance(FakeAmazonEc2.StateStopped, InstanceType.M5Large);
+			Instance i = ec2.AddInstance(FakeAmazonEc2.StateStopped, InstanceType.M5Large, tagPoolName: PoolName);
 			ec2.SetCapacity(FakeAmazonEc2.AzUsEast1A, InstanceType.M54xlarge, 1);
 
 			await ExpandPoolAsync(ec2.Get(), 1, new(new List<string> { InstanceType.M54xlarge }));
@@ -42,7 +46,7 @@ namespace HordeServer.Tests.Fleet
 		{
 			ILogger<AwsReuseFleetManager> logger = ServiceProvider.GetRequiredService<ILogger<AwsReuseFleetManager>>();
 #pragma warning disable CS0618 // Type or member is obsolete
-			IPool pool = await CreatePoolAsync(new PoolConfig { Name = "testPool", EnableAutoscaling = true, MinAgents = 0, NumReserveAgents = 0, SizeStrategy = PoolSizeStrategy.NoOp });
+			IPool pool = await CreatePoolAsync(new PoolConfig { Id = new PoolId(PoolId), Name = PoolName, EnableAutoscaling = true, MinAgents = 0, NumReserveAgents = 0, SizeStrategy = PoolSizeStrategy.NoOp });
 #pragma warning restore CS0618 // Type or member is obsolete
 			AwsReuseFleetManager manager = new(ec2, AgentCollection, settings, Tracer, logger);
 			await manager.ExpandPoolAsync(pool, new List<IAgent>(), numRequestedInstances, CancellationToken.None);
diff --git a/Engine/Source/Programs/Horde/Plugins/Build/HordeServer.Build.Tests/Fleet/FakeAmazonEc2.cs b/Engine/Source/Programs/Horde/Plugins/Build/HordeServer.Build.Tests/Fleet/FakeAmazonEc2.cs
--- a/Engine/Source/Programs/Horde/Plugins/Build/HordeServer.Build.Tests/Fleet/FakeAmazonEc2.cs
+++ b/Engine/Source/Programs/Horde/Plugins/Build/HordeServer.Build.Tests/Fleet/FakeAmazonEc2.cs
@@ -3,6 +3,7 @@
 using System.Net;
 using Amazon.EC2;
 using Amazon.EC2.Model;
+using HordeServer.Agents.Fleet.Providers;
 using Moq;
 
 namespace HordeServer.Tests.Fleet;
@@ -74,8 +75,13 @@ public class FakeAmazonEc2
 		_availabilityZoneCapacity[availabilityZone] = temp;
 	}
 
-	public Instance AddInstance(InstanceState state, InstanceType type,
-		string availabilityZone = AzUsEast1A, DateTime? launchTime = null)
+	public Instance AddInstance(
+		InstanceState state,
+		InstanceType type,
+		string availabilityZone = AzUsEast1A,
+		string? tagPoolId = null,
+		string? tagPoolName = null,
+		DateTime? launchTime = null)
 	{
 		Instance i = new()
 		{
@@ -85,6 +91,14 @@ public class FakeAmazonEc2
 			Placement = new Placement(availabilityZone),
 			LaunchTime = launchTime ?? DateTime.UnixEpoch
 		};
+		if (tagPoolId != null)
+		{
+			i.Tags.Add(new Tag(AwsFleetManager.AwsTagPoolId, tagPoolId));
+		}
+		if (tagPoolName != null)
+		{
+			i.Tags.Add(new Tag(AwsFleetManager.AwsTagPoolName, tagPoolName));
+		}
 		_instances.Add(i.InstanceId, i);
 		return i;
 	}
@@ -102,48 +116,45 @@ public class FakeAmazonEc2
 			InstanceType = i.InstanceType,
 			State = i.State,
 			Placement = i.Placement,
-			LaunchTime = i.LaunchTime
+			LaunchTime = i.LaunchTime,
+			Tags = [..i.Tags]
 		};
 	}
 
 	public int GetStoppedInstanceCount() { return _instances.Values.Count(x => x.State == StateStopped); }
 	public int GetPendingInstanceCount() { return _instances.Values.Count(x => x.State == StatePending); }
 
-	private Task<DescribeInstancesResponse> DescribeInstancesAsync(DescribeInstancesRequest request, CancellationToken cancellationToken)
+	private static bool MatchesFilter(Instance i, Filter f)
 	{
-		List<Instance> filteredInstances = new();
-		if (request.Filters.Count > 0)
+		if (f.Name == "instance-state-name")
 		{
-			foreach (Instance i in _instances.Values)
-			{
-				foreach (Filter f in request.Filters)
-				{
-					if (f.Name == "instance-state-name")
-					{
-						if (i.State.Name == f.Values[0])
-						{
-							filteredInstances.Add(i);
-						}
-					}
-					else if (f.Name.StartsWith("tag:", StringComparison.Ordinal))
-					{
-						// TODO: respect tag filtering
-					}
-					else
-					{
-						throw new Exception("Unknown filter " + f.Name);
-					}
-				}
-			}
+			return i.State.Name == f.Values[0];
 		}
-		else
+		
+		if (f.Name.StartsWith("tag:", StringComparison.Ordinal))
 		{
-			filteredInstances = _instances.Values.ToList();
+			string key = f.Name["tag:".Length..];
+			Tag? tag = i.Tags.FirstOrDefault(t => String.Equals(t.Key, key, StringComparison.Ordinal));
+			if (tag == null)
+			{
+				return false;
+			}
+			
+			return tag.Value == f.Values[0];
 		}
+		
+		throw new Exception("Unknown filter " + f.Name);
+	}
+	
+	private Task<DescribeInstancesResponse> DescribeInstancesAsync(DescribeInstancesRequest request, CancellationToken cancellationToken)
+	{
+		 List<Instance> filteredInstances = request.Filters.Count > 0
+			? _instances.Values.Where(i => request.Filters.All(f => MatchesFilter(i, f))).ToList()
+			: _instances.Values.ToList();
 
 		return Task.FromResult(new DescribeInstancesResponse
 		{
-			Reservations = new() { new Reservation { Instances = filteredInstances.Select(CopyInstance).ToList() } },
+			Reservations = [new Reservation { Instances = filteredInstances.Select(CopyInstance).ToList() }],
 			HttpStatusCode = HttpStatusCode.OK
 		});
 	}
diff --git a/Engine/Source/Programs/Horde/Plugins/Compute/HordeServer.Compute/Agents/Fleet/Providers/AwsFleetManager.cs b/Engine/Source/Programs/Horde/Plugins/Compute/HordeServer.Compute/Agents/Fleet/Providers/AwsFleetManager.cs
--- a/Engine/Source/Programs/Horde/Plugins/Compute/HordeServer.Compute/Agents/Fleet/Providers/AwsFleetManager.cs
+++ b/Engine/Source/Programs/Horde/Plugins/Compute/HordeServer.Compute/Agents/Fleet/Providers/AwsFleetManager.cs
@@ -123,9 +123,16 @@ namespace HordeServer.Agents.Fleet.Providers
 	public sealed class AwsFleetManager : IFleetManager
 	{
 		/// <summary>
-		/// Prefix for signaling that an EC2 instance belong to an agent pool
+		/// EC2 tag prefix signaling an instance belong to agent pool, through pool name
+		/// Deprecated in favor of <see cref="AwsTagPoolId" /> as pool ID is the stable reference to a pool. 
 		/// </summary>
-		public const string PoolTagName = "Horde_Autoscale_Pool";
+		public const string AwsTagPoolName = "Horde_Autoscale_Pool";
+		
+		/// <summary>
+		/// EC2 tag prefix signaling an instance belong to agent pool, through pool ID
+		/// </summary>
+		public const string AwsTagPoolId = "Horde_Autoscale_PoolId";
+		
 		private const string AwsTagPropertyName = "aws-tag";
 
 		/// <summary>
@@ -241,11 +248,12 @@ namespace HordeServer.Agents.Fleet.Providers
 			span.SetAttribute("numAgents", agents.Count);
 			span.SetAttribute("count", count);
 
-			string awsTagProperty = $"{AwsTagPropertyName}={PoolTagName}:{pool.Name}";
+			string awsTagPoolNameProperty = $"{AwsTagPropertyName}={AwsTagPoolName}:{pool.Name}";
+			string awsTagPoolIdProperty = $"{AwsTagPropertyName}={AwsTagPoolId}:{pool.Id.ToString()}";
 
-			// Sort the agents by number of active leases. It's better to shutdown agents currently doing nothing.
+			// Sort the agents by number of active leases. It's better to shut down agents currently doing nothing.
 			List<IAgent> filteredAgents = agents.OrderBy(x => x.Leases.Count).ToList();
-			List<IAgent> agentsWithAwsTags = filteredAgents.Where(x => x.HasProperty(awsTagProperty)).ToList();
+			List<IAgent> agentsWithAwsTags = filteredAgents.Where(x => x.HasProperty(awsTagPoolNameProperty) || x.HasProperty(awsTagPoolIdProperty)).ToList();
 			List<IAgent> agentsLimitedByCount = agentsWithAwsTags.Take(count).ToList();
 
 			span.SetAttribute("agents.num", agents.Count);
@@ -289,14 +297,66 @@ namespace HordeServer.Agents.Fleet.Providers
 		/// <inheritdoc/>
 		public async Task<int> GetNumStoppedInstancesAsync(IPoolConfig pool, CancellationToken cancellationToken)
 		{
-			// Find all instances in the pool
-			DescribeInstancesRequest describeRequest = new();
-			describeRequest.Filters = new List<Filter>();
-			describeRequest.Filters.Add(new Filter("instance-state-name", new List<string> { InstanceStateName.Stopped.Value }));
-			describeRequest.Filters.Add(new Filter("tag:" + PoolTagName, new List<string> { pool.Name }));
-
-			DescribeInstancesResponse describeResponse = await _ec2.DescribeInstancesAsync(describeRequest, cancellationToken);
-			return describeResponse.Reservations.SelectMany(x => x.Instances).Select(x => x.InstanceId).Distinct().Count();
+			return (await GetInstancesInPoolAsync(_ec2, InstanceStateName.Stopped, pool, _tracer, cancellationToken)).Count;
+		}
+		
+		/// <summary>
+		/// Get instances that are part of given Horde agent pool
+		/// </summary>
+		/// <param name="ec2">EC2 client</param>
+		/// <param name="instanceState">Optional instance state to filter on</param>
+		/// <param name="pool">Pool to filter on</param>
+		/// <param name="tracer">Tracer</param>
+		/// <param name="cancellationToken">Cancellation token</param>
+		/// <returns>A list of EC2 instances</returns>
+		public static Task<List<Instance>> GetInstancesInPoolAsync(IAmazonEC2 ec2, InstanceStateName? instanceState, IPoolConfig pool, Tracer tracer, CancellationToken cancellationToken)
+		{
+			return GetInstancesInPoolAsync(ec2, instanceState, pool.Id.ToString(), pool.Name, tracer, cancellationToken);
+		}
+		
+		/// <summary>
+		/// Get instances that are part of given Horde agent pool
+		/// </summary>
+		/// <param name="ec2">EC2 client</param>
+		/// <param name="instanceState">Optional instance state to filter on</param>
+		/// <param name="poolId">Pool ID to filter on</param>
+		/// <param name="poolName">Pool ID to filter on</param>
+		/// <param name="tracer">Tracer</param>
+		/// <param name="cancellationToken">Cancellation token</param>
+		/// <returns>A list of EC2 instances</returns>
+		public static async Task<List<Instance>> GetInstancesInPoolAsync(IAmazonEC2 ec2, InstanceStateName? instanceState, string poolId, string poolName, Tracer tracer, CancellationToken cancellationToken)
+		{
+			using TelemetrySpan span = tracer.StartActiveSpan($"{nameof(AwsFleetManager)}.{nameof(GetInstancesInPoolAsync)}");
+			
+			List<Filter> poolIdFilters = [new Filter($"tag:{AwsTagPoolId}", [poolId])];
+			List<Filter> poolNameFilters = [new Filter($"tag:{AwsTagPoolName}", [poolName])];
+			if (instanceState != null)
+			{
+				poolIdFilters.Add(new Filter("instance-state-name", [instanceState.Value]));
+				poolNameFilters.Add(new Filter("instance-state-name", [instanceState.Value]));
+			}
+			
+			List<Instance> result = [];
+			{
+				// Get matching instances by pool *name*
+				using TelemetrySpan ec2Span = tracer.StartActiveSpan($"DescribeInstances");
+				DescribeInstancesRequest request = new() { Filters = poolNameFilters };
+				DescribeInstancesResponse response = await ec2.DescribeInstancesAsync(request, cancellationToken);
+				result.AddRange(response.Reservations.SelectMany(x => x.Instances));
+				ec2Span.SetAttribute("res.statusCode", (int)response.HttpStatusCode);
+				ec2Span.SetAttribute("res.numReservations", response.Reservations.Count);
+			}
+			{
+				// Get matching instances by pool *ID*
+				using TelemetrySpan ec2Span = tracer.StartActiveSpan($"DescribeInstances");
+				DescribeInstancesRequest request = new() { Filters = poolIdFilters };
+				DescribeInstancesResponse response = await ec2.DescribeInstancesAsync(request, cancellationToken);
+				result.AddRange(response.Reservations.SelectMany(x => x.Instances));
+				ec2Span.SetAttribute("res.statusCode", (int)response.HttpStatusCode);
+				ec2Span.SetAttribute("res.numReservations", response.Reservations.Count);
+			}
+			
+			return result.DistinctBy(x => x.InstanceId).ToList();
 		}
 	}
 }
diff --git a/Engine/Source/Programs/Horde/Plugins/Compute/HordeServer.Compute/Agents/Fleet/Providers/AwsRecyclingFleetManager.cs b/Engine/Source/Programs/Horde/Plugins/Compute/HordeServer.Compute/Agents/Fleet/Providers/AwsRecyclingFleetManager.cs
--- a/Engine/Source/Programs/Horde/Plugins/Compute/HordeServer.Compute/Agents/Fleet/Providers/AwsRecyclingFleetManager.cs
+++ b/Engine/Source/Programs/Horde/Plugins/Compute/HordeServer.Compute/Agents/Fleet/Providers/AwsRecyclingFleetManager.cs
@@ -399,20 +399,8 @@ public sealed class AwsRecyclingFleetManager : IFleetManager
 		span.SetAttribute("poolId", pool.Id.ToString());
 
 		// Find pending instances in the correct pool
-		DescribeInstancesRequest request = new()
-		{
-			Filters = new List<Filter>
-			{
-				new("instance-state-name", new List<string> { InstanceStateName.Pending.Value }),
-				new("tag:" + AwsFleetManager.PoolTagName, new List<string> { pool.Name })
-			}
-		};
-		DescribeInstancesResponse response = await _ec2.DescribeInstancesAsync(request, cancellationToken);
-		span.SetAttribute("res.statusCode", (int)response.HttpStatusCode);
-		span.SetAttribute("res.numReservations", response.Reservations.Count);
-
-		List<string> stuckInstanceIds = response.Reservations
-			.SelectMany(x => x.Instances)
+		List<Instance> stuckInstances = await AwsFleetManager.GetInstancesInPoolAsync(_ec2, InstanceStateName.Pending, pool, _tracer, cancellationToken);
+		List<string> stuckInstanceIds = stuckInstances
 			.Where(x => x.LaunchTime < _clock.UtcNow - pendingStateTimeout)
 			.Select(x => x.InstanceId)
 			.ToList();
@@ -438,22 +426,10 @@ public sealed class AwsRecyclingFleetManager : IFleetManager
 		span.SetAttribute("poolId", pool.Id.ToString());
 
 		// Find stopped instances in the correct pool
-		DescribeInstancesRequest request = new()
-		{
-			Filters = new List<Filter>
-			{
-				new("instance-state-name", new List<string> { InstanceStateName.Stopped.Value }),
-				new("tag:" + AwsFleetManager.PoolTagName, new List<string> { pool.Name })
-			}
-		};
-		DescribeInstancesResponse response = await _ec2.DescribeInstancesAsync(request, cancellationToken);
-		span.SetAttribute("res.statusCode", (int)response.HttpStatusCode);
-		span.SetAttribute("res.numReservations", response.Reservations.Count);
+		List<Instance> instances = await AwsFleetManager.GetInstancesInPoolAsync(_ec2, InstanceStateName.Stopped, pool, _tracer, cancellationToken);
 
 		// Group instance objects by availability zone
-		List<IGrouping<string, Instance>> azInstancePairs = response.Reservations
-			.SelectMany(x => x.Instances)
-			.GroupBy(x => x.Placement.AvailabilityZone).ToList();
+		List<IGrouping<string, Instance>> azInstancePairs = instances.GroupBy(x => x.Placement.AvailabilityZone).ToList();
 
 		return azInstancePairs.ToDictionary(x => x.Key, y => y.ToList());
 	}
@@ -469,13 +445,7 @@ public sealed class AwsRecyclingFleetManager : IFleetManager
 	public async Task<int> GetNumStoppedInstancesAsync(IPoolConfig pool, CancellationToken cancellationToken)
 	{
 		// Find all instances in the pool
-		DescribeInstancesRequest describeRequest = new DescribeInstancesRequest();
-		describeRequest.Filters = new List<Filter>();
-		describeRequest.Filters.Add(new Filter("instance-state-name", new List<string> { InstanceStateName.Stopped.Value }));
-		describeRequest.Filters.Add(new Filter("tag:" + AwsFleetManager.PoolTagName, new List<string> { pool.Name }));
-
-		DescribeInstancesResponse describeResponse = await _ec2.DescribeInstancesAsync(describeRequest, cancellationToken);
-		return describeResponse.Reservations.SelectMany(x => x.Instances).Select(x => x.InstanceId).Distinct().Count();
+		return (await AwsFleetManager.GetInstancesInPoolAsync(_ec2, InstanceStateName.Stopped, pool, _tracer, cancellationToken)).Count;
 	}
 
 	private static bool IsResponseSuccessful(AmazonWebServiceResponse response)
diff --git a/Engine/Source/Programs/Horde/Plugins/Compute/HordeServer.Compute/Agents/Fleet/Providers/AwsReuseFleetManager.cs b/Engine/Source/Programs/Horde/Plugins/Compute/HordeServer.Compute/Agents/Fleet/Providers/AwsReuseFleetManager.cs
--- a/Engine/Source/Programs/Horde/Plugins/Compute/HordeServer.Compute/Agents/Fleet/Providers/AwsReuseFleetManager.cs
+++ b/Engine/Source/Programs/Horde/Plugins/Compute/HordeServer.Compute/Agents/Fleet/Providers/AwsReuseFleetManager.cs
@@ -66,20 +66,7 @@ namespace HordeServer.Agents.Fleet.Providers
 			span.SetAttribute("numAgents", agents.Count);
 			span.SetAttribute("count", requestedInstancesCount);
 
-			DescribeInstancesResponse describeResponse;
-			using (TelemetrySpan describeSpan = _tracer.StartActiveSpan("DescribeInstances"))
-			{
-				// Find stopped instances in the correct pool
-				DescribeInstancesRequest describeRequest = new();
-				describeRequest.Filters = new List<Filter>();
-				describeRequest.Filters.Add(new Filter("instance-state-name", new List<string> { InstanceStateName.Stopped.Value }));
-				describeRequest.Filters.Add(new Filter("tag:" + AwsFleetManager.PoolTagName, new List<string> { pool.Name }));
-				describeResponse = await _ec2.DescribeInstancesAsync(describeRequest, cancellationToken);
-				describeSpan.SetAttribute("res.statusCode", (int)describeResponse.HttpStatusCode);
-				describeSpan.SetAttribute("res.numReservations", describeResponse.Reservations.Count);
-			}
-
-			IEnumerable<Instance> stoppedInstances = describeResponse.Reservations.SelectMany(x => x.Instances);
+			List<Instance> stoppedInstances = await AwsFleetManager.GetInstancesInPoolAsync(_ec2, InstanceStateName.Stopped, pool, _tracer, cancellationToken);
 
 			if (Settings.InstanceTypes is { Count: >= 1 })
 			{
@@ -109,7 +96,7 @@ namespace HordeServer.Agents.Fleet.Providers
 			{
 				// Try to start the given instances
 				StartInstancesRequest startRequest = new();
-				startRequest.InstanceIds.AddRange(describeResponse.Reservations.SelectMany(x => x.Instances).Select(x => x.InstanceId).Take(requestedInstancesCount));
+				startRequest.InstanceIds.AddRange(stoppedInstances.Select(x => x.InstanceId).Take(requestedInstancesCount));
 				int stoppedInstancesCount = startRequest.InstanceIds.Count;
 
 				startSpan.SetAttribute("req.instanceIds", String.Join(",", startRequest.InstanceIds));
@@ -173,14 +160,7 @@ namespace HordeServer.Agents.Fleet.Providers
 		/// <inheritdoc/>
 		public async Task<int> GetNumStoppedInstancesAsync(IPoolConfig pool, CancellationToken cancellationToken)
 		{
-			// Find all instances in the pool
-			DescribeInstancesRequest describeRequest = new DescribeInstancesRequest();
-			describeRequest.Filters = new List<Filter>();
-			describeRequest.Filters.Add(new Filter("instance-state-name", new List<string> { InstanceStateName.Stopped.Value }));
-			describeRequest.Filters.Add(new Filter("tag:" + AwsFleetManager.PoolTagName, new List<string> { pool.Name }));
-
-			DescribeInstancesResponse describeResponse = await _ec2.DescribeInstancesAsync(describeRequest, cancellationToken);
-			return describeResponse.Reservations.SelectMany(x => x.Instances).Select(x => x.InstanceId).Distinct().Count();
+			return (await AwsFleetManager.GetInstancesInPoolAsync(_ec2, InstanceStateName.Stopped, pool, _tracer, cancellationToken)).Count;
 		}
 	}
 }
