Horde: Update release notes

#rnx

diff --git a/Engine/Source/Programs/Horde/Docs/ReleaseNotes.md b/Engine/Source/Programs/Horde/Docs/ReleaseNotes.md
--- a/Engine/Source/Programs/Horde/Docs/ReleaseNotes.md
+++ b/Engine/Source/Programs/Horde/Docs/ReleaseNotes.md
@@ -1,5 +1,35 @@
 # Release Notes
 
+## 2025-08-14
+
+* Disable logging of UAT exit code in JobExecutor Reverts 44661300. (44854494)
+* [Horde][BuildHealth][Dashboard] Job History time span filter & removal of job count. (44837915)
+* Add logging to shutdown of disabled agents (44820850)
+* Fix duplicate key error in AddBlobAsync when converting shadow blobs
+  * Shadow blobs are placeholders for referenced but not-yet-uploaded blobs
+  * Upsert filter included Shadow=true condition, causing duplicate key errors when blob existed with Shadow=false
+  * Removed Shadow condition from filter to handle all cases (new blobs, shadow conversions, and updates)
+  * Added tests for race conditions and concurrent shadow conversions (44788210)
+* Remove unused setting for compute queue AWS metric in PoolConfig
+  This has since long moved into a separate strategy - ComputeQueueAwsMetricStrategy. (44775313)
+* [Horde][Analytics] Initial proof of concept for Build Health Dashboard View. (44772914)
+* Horde Dashboard: repopulate adv add arg quotes as escaped quotes (44755079)
+* Remove unused WithDatadog setting (44753720)
+* Add a named HTTP client for storage with retries
+  Under high load certain operations may fail against the storage backend (for example S3). "503 Slow down" is a common one which can be retried. (44750205)
+* Cloud DDC - Added conversion of float fields when converting compact binary into json (44746011)
+* [Horde] Update GetLegacyDeviceResponse to return device tags (44702624)
+* Ensure unhandled exceptions are logged from controllers (44699863)
+* Simplify running of health checks in Perforce load balancer (44695146)
+* [Horde Dashboard] Adds support for viewing and editing device tags on the horde dashboard's device editor (44669255)
+* [Horde] Adds configurable device tags
+  Device tags are small bits of metadata that can be attached to each device
+  When making a reservation, users can specify tags in a variety of ways
+  Required Tags - Devices that do not have this tag will not be reserved
+  Preferred Tags - Devices that have this tag will be selected over devices that do not, provided at least one is available
+  Undesired Tags - Devices that have this tag will be de-prioritized over devices that do not
+  Blocked Tags - Devices that have this tag will not be reserved (44669125)
+
 ## 2025-08-07
 
 * Log an error when UAT terminates with a non-zero exit code
