Horde Agent: Fail step when AutomationTool process has a non-zero exit code

diff --git a/Engine/Source/Programs/Horde/Drivers/JobDriver/Execution/JobExecutor.cs b/Engine/Source/Programs/Horde/Drivers/JobDriver/Execution/JobExecutor.cs
--- a/Engine/Source/Programs/Horde/Drivers/JobDriver/Execution/JobExecutor.cs
+++ b/Engine/Source/Programs/Horde/Drivers/JobDriver/Execution/JobExecutor.cs
@@ -809,9 +809,7 @@ namespace JobDriver.Execution
 			int exitCode = await ExecuteAutomationToolAsync(step, workspaceDir, arguments.ToString(), useP4, logger, cancellationToken);
 			if (exitCode != 0)
 			{
-				// Disable this log line until an appropriate structured/systemtic log event can be assigned
-				// for correct grouping with other events in Build Health.
-				// logger.LogError("BuildGraph setup terminated with non-zero exit code: {ExitCode}", exitCode);
+				logger.LogError(KnownLogEvents.ExitCode, "BuildGraph setup terminated with non-zero exit code: {ExitCode}", exitCode);
 				return false;
 			}
 
@@ -1258,9 +1256,7 @@ namespace JobDriver.Execution
 			int exitCode = await ExecuteAutomationToolAsync(step, workspaceDir, arguments, useP4, logger, cancellationToken);
 			if (exitCode != 0)
 			{
-				// Disable this log line until an appropriate structured/systemtic log event can be assigned
-				// for correct grouping with other events in Build Health.
-				// logger.LogError("AutomationTool (UAT) terminated with non-zero exit code: {ExitCode}", exitCode);
+				logger.LogError(KnownLogEvents.ExitCode, "AutomationTool (UAT) terminated with non-zero exit code: {ExitCode}", exitCode);
 				return false;
 			}
 
diff --git a/Engine/Source/Programs/Horde/Plugins/Build/HordeServer.Build.Tests/Issues/IssueServiceTests.cs b/Engine/Source/Programs/Horde/Plugins/Build/HordeServer.Build.Tests/Issues/IssueServiceTests.cs
--- a/Engine/Source/Programs/Horde/Plugins/Build/HordeServer.Build.Tests/Issues/IssueServiceTests.cs
+++ b/Engine/Source/Programs/Horde/Plugins/Build/HordeServer.Build.Tests/Issues/IssueServiceTests.cs
@@ -3156,5 +3156,34 @@ namespace HordeServer.Tests.Issues
 				Assert.AreEqual(new StreamId("ue5-release"), results[0].StreamId);
 			}
 		}
+
+		[TestMethod]
+		public async Task AutomationToolExitCodeTestAsync()
+		{
+			IJob job1 = CreateJob(MainStreamId, 125, "Test Build", Graph);
+			await using (TestJsonLogger logger = await CreateLoggerAsync(job1, 0, 0))
+			{
+				logger.LogError(KnownLogEvents.ExitCode, "AutomationTool (UAT) terminated with non-zero exit code: {ExitCode}", -1);
+			}
+			await UpdateCompleteStepAsync(job1, 0, 0, JobStepOutcome.Failure);
+
+			IReadOnlyList<IIssue> issues1 = await IssueCollection.FindIssuesAsync();
+			Assert.AreEqual(1, issues1.Count);
+			Assert.AreEqual("Errors in Update Version Files", issues1[0].Summary);
+
+			IJob job2 = CreateJob(MainStreamId, 170, "Cook Test", Graph);
+			await using (TestJsonLogger logger = await CreateLoggerAsync(job2, 0, 1))
+			{
+				logger.LogError(KnownLogEvents.Gauntlet_TestEvent, "Test {Name} failed", "Bar.Foo.Test");
+				logger.LogError(KnownLogEvents.ExitCode, "AutomationTool (UAT) terminated with non-zero exit code: {ExitCode}", -1);
+			}
+			await UpdateCompleteStepAsync(job2, 0, 1, JobStepOutcome.Failure);
+
+			IReadOnlyList<IIssue> issues2 = await IssueCollection.FindIssuesAsync();
+			issues2 = issues2.OrderBy(x => x.Id).ToList();
+			Assert.AreEqual(2, issues2.Count);
+			Assert.AreEqual("Errors in Update Version Files", issues2[0].Summary);
+			Assert.AreEqual("Automation test errors in Compile UnrealHeaderTool Win64", issues2[1].Summary);
+		}
 	}
 }
diff --git a/Engine/Source/Programs/Horde/Plugins/Build/HordeServer.Build/Issues/IssueService.cs b/Engine/Source/Programs/Horde/Plugins/Build/HordeServer.Build/Issues/IssueService.cs
--- a/Engine/Source/Programs/Horde/Plugins/Build/HordeServer.Build/Issues/IssueService.cs
+++ b/Engine/Source/Programs/Horde/Plugins/Build/HordeServer.Build/Issues/IssueService.cs
@@ -754,13 +754,37 @@ namespace HordeServer.Issues
 			}
 			handlers.SortBy(x => -x.Priority);
 
-			// Create all the issue definitions by passing each log event to the handlers in order until one attaches it to an issue
+			// General events which have no significant context
+			HashSet<EventId> generalEvents = new HashSet<EventId> { KnownLogEvents.ExitCode };
+
+			// Generate internal issue events from log anchors
 			List<ILogAnchor> stepAnchors = await log.GetAnchorsAsync(cancellationToken: cancellationToken);
+			List<IssueEventInternal> issueEvents = new List<IssueEventInternal>();
+
+			bool found = false;
 			foreach (ILogAnchor stepAnchor in stepAnchors)
 			{
 				ILogEventData stepEventData = await stepAnchor.GetDataAsync(cancellationToken);
 
-				IssueEventInternal issueEvent = new IssueEventInternal(stepAnchor, stepEventData);
+				bool general = stepEventData.EventId.HasValue && generalEvents.Contains(stepEventData.EventId.Value);
+
+				// Remove any general events after specific events, as these generate noise in those cases
+				if (found && general)
+				{
+					continue;
+				}
+
+				if (!found && !general)
+				{
+					found = true;
+				}
+
+				issueEvents.Add(new IssueEventInternal(stepAnchor, stepEventData));
+			}			
+
+			// Create all the issue definitions by passing each log event to the handlers in order until one attaches it to an issue
+			foreach (IssueEventInternal issueEvent in issueEvents)
+			{
 				foreach (IssueHandler handler in handlers)
 				{
 					if (handler.HandleEvent(issueEvent))
@@ -769,7 +793,7 @@ namespace HordeServer.Issues
 					}
 				}
 			}
-
+			
 			// Gather all the issue event groups
 			List<IssueEventGroup> issues = new List<IssueEventGroup>();
 			foreach (IssueHandler handler in handlers)
