Horde: Add raw log event data, job/step information, and issue grouping details to issue audit log.

diff --git a/Engine/Source/Programs/Horde/Plugins/Build/HordeServer.Build/Issues/IssueService.cs b/Engine/Source/Programs/Horde/Plugins/Build/HordeServer.Build/Issues/IssueService.cs
--- a/Engine/Source/Programs/Horde/Plugins/Build/HordeServer.Build/Issues/IssueService.cs
+++ b/Engine/Source/Programs/Horde/Plugins/Build/HordeServer.Build/Issues/IssueService.cs
@@ -780,11 +780,15 @@ namespace HordeServer.Issues
 				}
 
 				issueEvents.Add(new IssueEventInternal(stepAnchor, stepEventData));
-			}			
+			}
 
 			// Create all the issue definitions by passing each log event to the handlers in order until one attaches it to an issue
 			foreach (IssueEventInternal issueEvent in issueEvents)
 			{
+				// Attach issue audit log buffers to each issue event
+				issueEvent.AuditLogger = new IssueAuditLogger();
+				issueEvent.AuditLogger.LogDebug("Logging events from Job {JobId}, step {StepId}", job.Id, step.Id);
+
 				foreach (IssueHandler handler in handlers)
 				{
 					if (handler.HandleEvent(issueEvent))
@@ -793,7 +797,7 @@ namespace HordeServer.Issues
 					}
 				}
 			}
-			
+
 			// Gather all the issue event groups
 			List<IssueEventGroup> issues = new List<IssueEventGroup>();
 			foreach (IssueHandler handler in handlers)
@@ -808,6 +812,10 @@ namespace HordeServer.Issues
 				foreach (IssueEventGroup issue in issues)
 				{
 					issue.Type = $"{issue.Type}:{group}";
+					foreach (IssueEvent e in issue.Events)
+					{
+						e.AuditLogger?.LogDebug("Fingerprint: '{Type}' modified by node annotation: '{Group}'", issue.Type, group);
+					}
 				}
 			}
 
@@ -884,6 +892,11 @@ namespace HordeServer.Issues
 						foreach (IssueEventGroupInternal eventGroup in matchEventGroups)
 						{
 							_logger.LogDebug("Matched fingerprint {Digest} ({NumLogEvents} log events) to span {SpanId}", eventGroup.Id, eventGroup.Events.Count, newSpan.Id);
+
+							foreach (IssueEvent e in eventGroup.Events)
+							{
+								e.AuditLogger?.LogDebug("Matched {Event} log event to span {SpanId}", e.Render(), newSpan.Id);
+							}
 						}
 
 						// Assign all the events to the span
@@ -891,6 +904,10 @@ namespace HordeServer.Issues
 
 						// Remove the matches from the set of events
 						newEventGroups.ExceptWith(matchEventGroups);
+
+						// Flush audit log buffers in event groups
+						ILogger issueLogger = _issueCollection.GetLogger(newSpan.IssueId);
+						EmitIssueAuditLogs(issueLogger, matchEventGroups);
 					}
 
 					// Add the span id to the list of checked spans, so we don't have to check it again
@@ -923,6 +940,46 @@ namespace HordeServer.Issues
 			return severity;
 		}
 
+		/// <summary>
+		/// Dedupes and emits all buffered audit log entries in the event groups
+		/// </summary>
+		/// <param name="auditLogger">Issue audit log chanel to log to</param>
+		/// <param name="eventGroups">List of event groups to emit buffered log entries for</param>
+		static void EmitIssueAuditLogs(ILogger auditLogger, List<IssueEventGroupInternal> eventGroups)
+		{
+			HashSet<IssueAuditLogEntry> uniqueEntries = new HashSet<IssueAuditLogEntry>();
+			foreach (IssueEventGroupInternal eventGroup in eventGroups)
+			{
+				uniqueEntries.UnionWith(eventGroup.Events.SelectMany(e => (e.AuditLogger != null && e.AuditLogger?.GetEntries().Count > 0) ? e.AuditLogger.GetEntries() : []));
+			}
+
+			if (uniqueEntries != null && uniqueEntries.Count > 0)
+			{
+				foreach (IssueAuditLogEntry e in uniqueEntries)
+				{
+					if (e._messageTemplate != null)
+					{
+						if (e._logLevel == LogLevel.Error)
+						{
+							auditLogger.LogError(e._messageTemplate, e._args);
+						}
+						else if (e._logLevel == LogLevel.Warning)
+						{
+							auditLogger.LogWarning(e._messageTemplate, e._args);
+						}
+						else if (e._logLevel == LogLevel.Information)
+						{
+							auditLogger.LogInformation(e._messageTemplate, e._args);
+						}
+						else
+						{
+							auditLogger.LogDebug(e._messageTemplate, e._args);
+						}
+					}
+				}
+			}
+		}
+
 		/// <summary>
 		/// Adds new spans for the given events
 		/// </summary>
@@ -954,6 +1011,11 @@ namespace HordeServer.Issues
 					{
 						IssueEventGroupInternal newEventGroup = eventGroup.MergeWith(otherEventGroup);
 						_logger.LogDebug("Merging group {Group} with group {OtherGroup} to form {NewGroup}", eventGroup.Id.ToString(), otherEventGroup.Id.ToString(), newEventGroup.Id.ToString());
+						foreach (IssueEventInternal e in eventGroup.Events)
+						{
+							e.AuditLogger?.LogDebug("Group {Group} merged with group {OtherGroup} to form group {NewGroup} based on matching fingerprint: '{Fingerprint}'", eventGroup.Id.ToString(), otherEventGroup.Id.ToString(), newEventGroup.Id.ToString(), eventGroup.Fingerprint);
+						}
+			
 						sourceEventGroups.Add(otherEventGroup);
 						eventGroup = newEventGroup;
 					}
@@ -982,6 +1044,11 @@ namespace HordeServer.Issues
 				IIssue newIssue = await FindOrAddIssueForSpanAsync(spanData, cancellationToken);
 				IIssueSpan newSpan = await _issueCollection.AddSpanAsync(newIssue.Id, spanData, cancellationToken);
 				IIssueStep newStep = await _issueCollection.AddStepAsync(newSpan.Id, stepData, cancellationToken);
+
+				// Flush audit log buffers in event group
+				ILogger issueLogger = _issueCollection.GetLogger(newIssue.Id);
+				EmitIssueAuditLogs(issueLogger, [eventGroup]);
+
 				await UpdateIssueDerivedDataAsync(newIssue, cancellationToken);
 
 				// Update the log events
diff --git a/Engine/Source/Programs/Shared/EpicGames.Horde/Issues/Handlers/BuildGraphIssueHandler.cs b/Engine/Source/Programs/Shared/EpicGames.Horde/Issues/Handlers/BuildGraphIssueHandler.cs
--- a/Engine/Source/Programs/Shared/EpicGames.Horde/Issues/Handlers/BuildGraphIssueHandler.cs
+++ b/Engine/Source/Programs/Shared/EpicGames.Horde/Issues/Handlers/BuildGraphIssueHandler.cs
@@ -41,6 +41,7 @@ namespace EpicGames.Horde.Issues.Handlers
 					issue.Events.Add(issueEvent);
 					issue.Keys.AddSourceFiles(issueEvent);
 					_issues.Add(issue);
+					issueEvent.AuditLogger?.LogDebug("BuildGraph issue added by event: '{Event}'", issueEvent.Render());
 
 					return true;
 				}
diff --git a/Engine/Source/Programs/Shared/EpicGames.Horde/Issues/Handlers/CompileIssueHandler.cs b/Engine/Source/Programs/Shared/EpicGames.Horde/Issues/Handlers/CompileIssueHandler.cs
--- a/Engine/Source/Programs/Shared/EpicGames.Horde/Issues/Handlers/CompileIssueHandler.cs
+++ b/Engine/Source/Programs/Shared/EpicGames.Horde/Issues/Handlers/CompileIssueHandler.cs
@@ -73,6 +73,7 @@ namespace EpicGames.Horde.Issues.Handlers
 					issue.Keys.AddSourceFiles(issueEvent);
 					issue.Metadata.Add(CompileTypeAnnotation, compileType);
 					_issues.Add(issue);
+					issueEvent.AuditLogger?.LogDebug("{Type} issue added for event: '{Event}'", fingerprintType, issueEvent.Render());
 
 					return true;
 				}
diff --git a/Engine/Source/Programs/Shared/EpicGames.Horde/Issues/Handlers/ContentIssueHandler.cs b/Engine/Source/Programs/Shared/EpicGames.Horde/Issues/Handlers/ContentIssueHandler.cs
--- a/Engine/Source/Programs/Shared/EpicGames.Horde/Issues/Handlers/ContentIssueHandler.cs
+++ b/Engine/Source/Programs/Shared/EpicGames.Horde/Issues/Handlers/ContentIssueHandler.cs
@@ -32,6 +32,7 @@ namespace EpicGames.Horde.Issues.Handlers
 					issue.Events.Add(issueEvent);
 					issue.Keys.AddAssets(issueEvent);
 					_issues.Add(issue);
+					issueEvent.AuditLogger?.LogDebug("{Type} issue added for event: '{Event}'", issue.Type, issueEvent.Render());
 
 					return true;
 				}
diff --git a/Engine/Source/Programs/Shared/EpicGames.Horde/Issues/Handlers/CopyrightIssueHandler.cs b/Engine/Source/Programs/Shared/EpicGames.Horde/Issues/Handlers/CopyrightIssueHandler.cs
--- a/Engine/Source/Programs/Shared/EpicGames.Horde/Issues/Handlers/CopyrightIssueHandler.cs
+++ b/Engine/Source/Programs/Shared/EpicGames.Horde/Issues/Handlers/CopyrightIssueHandler.cs
@@ -2,6 +2,7 @@
 
 using System.Collections.Generic;
 using EpicGames.Core;
+using Microsoft.Extensions.Logging;
 
 namespace EpicGames.Horde.Issues.Handlers
 {
@@ -25,6 +26,7 @@ namespace EpicGames.Horde.Issues.Handlers
 				issue.Events.Add(issueEvent);
 				issue.Keys.AddSourceFiles(issueEvent);
 				_issues.Add(issue);
+				issueEvent.AuditLogger?.LogDebug("{Type} issue added for event: '{Event}'", issue.Type, issueEvent.Render());
 
 				return true;
 			}
diff --git a/Engine/Source/Programs/Shared/EpicGames.Horde/Issues/Handlers/DefaultIssueHandler.cs b/Engine/Source/Programs/Shared/EpicGames.Horde/Issues/Handlers/DefaultIssueHandler.cs
--- a/Engine/Source/Programs/Shared/EpicGames.Horde/Issues/Handlers/DefaultIssueHandler.cs
+++ b/Engine/Source/Programs/Shared/EpicGames.Horde/Issues/Handlers/DefaultIssueHandler.cs
@@ -1,6 +1,7 @@
 // Copyright Epic Games, Inc. All Rights Reserved.
 
 using System.Collections.Generic;
+using Microsoft.Extensions.Logging;
 
 namespace EpicGames.Horde.Issues.Handlers
 {
@@ -29,6 +30,7 @@ namespace EpicGames.Horde.Issues.Handlers
 				_issue = new IssueEventGroup("Default", "{Severity} in {Meta:Node}", IssueChangeFilter.All);
 				_issue.Metadata.Add("Node", _context.NodeName);
 				_issue.Keys.Add(IssueKey.FromStep(_context.StreamId, _context.TemplateId, _context.NodeName));
+				issueEvent.AuditLogger?.LogDebug("Fingerprint key: '{Key}' generated from event: '{Event}'", IssueKey.FromStep(_context.StreamId, _context.TemplateId, _context.NodeName), issueEvent.Render());
 			}
 
 			_issue.Events.Add(issueEvent);
diff --git a/Engine/Source/Programs/Shared/EpicGames.Horde/Issues/Handlers/ExternalIssueHandler.cs b/Engine/Source/Programs/Shared/EpicGames.Horde/Issues/Handlers/ExternalIssueHandler.cs
--- a/Engine/Source/Programs/Shared/EpicGames.Horde/Issues/Handlers/ExternalIssueHandler.cs
+++ b/Engine/Source/Programs/Shared/EpicGames.Horde/Issues/Handlers/ExternalIssueHandler.cs
@@ -4,6 +4,7 @@ using System.Collections.Generic;
 using System.Text.Json;
 using System.Text.Json.Nodes;
 using EpicGames.Core;
+using Microsoft.Extensions.Logging;
 
 namespace EpicGames.Horde.Issues.Handlers
 {
@@ -43,6 +44,7 @@ namespace EpicGames.Horde.Issues.Handlers
 					group.Events.Add(issueEvent);
 
 					_issues.Add(group);
+					issueEvent.AuditLogger?.LogDebug("{Type} issue added for event: '{Event}'", group.Type, issueEvent.Render());
 					return true;
 				}
 			}
diff --git a/Engine/Source/Programs/Shared/EpicGames.Horde/Issues/Handlers/GauntletIssueHandler.cs b/Engine/Source/Programs/Shared/EpicGames.Horde/Issues/Handlers/GauntletIssueHandler.cs
--- a/Engine/Source/Programs/Shared/EpicGames.Horde/Issues/Handlers/GauntletIssueHandler.cs
+++ b/Engine/Source/Programs/Shared/EpicGames.Horde/Issues/Handlers/GauntletIssueHandler.cs
@@ -166,12 +166,14 @@ namespace EpicGames.Horde.Issues.Handlers
 					key += $":{_context.StreamId}:{_context.NodeName}";
 				}
 				keys.Add(key, IssueKeyType.None);
+				issueEvent.AuditLogger?.LogDebug("Fingerprint key: '{Key}' generated from event: '{Event}'", key, issueEvent.Render());
 				metadata.Add("Hash", hash.ToString());
 			}
 			else
 			{
 				// Not enough information, make it an issue associated with only the job step
 				keys.Add($"{_context.StreamId}:{_context.NodeName}", IssueKeyType.None);
+				issueEvent.AuditLogger?.LogDebug("Fingerprint key: '{Key}' generated from event: '{Event}'", $"{_context.StreamId}:{_context.NodeName}", issueEvent.Render());
 			}
 			metadata.Add("Node", _context.NodeName);
 		}
diff --git a/Engine/Source/Programs/Shared/EpicGames.Horde/Issues/Handlers/HashedIssueHandler.cs b/Engine/Source/Programs/Shared/EpicGames.Horde/Issues/Handlers/HashedIssueHandler.cs
--- a/Engine/Source/Programs/Shared/EpicGames.Horde/Issues/Handlers/HashedIssueHandler.cs
+++ b/Engine/Source/Programs/Shared/EpicGames.Horde/Issues/Handlers/HashedIssueHandler.cs
@@ -81,6 +81,7 @@ namespace EpicGames.Horde.Issues.Handlers
 					IssueEventGroup issue = new IssueEventGroup("Hashed", "{Severity} in {Meta:Node}", IssueChangeFilter.All);
 					issue.Events.Add(stepEvent);
 					issue.Keys.AddHash(hash);
+					stepEvent.AuditLogger?.LogDebug("Fingerprint key: '{Key}' generated from event: '{Event}'", hash, stepEvent.Render());
 					issue.Metadata.Add("Node", _context.NodeName);
 					issues.Add(issue);
 				}
@@ -92,6 +93,7 @@ namespace EpicGames.Horde.Issues.Handlers
 						{
 							genericErrorsFingerprint = new IssueEventGroup("Hashed", "{Severity} in {Meta:Node}", IssueChangeFilter.All);
 							genericErrorsFingerprint.Keys.Add(IssueKey.FromStepAndSeverity(_context.StreamId, _context.TemplateId, _context.NodeName, LogLevel.Error));
+							stepEvent.AuditLogger?.LogDebug("Fingerprint key: '{Key}' generated from event: '{Event}'", IssueKey.FromStepAndSeverity(_context.StreamId, _context.TemplateId, _context.NodeName, LogLevel.Error), stepEvent.Render());
 							genericErrorsFingerprint.Metadata.Add("Node", _context.NodeName);
 							issues.Add(genericErrorsFingerprint);
 						}
@@ -103,6 +105,7 @@ namespace EpicGames.Horde.Issues.Handlers
 						{
 							genericFingerprint = new IssueEventGroup("Hashed", "{Severity} in {Meta:Node}", IssueChangeFilter.All);
 							genericFingerprint.Keys.Add(IssueKey.FromStep(_context.StreamId, _context.TemplateId, _context.NodeName));
+							stepEvent.AuditLogger?.LogDebug("Fingerprint key: '{Key}' generated from event: '{Event}'", IssueKey.FromStep(_context.StreamId, _context.TemplateId, _context.NodeName), stepEvent.Render());
 							genericFingerprint.Metadata.Add("Node", _context.NodeName);
 							issues.Add(genericFingerprint);
 						}
diff --git a/Engine/Source/Programs/Shared/EpicGames.Horde/Issues/Handlers/IgnoredIssueHandler.cs b/Engine/Source/Programs/Shared/EpicGames.Horde/Issues/Handlers/IgnoredIssueHandler.cs
--- a/Engine/Source/Programs/Shared/EpicGames.Horde/Issues/Handlers/IgnoredIssueHandler.cs
+++ b/Engine/Source/Programs/Shared/EpicGames.Horde/Issues/Handlers/IgnoredIssueHandler.cs
@@ -44,6 +44,7 @@ namespace EpicGames.Horde.Issues.Handlers
 			{
 				if (IsMatchingEventId(logEvent.EventId.Value))
 				{
+					logEvent.AuditLogger?.LogDebug("IgnoreIssueHandler matched event: '{Event}'", logEvent.Render());
 					return true;
 				}
 			}
diff --git a/Engine/Source/Programs/Shared/EpicGames.Horde/Issues/Handlers/LocalizationIssueHandler.cs b/Engine/Source/Programs/Shared/EpicGames.Horde/Issues/Handlers/LocalizationIssueHandler.cs
--- a/Engine/Source/Programs/Shared/EpicGames.Horde/Issues/Handlers/LocalizationIssueHandler.cs
+++ b/Engine/Source/Programs/Shared/EpicGames.Horde/Issues/Handlers/LocalizationIssueHandler.cs
@@ -275,6 +275,7 @@ namespace EpicGames.Horde.Issues.Handlers
 					if(issueGroup.Keys.Count > 0 && !_issues.Contains(issueGroup))
 					{
 						_issues.Add(issueGroup);
+						issueEvent.AuditLogger?.LogDebug("{Type} issue added for event: '{Event}'", issueGroup.Type, issueEvent.Render());
 					}
 
 					return true;
diff --git a/Engine/Source/Programs/Shared/EpicGames.Horde/Issues/Handlers/PerforceCaseIssueHandler.cs b/Engine/Source/Programs/Shared/EpicGames.Horde/Issues/Handlers/PerforceCaseIssueHandler.cs
--- a/Engine/Source/Programs/Shared/EpicGames.Horde/Issues/Handlers/PerforceCaseIssueHandler.cs
+++ b/Engine/Source/Programs/Shared/EpicGames.Horde/Issues/Handlers/PerforceCaseIssueHandler.cs
@@ -36,6 +36,7 @@ namespace EpicGames.Horde.Issues.Handlers
 				issue.Events.Add(issueEvent);
 				issue.Keys.AddDepotPaths(issueEvent);
 				_issues.Add(issue);
+				issueEvent.AuditLogger?.LogDebug("{Type} issue added for event: '{Event}'", issue.Type, issueEvent.Render());
 
 				return true;
 			}
diff --git a/Engine/Source/Programs/Shared/EpicGames.Horde/Issues/Handlers/SanitizerIssueHandler.cs b/Engine/Source/Programs/Shared/EpicGames.Horde/Issues/Handlers/SanitizerIssueHandler.cs
--- a/Engine/Source/Programs/Shared/EpicGames.Horde/Issues/Handlers/SanitizerIssueHandler.cs
+++ b/Engine/Source/Programs/Shared/EpicGames.Horde/Issues/Handlers/SanitizerIssueHandler.cs
@@ -75,6 +75,7 @@ namespace EpicGames.Horde.Issues.Handlers
 				issue.Events.Add(issueEvent);
 
 				_issues.Add(issue);
+				issueEvent.AuditLogger?.LogDebug("{Type} issue added for event: '{Event}'", issue.Type, issueEvent.Render());
 				return true;
 			}
 			return false;
diff --git a/Engine/Source/Programs/Shared/EpicGames.Horde/Issues/Handlers/ScopedIssueHandler.cs b/Engine/Source/Programs/Shared/EpicGames.Horde/Issues/Handlers/ScopedIssueHandler.cs
--- a/Engine/Source/Programs/Shared/EpicGames.Horde/Issues/Handlers/ScopedIssueHandler.cs
+++ b/Engine/Source/Programs/Shared/EpicGames.Horde/Issues/Handlers/ScopedIssueHandler.cs
@@ -6,6 +6,7 @@ using System.Text;
 using System.Text.Json;
 using System.Text.RegularExpressions;
 using EpicGames.Core;
+using Microsoft.Extensions.Logging;
 
 namespace EpicGames.Horde.Issues.Handlers
 {
@@ -74,6 +75,7 @@ namespace EpicGames.Horde.Issues.Handlers
 				issue.Metadata.Add(NodeName, _context.NodeName);
 				issue.Metadata.Add(ScopeName, scope);
 				_issues.Add(issue);
+				logEvent.AuditLogger?.LogDebug("{Type} issue added for event: '{Event}'", issue.Type, logEvent.Render());
 
 				return true;
 			}
diff --git a/Engine/Source/Programs/Shared/EpicGames.Horde/Issues/Handlers/ShaderIssueHandler.cs b/Engine/Source/Programs/Shared/EpicGames.Horde/Issues/Handlers/ShaderIssueHandler.cs
--- a/Engine/Source/Programs/Shared/EpicGames.Horde/Issues/Handlers/ShaderIssueHandler.cs
+++ b/Engine/Source/Programs/Shared/EpicGames.Horde/Issues/Handlers/ShaderIssueHandler.cs
@@ -32,6 +32,7 @@ namespace EpicGames.Horde.Issues.Handlers
 					issue.Events.Add(logEvent);
 					issue.Keys.AddSourceFiles(logEvent);
 					_issues.Add(issue);
+					logEvent.AuditLogger?.LogDebug("{Type} issue added for event: '{Event}'", issue.Type, logEvent.Render());
 
 					return true;
 				}
diff --git a/Engine/Source/Programs/Shared/EpicGames.Horde/Issues/Handlers/SymbolIssueHandler.cs b/Engine/Source/Programs/Shared/EpicGames.Horde/Issues/Handlers/SymbolIssueHandler.cs
--- a/Engine/Source/Programs/Shared/EpicGames.Horde/Issues/Handlers/SymbolIssueHandler.cs
+++ b/Engine/Source/Programs/Shared/EpicGames.Horde/Issues/Handlers/SymbolIssueHandler.cs
@@ -55,6 +55,7 @@ namespace EpicGames.Horde.Issues.Handlers
 					if (issue.Keys.Count > 0)
 					{
 						_issues.Add(issue);
+						issueEvent.AuditLogger?.LogDebug("{Type} issue added for event: '{Event}'", issue.Type, issueEvent.Render());
 						return true;
 					}
 					if (_issues.Count > 0)
diff --git a/Engine/Source/Programs/Shared/EpicGames.Horde/Issues/Handlers/SystemicIssueHandler.cs b/Engine/Source/Programs/Shared/EpicGames.Horde/Issues/Handlers/SystemicIssueHandler.cs
--- a/Engine/Source/Programs/Shared/EpicGames.Horde/Issues/Handlers/SystemicIssueHandler.cs
+++ b/Engine/Source/Programs/Shared/EpicGames.Horde/Issues/Handlers/SystemicIssueHandler.cs
@@ -56,6 +56,7 @@ namespace EpicGames.Horde.Issues.Handlers
 						issue.Keys.Add(IssueKey.FromStep(_context.StreamId, _context.TemplateId, _context.NodeName));
 						issue.Events.Add(logEvent);
 						_issues.Add(issue);
+						logEvent.AuditLogger?.LogDebug("Fingerprint key: '{Key}' generated from event: '{Event}'", $"{_context.StreamId}:{_context.NodeName}", logEvent.Render());
 
 						return true;
 					}
diff --git a/Engine/Source/Programs/Shared/EpicGames.Horde/Issues/Handlers/UnacceptableWordsIssueHandler.cs b/Engine/Source/Programs/Shared/EpicGames.Horde/Issues/Handlers/UnacceptableWordsIssueHandler.cs
--- a/Engine/Source/Programs/Shared/EpicGames.Horde/Issues/Handlers/UnacceptableWordsIssueHandler.cs
+++ b/Engine/Source/Programs/Shared/EpicGames.Horde/Issues/Handlers/UnacceptableWordsIssueHandler.cs
@@ -2,6 +2,7 @@
 
 using System.Collections.Generic;
 using EpicGames.Core;
+using Microsoft.Extensions.Logging;
 
 namespace EpicGames.Horde.Issues.Handlers
 {
@@ -25,6 +26,7 @@ namespace EpicGames.Horde.Issues.Handlers
 				issue.Events.Add(issueEvent);
 				issue.Keys.AddSourceFiles(issueEvent);
 				_issues.Add(issue);
+				issueEvent.AuditLogger?.LogDebug("{Type} issue added for event: '{Event}'", issue.Type, issueEvent.Render());
 
 				return true;
 			}
diff --git a/Engine/Source/Programs/Shared/EpicGames.Horde/Issues/IssueEvent.cs b/Engine/Source/Programs/Shared/EpicGames.Horde/Issues/IssueEvent.cs
--- a/Engine/Source/Programs/Shared/EpicGames.Horde/Issues/IssueEvent.cs
+++ b/Engine/Source/Programs/Shared/EpicGames.Horde/Issues/IssueEvent.cs
@@ -33,6 +33,11 @@ namespace EpicGames.Horde.Issues
 		/// </summary>
 		public IReadOnlyList<JsonLogEvent> Lines { get; }
 
+		/// <summary>
+		/// <see cref="IssueAuditLogger"/> attached to this event
+		/// </summary>
+		public IssueAuditLogger? AuditLogger { get; set; }
+
 		/// <summary>
 		/// Constructor
 		/// </summary>
@@ -113,4 +118,159 @@ namespace EpicGames.Horde.Issues
 		{
 		}
 	}
+
+	/// <summary>
+	/// Temporary log buffer to store information about an <see cref="IssueEvent"/> before it's assigned to a build health issue
+	/// </summary>
+	public class IssueAuditLogger : ILogger
+	{
+		/// <inheritdoc/>
+		public IDisposable? BeginScope<TState>(TState state) where TState : notnull => null;
+
+		/// <inheritdoc/>
+		public bool IsEnabled(LogLevel logLevel) => true;
+
+		private readonly List<IssueAuditLogEntry> _entries;
+
+		/// <summary>
+		/// Constructor
+		/// </summary>
+		public IssueAuditLogger()
+		{
+			_entries = new List<IssueAuditLogEntry>();
+		}
+
+		/// <summary>
+		/// Gets list of <see cref="IssueAuditLogEntry"/> objects in the buffer
+		/// </summary>
+		/// <returns>List of log entries</returns>
+		public List<IssueAuditLogEntry> GetEntries()
+		{
+			return _entries;
+		}
+
+		/// <summary>
+		/// Stores a <see cref="IssueAuditLogEntry"/> in the log buffer
+		/// </summary>
+		/// <param name="logLevel">Entry will be written on this level.</param>
+		/// <param name="eventId">Id of the event.</param>
+		/// <param name="state">The entry to be written. Can be also an object.</param>
+		/// <param name="exception">The exception related to this entry.</param>
+		/// <param name="formatter">Function to create a <see cref="String"/> message of the <paramref name="state"/> and <paramref name="exception"/>.</param>
+		/// <typeparam name="TState">The type of the object to be written.</typeparam>
+		public void Log<TState>(LogLevel logLevel, EventId eventId, TState state, Exception? exception, Func<TState, Exception?, string> formatter)
+		{
+			// extract and store message template and args
+			string? messageTemplate = null;
+			List<object> args = [];
+			IEnumerable<KeyValuePair<string, object>>? innerEnumerable = state as IEnumerable<KeyValuePair<string, object>>;
+			if (innerEnumerable != null)
+			{
+				foreach (KeyValuePair<string, object> pair in innerEnumerable)
+				{
+					if (pair.Key.Equals("{OriginalFormat}", StringComparison.Ordinal))
+					{
+						messageTemplate = pair.Value.ToString();
+					}
+					else
+					{
+						args.Add(pair.Value);
+					}
+				}
+			}
+
+			IssueAuditLogEntry entry = new IssueAuditLogEntry() { _logLevel = logLevel, _eventId = eventId, _messageTemplate = messageTemplate, _args = [.. args], _exception = exception };
+			_entries.Add(entry);
+		}
+	}
+
+	/// <summary>
+	/// Buffered log entry stored in a <see cref="IssueAuditLogger"/>
+	/// </summary>
+	public struct IssueAuditLogEntry : IEquatable<IssueAuditLogEntry>
+	{
+		/// <summary>
+		/// Severity level at which to log the entry
+		/// </summary>
+		public LogLevel _logLevel;
+
+		/// <summary>
+		/// Id assocaited with the entry
+		/// </summary>
+		public EventId _eventId;
+
+		/// <summary>
+		/// Message template for structured log entry
+		/// </summary>
+		public string? _messageTemplate;
+
+		/// <summary>
+		/// Arguments to be injected in the message template of the entry
+		/// </summary>
+		public object?[] _args;
+
+
+		/// <summary>
+		/// The exception related to this entry.
+		/// </summary>
+		public Exception? _exception;
+
+		/// <summary>
+		/// Compares against another IssueAuditLogEntry object for equality
+		/// </summary>
+		/// <param name="otherEntry">Other entry to compare against</param>
+		/// <returns>True if the two entries have the same log level, message template, and arguments.</returns>
+		public readonly bool Equals(IssueAuditLogEntry otherEntry)
+		{
+			if (_logLevel != otherEntry._logLevel || _messageTemplate != otherEntry._messageTemplate)
+			{
+				return false;
+			}
+
+			if (_args.Length != otherEntry._args.Length)
+			{
+				return false;
+			}
+
+			for (int i = 0; i < _args.Length; i++)
+			{
+				if (!_args[i]!.Equals(otherEntry._args[i]))
+				{
+					return false;
+				}
+			}
+
+			return true;
+		}
+
+		/// <summary>
+		/// Returns the underlying value hashcode
+		/// </summary>
+		public override readonly int GetHashCode()
+		{
+			return HashCode.Combine(_logLevel, _messageTemplate);
+		}
+
+		/// <inheritdoc/>
+		public override readonly bool Equals(object? obj)
+		{
+			return obj is IssueAuditLogEntry entry && Equals(entry);
+		}
+
+		/// <summary>
+		/// Compares two IssueAuditLogEntry objects for equality
+		/// </summary>
+		public static bool operator ==(IssueAuditLogEntry left, IssueAuditLogEntry right)
+		{
+			return left.Equals(right);
+		}
+
+		/// <summary>
+		/// Compares two IssueAuditLogEntry objects for inequality
+		/// </summary>
+		public static bool operator !=(IssueAuditLogEntry left, IssueAuditLogEntry right)
+		{
+			return !(left == right);
+		}
+	}
 }
