Horde: Reduce the timeout when retrieving the system capabilities on macOS

diff --git a/Engine/Source/Programs/Horde/HordeAgent/Services/CapabilitiesService.cs b/Engine/Source/Programs/Horde/HordeAgent/Services/CapabilitiesService.cs
--- a/Engine/Source/Programs/Horde/HordeAgent/Services/CapabilitiesService.cs
+++ b/Engine/Source/Programs/Horde/HordeAgent/Services/CapabilitiesService.cs
@@ -230,7 +230,7 @@ namespace HordeAgent.Services
 				using (Process process = new Process())
 				{
 					process.StartInfo.FileName = "system_profiler";
-					process.StartInfo.Arguments = "SPHardwareDataType SPSoftwareDataType SPDisplaysDataType -xml";
+					process.StartInfo.Arguments = "SPHardwareDataType SPSoftwareDataType SPDisplaysDataType -xml -timeout 15";
 					process.StartInfo.CreateNoWindow = true;
 					process.StartInfo.RedirectStandardOutput = true;
 					process.StartInfo.RedirectStandardInput = false;
