Add Device clean up hang to slack notifications

#rnx

diff --git a/Engine/Source/Programs/Horde/Plugins/Build/HordeServer.Build.Tests/Devices/DeviceServiceTests.cs b/Engine/Source/Programs/Horde/Plugins/Build/HordeServer.Build.Tests/Devices/DeviceServiceTests.cs
--- a/Engine/Source/Programs/Horde/Plugins/Build/HordeServer.Build.Tests/Devices/DeviceServiceTests.cs
+++ b/Engine/Source/Programs/Horde/Plugins/Build/HordeServer.Build.Tests/Devices/DeviceServiceTests.cs
@@ -992,6 +992,10 @@ namespace HordeServer.Tests.Devices
 			referencePoint = referencePoint.AddHours(1);
 			AddProblemSpikes(telemetry, devices, pools[0].Id, platforms[0].Id, ref referencePoint, 40, TimeSpan.FromMinutes(3));
 
+			string deviceToClean = devices.Last().Id.ToString();
+			await DeviceController.CleanDeviceAsync(deviceToClean, new CleanDeviceRequest() { Clean = true });
+			devices = (await DeviceService!.GetDevicesAsync()).OrderBy(d => d.PoolId.ToString()).ThenBy(d => d.PlatformId.ToString()).ToList();
+
 			DeviceReportData deviceData = new DeviceReportData()
 			{
 				Pools = pools,
@@ -1005,19 +1009,30 @@ namespace HordeServer.Tests.Devices
 				ReportMinutes = 3 * 60,
 				DeviceProblemCoolDownMinutes = 5,
 				DeviceSaturationSpikeThresholdMinutes = 1,
+				DeviceCleaningThresholdHours = 0,
 				ServerUrl = new Uri("https://testserver")
 			};
 
 			DeviceIssueReport issueReport = DeviceReportService.CreateDeviceReport(deviceData, settings);
 
-			Assert.AreEqual(issueReport.PoolReports.Count, 2);
-			Assert.AreEqual(issueReport.PoolReports[0].Metrics.Count, 3);
-			Assert.AreEqual(issueReport.PoolReports[0].Metrics[0].AverageLoadPercentage, 6);
-			Assert.AreEqual(issueReport.PoolReports[0].Metrics[0].SaturationSpikes, 3);
-			Assert.AreEqual(issueReport.PoolReports[0].Metrics[0].SpikeDurationAverage.ToString(@"hh\:mm"), "00:02");
-			Assert.AreEqual(issueReport.PoolReports[0].Metrics[0].SpikeDurationPercentage, 4);
-			Assert.AreEqual(issueReport.PoolReports[0].Metrics[0].MaxConcurrentProblemsPercentage, 37);
-			Assert.AreEqual(issueReport.PoolReports[0].Metrics[0].MaxConcurrentProblems, 3);
+			Assert.AreEqual(2, issueReport.PlatformReports.Count);
+			Assert.AreEqual(3, issueReport.PlatformReports[0].DeviceReports.Count);
+			Assert.AreNotEqual(null, issueReport.PlatformReports[0].DeviceReports[0].DevicePoolURL);
+			Assert.AreNotEqual(null, issueReport.PlatformReports[0].DeviceReports[1].DevicePoolURL);
+			Assert.AreNotEqual(null, issueReport.PlatformReports[0].DeviceReports[2].DevicePoolURL);
+			Assert.AreEqual(1, issueReport.PlatformReports[1].DeviceReports.Count);
+			Assert.AreNotEqual(null, issueReport.PlatformReports[1].DeviceReports[0].DevicePoolURL);
+			Assert.AreNotEqual(null, issueReport.PlatformReports[1].DeviceReports[0].CleaningTime);
+			Assert.AreEqual(deviceToClean, issueReport.PlatformReports[1].DeviceReports[0].DeviceId);
+
+			Assert.AreEqual(2, issueReport.PoolReports.Count);
+			Assert.AreEqual(3, issueReport.PoolReports[0].Metrics.Count);
+			Assert.AreEqual(6, issueReport.PoolReports[0].Metrics[0].AverageLoadPercentage);
+			Assert.AreEqual(3, issueReport.PoolReports[0].Metrics[0].SaturationSpikes);
+			Assert.AreEqual("00:02", issueReport.PoolReports[0].Metrics[0].SpikeDurationAverage.ToString(@"hh\:mm"));
+			Assert.AreEqual(4, issueReport.PoolReports[0].Metrics[0].SpikeDurationPercentage);
+			Assert.AreEqual(37, issueReport.PoolReports[0].Metrics[0].MaxConcurrentProblemsPercentage);
+			Assert.AreEqual(3, issueReport.PoolReports[0].Metrics[0].MaxConcurrentProblems);
 		}
 
 		[TestMethod]
diff --git a/Engine/Source/Programs/Horde/Plugins/Build/HordeServer.Build/Devices/DeviceReportService.cs b/Engine/Source/Programs/Horde/Plugins/Build/HordeServer.Build/Devices/DeviceReportService.cs
--- a/Engine/Source/Programs/Horde/Plugins/Build/HordeServer.Build/Devices/DeviceReportService.cs
+++ b/Engine/Source/Programs/Horde/Plugins/Build/HordeServer.Build/Devices/DeviceReportService.cs
@@ -37,6 +37,11 @@ namespace HordeServer.Devices
 		/// </summary>
 		public int ProblemPercent { get; set; } = 0;
 
+		/// <summary>
+		/// Time elapse in cleaning
+		/// </summary>
+		public TimeSpan? CleaningTime { get; set; }
+
 		/// <summary>
 		/// The last problem Description
 		/// </summary>
@@ -228,6 +233,7 @@ namespace HordeServer.Devices
 		public int ReportMinutes { get; init; }
 		public int DeviceProblemCoolDownMinutes { get; init; }
 		public int DeviceSaturationSpikeThresholdMinutes { get; init; }
+		public int DeviceCleaningThresholdHours { get; init; }
 		public Uri ServerUrl { get; init; }
 	}
 
@@ -257,6 +263,7 @@ namespace HordeServer.Devices
 
 		readonly int _deviceProblemCoolDownMinutes = 10;
 		readonly int _deviceSaturationSpikeThresholdMinutes = 10;
+		readonly int _deviceCleaningThresholdHours = 6;
 
 		/// <summary>
 		/// Constructor
@@ -322,6 +329,7 @@ namespace HordeServer.Devices
 				ReportMinutes = _reportIntervalMinutes,
 				DeviceProblemCoolDownMinutes = _deviceProblemCoolDownMinutes,
 				DeviceSaturationSpikeThresholdMinutes = _deviceSaturationSpikeThresholdMinutes,
+				DeviceCleaningThresholdHours = _deviceCleaningThresholdHours,
 				ServerUrl = _serverInfo.DashboardUrl
 			};
 
@@ -429,8 +437,35 @@ namespace HordeServer.Devices
 					}
 				}
 
+				Func<DevicePlatformReport> getPlatformReport = () =>
+				{
+					DevicePlatformReport? platformReport = issueReport.PlatformReports.Find(p => p.PlatformId == device.PlatformId.ToString());
+					if (platformReport == null)
+					{
+						platformReport = new DevicePlatformReport(platform.Id.ToString(), platform.Name);
+						issueReport.PlatformReports.Add(platformReport);
+					}
+					return platformReport;
+				};
+
+				Func<string> buildDevicePoolUrl = () => new Uri($"{settings.ServerUrl}devices?pivotKey=pivot-key-automation&pool={pool.Id}&platform={platform.Id}&filter={device.Name}").ToString();
+
+				IEnumerable<IDeviceTelemetry> problems;
+				// Check for clean up hang
+				TimeSpan? cleaningTime = device.CleanStartTimeUtc != null && device.CleanFinishTimeUtc == null ? (DateTime.UtcNow - device.CleanStartTimeUtc) : null;
+				if (cleaningTime != null && cleaningTime.Value.TotalHours >= settings.DeviceCleaningThresholdHours)
+				{
+					DeviceReport report = new DeviceReport(platform.Id.ToString(), platform.Name, device.Id.ToString(), device.Name, device.Address ?? "Unknown Address", pool.Id.ToString(), pool.Name, telemetry);
+					report.CleaningTime = cleaningTime;
+					report.DevicePoolURL = buildDevicePoolUrl();
+
+					getPlatformReport().DeviceReports.Add(report);
+					
+					return;
+				}
+
 				// Report device issues, though only if not since disabled or put into maintenance
-				IEnumerable<IDeviceTelemetry> problems = telemetry.Where(t => t.ProblemTimeUtc != null).OrderByDescending(t => t.ProblemTimeUtc!);
+				problems = telemetry.Where(t => t.ProblemTimeUtc != null).OrderByDescending(t => t.ProblemTimeUtc!);
 
 				if (!problems.Any())
 				{
@@ -462,16 +497,9 @@ namespace HordeServer.Devices
 					deviceReport.LastProblemURL = new Uri($"{settings.ServerUrl}job/{lastProblem.JobId}?step={lastProblem.StepId}").ToString();
 				}
 
-				deviceReport.DevicePoolURL = new Uri($"{settings.ServerUrl}devices?pivotKey=pivot-key-automation&pool={pool.Id}&platform={platform.Id}&filter={device.Name}").ToString();
-
-				DevicePlatformReport? platformReport = issueReport.PlatformReports.Find(p => p.PlatformId == device.PlatformId.ToString());
-				if (platformReport == null)
-				{
-					platformReport = new DevicePlatformReport(platform.Id.ToString(), platform.Name);
-					issueReport.PlatformReports.Add(platformReport);
-				}
+				deviceReport.DevicePoolURL = buildDevicePoolUrl();
 
-				platformReport.DeviceReports.Add(deviceReport);
+				getPlatformReport().DeviceReports.Add(deviceReport);
 			});
 
 			// Calculate usage and spikes
diff --git a/Engine/Source/Programs/Horde/Plugins/Build/HordeServer.Build/Notifications/Sinks/SlackNotificationSink.cs b/Engine/Source/Programs/Horde/Plugins/Build/HordeServer.Build/Notifications/Sinks/SlackNotificationSink.cs
--- a/Engine/Source/Programs/Horde/Plugins/Build/HordeServer.Build/Notifications/Sinks/SlackNotificationSink.cs
+++ b/Engine/Source/Programs/Horde/Plugins/Build/HordeServer.Build/Notifications/Sinks/SlackNotificationSink.cs
@@ -70,6 +70,7 @@ namespace HordeServer.Notifications.Sinks
 		const string YellowgBadge = ":large_yellow_square:";
 		const string OrangeBadge = ":large_orange_square:";
 		const string RedBadge = ":large_red_square:";
+		const string BlueBadge = ":large_blue_square:";
 		const string BlackBadge = ":black_square:";
 
 		class SocketResponse
@@ -2568,13 +2569,26 @@ namespace HordeServer.Notifications.Sinks
 			{
 				foreach (DevicePlatformReport platform in report.PlatformReports)
 				{
+					if (platform.DeviceReports.Count == 0)
+					{
+						continue;
+					}
+
 					SlackMessage message = new SlackMessage();
 					message.Blocks.Add(new SectionBlock($"*{platform.PlatformName} | Device Problems Since Last Update*"));
 
 					foreach (DeviceReport device in platform.DeviceReports.OrderByDescending(r => r.ProblemDelta))
 					{
-						StringBuilder builder = new StringBuilder($"{RedBadge} <{device.DevicePoolURL}|*{device.DeviceName}*> | {device.DeviceAddress} | {device.PoolName}");
-						builder.Append($"\n> {device.ProblemDelta} Problem(s) | {device.ProblemPercent}% Rate");
+						string badge = device.CleaningTime != null ? BlueBadge : RedBadge;
+						StringBuilder builder = new StringBuilder($"{badge} <{device.DevicePoolURL}|*{device.DeviceName}*> | {device.DeviceAddress} | {device.PoolName}");
+						if (device.CleaningTime != null)
+						{
+							builder.Append($"\n> Cleaning for {(int) (device.CleaningTime.Value.TotalHours)} hours");
+						}
+						else
+						{
+							builder.Append($"\n> {device.ProblemDelta} Problem(s) | {device.ProblemPercent}% Rate");
+						}
 
 						if (!String.IsNullOrEmpty(device.LastProblemURL))
 						{
