Fetching zip artifacts using POST request

#rnx

diff --git a/Engine/Source/Programs/Horde/HordeDashboard/plugins/test-automation/testAutomationCommon.tsx b/Engine/Source/Programs/Horde/HordeDashboard/plugins/test-automation/testAutomationCommon.tsx
--- a/Engine/Source/Programs/Horde/HordeDashboard/plugins/test-automation/testAutomationCommon.tsx
+++ b/Engine/Source/Programs/Horde/HordeDashboard/plugins/test-automation/testAutomationCommon.tsx
@@ -694,7 +694,7 @@ export const PhaseArtifactsModal: React.FC<{ phase: TestPhaseStatus, artifactPat
                 let color = last ? (dashboard.darktheme ? "#FFFFFF" : "#605E5C") : undefined;
                 const font = last ? undefined : "Horde Open Sans Light";
                 const sep = last ? undefined : "/"
-                return <Text key={`p-${item.name}-${index}`} styles={{ root: { fontFamily: font } }} style={{ color: color }}>{t}{sep}</Text>
+                return <Text key={`p-${index}`} styles={{ root: { fontFamily: font } }} style={{ color: color }}>{t}{sep}</Text>
             });
             if (start_index !== 0) pathElements.splice(0, 0, <Text key={`p-${item.name}`} >... /</Text>);
 
diff --git a/Engine/Source/Programs/Horde/HordeDashboard/plugins/test-automation/testData.ts b/Engine/Source/Programs/Horde/HordeDashboard/plugins/test-automation/testData.ts
--- a/Engine/Source/Programs/Horde/HordeDashboard/plugins/test-automation/testData.ts
+++ b/Engine/Source/Programs/Horde/HordeDashboard/plugins/test-automation/testData.ts
@@ -1559,7 +1559,9 @@ export class ArtifactFactory {
 
     async downloadZip(referencePaths: string[]) : Promise<void> {
         await this.getJobStepSavedArtifacts();
-        if (!!this.artifactId) backend.downloadArtifactZipV2(this.artifactId, {filter: referencePaths})
+        if (!!this.artifactId) {
+            backend.downloadArtifactZipV2(this.artifactId, {filter: referencePaths})
+        }
     }
 
     private jobId: string;
diff --git a/Engine/Source/Programs/Horde/HordeDashboard/src/backend/Backend.ts b/Engine/Source/Programs/Horde/HordeDashboard/src/backend/Backend.ts
--- a/Engine/Source/Programs/Horde/HordeDashboard/src/backend/Backend.ts
+++ b/Engine/Source/Programs/Horde/HordeDashboard/src/backend/Backend.ts
@@ -571,18 +571,35 @@ export class Backend {
         window.location.assign(url);
     }
 
-    downloadArtifactZipV2(artifactId: string, request?: CreateZipRequest) {        
+    downloadArtifactZipV2(artifactId: string, request?: CreateZipRequest) {
+        const url = `/api/v2/artifacts/${artifactId}/zip`;
 
-        const filter = request?.filter?.map(f => `filter=${encodeURIComponent(f)}`).join("&")
+        const filter = request?.filter?.map(f => `filter=${encodeURIComponent(f)}`).join("&") ?? '';
 
-        let url = `/api/v2/artifacts/${artifactId}/zip`;
-        
-        if (filter?.length) {
-            url += "?" + filter;
+        if (filter.length < 2000) {
+            // 2'048 characters is the usual limit enforced by browsers for the whole url
+            window.location.assign(`${url}?${filter}`);
+            return;
         }
 
-        window.location.assign(url);
-
+        // the disadvantage of the POST method using fetch is the lack of progression indicator for large download size
+        // this method wait for the whole data to be downloaded before passing the data blob to the new window
+        return new Promise<object>(() => {
+            this._fetch.post(url, request, {responseBlob: true}).then((response) => {
+                // Create a temp link and activate it to save the incoming data
+                const blobURL = URL.createObjectURL(response.data);
+                const link = document.createElement('a');
+                link.href = blobURL;
+                document.body.appendChild(link);
+                link.click();
+                URL.revokeObjectURL(blobURL);
+                document.body.removeChild(link);
+            }).catch(reason => {
+                if (reason !== "Not Found") {
+                    console.error(reason);
+                }
+            });
+        });
     }
 
     getLogData(logId: string): Promise<LogData> {
