Keep earliest run date for CL date

#rnx

diff --git a/Engine/Source/Programs/Horde/HordeDashboard/plugins/test-automation/testData.ts b/Engine/Source/Programs/Horde/HordeDashboard/plugins/test-automation/testData.ts
--- a/Engine/Source/Programs/Horde/HordeDashboard/plugins/test-automation/testData.ts
+++ b/Engine/Source/Programs/Horde/HordeDashboard/plugins/test-automation/testData.ts
@@ -1219,8 +1219,8 @@ export class TestDataHandler {
             /// Update uniform commit dates
             const date = commitIdDates.get(session.commitId);
             const rdate = session.start;
-            // Use the latest date for summary, otherwise gets confusing when a CL spans days
-            if (!date || date.getTime() < rdate.getTime()) {
+            // Use the earliest date for summary, otherwise gets confusing when a CL spans days
+            if (!date || date.getTime() > rdate.getTime()) {
                 commitIdDates.set(session.commitId, rdate);
             }
         }
