Horde: Filter out anonymous pools in pools view
#submitool bypass

diff --git a/Engine/Source/Programs/Horde/HordeDashboard/src/components/PoolsView.tsx b/Engine/Source/Programs/Horde/HordeDashboard/src/components/PoolsView.tsx
--- a/Engine/Source/Programs/Horde/HordeDashboard/src/components/PoolsView.tsx
+++ b/Engine/Source/Programs/Horde/HordeDashboard/src/components/PoolsView.tsx
@@ -340,6 +340,12 @@ const PoolList: React.FC = observer(() => {
 
    const filter = handler.filter.toLowerCase();
    const filtered = handler.pools.filter(p => {
+
+      // filter out anonymous pools from display
+      if (!p.name) {
+         return false;
+      }
+
       if (!filter) {
          return true;
       }
