Horde: Add max retries to agent delete

diff --git a/Engine/Source/Programs/Horde/Plugins/Compute/HordeServer.Compute/Agents/AgentService.cs b/Engine/Source/Programs/Horde/Plugins/Compute/HordeServer.Compute/Agents/AgentService.cs
--- a/Engine/Source/Programs/Horde/Plugins/Compute/HordeServer.Compute/Agents/AgentService.cs
+++ b/Engine/Source/Programs/Horde/Plugins/Compute/HordeServer.Compute/Agents/AgentService.cs
@@ -281,7 +281,9 @@ namespace HordeServer.Agents
 			{
 				return;
 			}
-
+			
+			int maxRetries = 10;
+			int retryCount = 0;
 			while (agent is { Deleted: false })
 			{
 				IAgent? newAgent = await agent.TryDeleteAsync(cancellationToken);
@@ -289,6 +291,14 @@ namespace HordeServer.Agents
 				{
 					break;
 				}
+				retryCount++;
+				if (retryCount >= maxRetries)
+				{
+					_logger.LogError("Unable to delete agent {AgentId} after {MaxRetries} retry attempts", agent.Id, maxRetries);
+					throw new InvalidOperationException($"Unable to delete agent {agent.Id} after {maxRetries} retry attempts");
+				}
+				
+				await Task.Delay(100, cancellationToken);
 				agent = await GetAgentAsync(agent.Id, cancellationToken);
 			}
 		}
