[Horde][BuildHealth][Dashboard] Multiple Stream selections, date headers, loading spinners on refreshes, URL query encoding fixes to include hierarchy, scroll-locking optimization, and various UI cleanup.

diff --git a/Engine/Source/Programs/Horde/HordeDashboard/src/components/buildhealth/BuildHealthDataHandler.ts b/Engine/Source/Programs/Horde/HordeDashboard/src/components/buildhealth/BuildHealthDataHandler.ts
--- a/Engine/Source/Programs/Horde/HordeDashboard/src/components/buildhealth/BuildHealthDataHandler.ts
+++ b/Engine/Source/Programs/Horde/HordeDashboard/src/components/buildhealth/BuildHealthDataHandler.ts
@@ -207,7 +207,7 @@ export class BuildHealthDataHandler {
                         for (const step of batch.steps) {
                             const stepKey = step.name;
                             if (!stepMap.has(stepKey)) {
-                                stepMap.set(stepKey, { ...step, templateId: decodeTemplateKey(key).templateId });
+                                stepMap.set(stepKey, { ...step, streamId: templateData.streamId, templateId: decodeTemplateKey(key).templateId });
                             }
                         }
                     }
diff --git a/Engine/Source/Programs/Horde/HordeDashboard/src/components/buildhealth/BuildHealthDataTypes.ts b/Engine/Source/Programs/Horde/HordeDashboard/src/components/buildhealth/BuildHealthDataTypes.ts
--- a/Engine/Source/Programs/Horde/HordeDashboard/src/components/buildhealth/BuildHealthDataTypes.ts
+++ b/Engine/Source/Programs/Horde/HordeDashboard/src/components/buildhealth/BuildHealthDataTypes.ts
@@ -23,4 +23,5 @@ export type TemplateRefData = GetTemplateRefResponse & {
  */
 export type StepRefData = GetStepResponse & {
     templateId?: string;
+    streamId?: string;
 }
\ No newline at end of file
diff --git a/Engine/Source/Programs/Horde/HordeDashboard/src/components/buildhealth/BuildHealthOptions.ts b/Engine/Source/Programs/Horde/HordeDashboard/src/components/buildhealth/BuildHealthOptions.ts
--- a/Engine/Source/Programs/Horde/HordeDashboard/src/components/buildhealth/BuildHealthOptions.ts
+++ b/Engine/Source/Programs/Horde/HordeDashboard/src/components/buildhealth/BuildHealthOptions.ts
@@ -5,6 +5,8 @@ import { PROJECTS_URL_PARAM_NAME, STEPS_URL_PARAM_NAME, STREAMS_URL_PARAM_NAME,
 import { BuildHealthDataHandler } from "./BuildHealthDataHandler";
 import { decodeStepKey, decodeTemplateKey, encodeStepKey, encodeTemplateKey } from "./BuildHealthUtilities";
 import { Location } from 'react-router-dom';
+import LZString from "lz-string"
+import { useQuery } from "horde/base/utilities/hooks";
 
 // #region -- Options Utilities & Data Types --
 
@@ -30,7 +32,7 @@ export type TimeSpan = {
 
 /**
 * List of currently supported time spans.
-*/ 
+*/
 export const JobHistoryTimeSpans: TimeSpan[] = [
     {
         text: "Past 15 Minutes", key: "time_15_minutes", minutes: 15
@@ -64,14 +66,26 @@ export const DEFAULT_JOB_HISTORY_TIME_SPAN: TimeSpan = JobHistoryTimeSpans[1];
  */
 export function parseBuildHealthQueryParams(location: Location): BuildHealthQueryParams {
     const searchParams = new URLSearchParams(location.search);
-    const getList = (key: string) => searchParams.get(key)?.split(",").filter(Boolean) ?? [];
+    const getList = (key: string) => {
+        const raw = searchParams.get(key);
+        if (!raw) return [];
+
+        const decompressed = LZString.decompressFromEncodedURIComponent(raw);
+        if (!decompressed) return [];
+
+        return decompressed.split(",").filter(Boolean);
+    };
 
     const projectIds = getList(PROJECTS_URL_PARAM_NAME);
     const streamIds = getList(STREAMS_URL_PARAM_NAME);
     const templateIds = getList(TEMPLATES_URL_PARAM_NAME);
     const stepNames = getList(STEPS_URL_PARAM_NAME);
 
-    let timeSpanKey: string = searchParams.get(TIME_SPAN_URL_PARAM) ?? "";
+    let timeSpanKey: string = (() => {
+        const raw = searchParams.get(TIME_SPAN_URL_PARAM);
+        return raw ? LZString.decompressFromEncodedURIComponent(raw) ?? "" : "";
+    })();
+    
     if (!JobHistoryTimeSpans.some(ts => ts.key === timeSpanKey)) {
         timeSpanKey = DEFAULT_JOB_HISTORY_TIME_SPAN.key;
     }
@@ -131,8 +145,10 @@ export async function loadBuildHealthOptionsFromParams(
 
     runInAction(() => {
         options.clearTemplates();
-        params.templateIds.forEach(tid => {
-            const template = handler.templatesData.find(t => t.id === tid);
+        const decodedTemplates = params.templateIds.map(templateKey => decodeTemplateKey(templateKey));
+
+        decodedTemplates.forEach(tid => {
+            const template = handler.templatesData.find(t => t.id === tid.templateId && t.streamId == tid.streamId);
             if (template) {
                 options.toggleTemplate(encodeTemplateKey(template), template.name, true);
             }
@@ -146,8 +162,9 @@ export async function loadBuildHealthOptionsFromParams(
 
     runInAction(() => {
         options.clearSteps();
-        params.stepNames.forEach(name => {
-            const step = handler.stepData.find(s => s.name.toLowerCase() === name.toLowerCase());
+        const decodedStepNames = params.stepNames.map(stepNameKey => decodeStepKey(stepNameKey));
+        decodedStepNames.forEach(stepData => {
+            const step = handler.stepData.find(s => s.name.toLowerCase() === stepData.stepName.toLowerCase() && s.streamId == stepData.streamId && s.templateId == stepData.templateId);
             if (step) {
                 options.toggleStep(encodeStepKey(step), step.name, true);
             }
@@ -393,8 +410,8 @@ export class BuildHealthOptionsController {
     synchronizeDerivedKeys() {
         if (this.lastSynchronize < this.optionsChangeVersion) {
             this.state.stepOutcomeEnabledStreamKeys = Object.keys(this.state.enabledStreams);
-            this.state.stepOutcomeEnabledJobKeys = Object.keys(this.state.enabledTemplates).map(key => decodeTemplateKey(key).templateId);
-            this.state.stepOutcomeEnabledStepKeys = Object.keys(this.state.enabledSteps).map(key => decodeStepKey(key).stepName);
+            this.state.stepOutcomeEnabledJobKeys = Object.keys(this.state.enabledTemplates);
+            this.state.stepOutcomeEnabledStepKeys = Object.keys(this.state.enabledSteps);
             this.state.startDate = new Date(new Date().valueOf() - (this.state.jobHistoryTimeSpan.minutes * 60000));
             this.state.endDate = new Date();
 
@@ -417,11 +434,11 @@ export class BuildHealthOptionsController {
 
         if (!projects.length) return params;
 
-        params.set(PROJECTS_URL_PARAM_NAME, projects);
-        params.set(STREAMS_URL_PARAM_NAME, streams);
-        params.set(TEMPLATES_URL_PARAM_NAME, jobs);
-        params.set(STEPS_URL_PARAM_NAME, steps);
-        params.set(TIME_SPAN_URL_PARAM, lastTimeRange);
+        params.set(PROJECTS_URL_PARAM_NAME, LZString.compressToEncodedURIComponent(projects));
+        params.set(STREAMS_URL_PARAM_NAME, LZString.compressToEncodedURIComponent(streams));
+        params.set(TEMPLATES_URL_PARAM_NAME, LZString.compressToEncodedURIComponent(jobs));
+        params.set(STEPS_URL_PARAM_NAME, LZString.compressToEncodedURIComponent(steps));
+        params.set(TIME_SPAN_URL_PARAM, LZString.compressToEncodedURIComponent(lastTimeRange));
 
         return params;
     }
diff --git a/Engine/Source/Programs/Horde/HordeDashboard/src/components/buildhealth/BuildHealthUtilities.ts b/Engine/Source/Programs/Horde/HordeDashboard/src/components/buildhealth/BuildHealthUtilities.ts
--- a/Engine/Source/Programs/Horde/HordeDashboard/src/components/buildhealth/BuildHealthUtilities.ts
+++ b/Engine/Source/Programs/Horde/HordeDashboard/src/components/buildhealth/BuildHealthUtilities.ts
@@ -27,7 +27,7 @@ export function decodeTemplateKey(templateKey: string): { streamId: string; temp
  * @returns The key representing the step.
  */
 export function encodeStepKey(step: StepRefData) {
-    return `${step.templateId}${KEY_SEPARATOR}${step.name}`;
+    return `${step.streamId}${KEY_SEPARATOR}${step.templateId}${KEY_SEPARATOR}${step.name}`;
 }
 
 /**
@@ -35,7 +35,7 @@ export function encodeStepKey(step: StepRefData) {
  * @param templateKey The step key to decode.
  * @returns The constituent parts.
  */
-export function decodeStepKey(stepKey: string): { templateId: string, stepName: string } {
-    const [templateId, stepName] = stepKey.split(KEY_SEPARATOR);
-    return { templateId, stepName };
+export function decodeStepKey(stepKey: string): { streamId: string, templateId: string, stepName: string } {
+    const [streamId, templateId, stepName] = stepKey.split(KEY_SEPARATOR);
+    return { streamId, templateId, stepName };
 }
\ No newline at end of file
diff --git a/Engine/Source/Programs/Horde/HordeDashboard/src/components/buildhealth/BuildhealthView.tsx b/Engine/Source/Programs/Horde/HordeDashboard/src/components/buildhealth/BuildhealthView.tsx
--- a/Engine/Source/Programs/Horde/HordeDashboard/src/components/buildhealth/BuildhealthView.tsx
+++ b/Engine/Source/Programs/Horde/HordeDashboard/src/components/buildhealth/BuildhealthView.tsx
@@ -353,7 +353,7 @@ const StepDropdown: React.FC<{ handler: BuildHealthDataHandler, buildHealthOptio
         let listParam: BuildHealthOptionData = {
             id: encodeStepKey(step),
             text: step.name,
-            group: step.templateId
+            group: `${step.streamId}${KEY_SEPARATOR}${step.templateId}`
         };
         return listParam;
     }).sort((a, b) => a.text.localeCompare(b.text));
@@ -396,6 +396,7 @@ const BuildHealthUrlSync: React.FC<{ buildHealthController: BuildHealthOptionsCo
 
         initializeFromUrl();
     }, []);
+    
     useEffect(() => {
         const desiredSearch = buildHealthController.toNavigationQuery().toString();
         if (desiredSearch && location.search !== `?${desiredSearch}`) {
@@ -434,7 +435,7 @@ const BuildHealthSidebar: React.FC<{ handler: BuildHealthDataHandler, buildHealt
                         <Label>Step Outcome Filter</Label>
                     </Stack>
                     <ProjectDropdownSingle handler={handler} buildHealthOptions={buildHealthOptions} />
-                    <StreamDropdownSingle handler={handler} buildHealthOptions={buildHealthOptions} />
+                    <StreamDropdown handler={handler} buildHealthOptions={buildHealthOptions} />
                     <TemplateDropdown handler={handler} buildHealthOptions={buildHealthOptions} />
                     <StepDropdown handler={handler} buildHealthOptions={buildHealthOptions} />
                     <Stack>
diff --git a/Engine/Source/Programs/Horde/HordeDashboard/src/components/buildhealth/stepoutcome/StepOutcome.tsx b/Engine/Source/Programs/Horde/HordeDashboard/src/components/buildhealth/stepoutcome/StepOutcome.tsx
--- a/Engine/Source/Programs/Horde/HordeDashboard/src/components/buildhealth/stepoutcome/StepOutcome.tsx
+++ b/Engine/Source/Programs/Horde/HordeDashboard/src/components/buildhealth/stepoutcome/StepOutcome.tsx
@@ -1,15 +1,41 @@
 // Copyright Epic Games, Inc. All Rights Reserved.
 
-import { Stack, DetailsList, IColumn, mergeStyleSets, SelectionMode, TooltipHost, DetailsListLayoutMode, ConstrainMode, Dialog, DialogType } from "@fluentui/react";
+import { Stack, DetailsList, IColumn, mergeStyleSets, SelectionMode, TooltipHost, DetailsListLayoutMode, ConstrainMode, Dialog, DialogType, Label, IGroup, Spinner, SpinnerSize } from "@fluentui/react";
 import { observer } from "mobx-react-lite";
 import { getHordeStyling } from 'horde/styles/Styles';
 import React, { useEffect, useRef, useState } from 'react';
 import { Text } from '@fluentui/react/lib/Text';
-import { StepOutcomeTable, StepOutcomeTableEntry } from "./StepOutcomeDataTypes";
+import { encodeStepNameFromStepNameHeader, StepNameHeader, StepOutcomeTable, StepOutcomeTableEntry } from "./StepOutcomeDataTypes";
 import { StepOutcomeFilters, StepOutcomeDataHandler } from "./StepOutcomeDataHandler";
 
 // #region -- UI Helper Methods --
 
+const ROW_HEADER_NAME: string = "stepName" as const;
+
+/**
+ * Produces Fluent UI compatible groups from the provided Step Outcome Table. Groups are currently derived based on Streams.
+ * @param stepOutcomeTable The step outcome table to use in producing Fluent UI groups.
+ * @returns The Fluent UI compatible groups.
+ */
+function getFluentGroups(stepOutcomeTable: StepOutcomeTable): any[] {
+    let returnGroups: IGroup[] = [];
+    let startIdx = 0;
+
+    for (const [key, value] of stepOutcomeTable.tableStreamMetadata) {
+        let streamGroup: IGroup = {
+            key: key,
+            name: key,
+            count: value,
+            startIndex: startIdx,
+            isCollapsed: false
+        };
+        startIdx += value;
+        returnGroups.push(streamGroup);
+    }
+    return returnGroups;
+}
+
+
 /**
  * Produces Fluent UI compatible columns from the provided Step Outcome Table.
  * @param stepOutcomeTable The step outcome table to use in producing Fluent UI columns.
@@ -18,15 +44,30 @@ import { StepOutcomeFilters, StepOutcomeDataHandler } from "./StepOutcomeDataHan
 function getFluentColumns(stepOutcomeTable: StepOutcomeTable): IColumn[] {
     const columns: IColumn[] = [
         {
-            key: 'stepName',
+            key: ROW_HEADER_NAME,
             name: '',
-            fieldName: 'stepName',
+            fieldName: ROW_HEADER_NAME,
             minWidth: 300,
             onRender: (item) => {
                 return (
-                    <TooltipHost content={`Step: ${item.stepName}`}>
+                    <TooltipHost content={
+                        <div>
+                            <Stack>
+                                <Label>Step:</Label>
+                                <Text style={{ paddingLeft: 8 }}>
+                                    {item.stepName.stepName}
+                                </Text>
+                            </Stack>
+                            <Stack>
+                                <Label>Stream:</Label>
+                                <Text style={{ paddingLeft: 8 }}>
+                                    {item.stepName.streamId}
+                                </Text>
+                            </Stack>
+                        </div>
+                    }>
                         <div style={{ whiteSpace: 'normal', wordWrap: 'break-word', minHeight: 20, maxHeight: 20 }}>
-                            {item.stepName}
+                            {` ${item.stepName.stepName}`}
                         </div>
                     </TooltipHost>
                 );
@@ -34,15 +75,74 @@ function getFluentColumns(stepOutcomeTable: StepOutcomeTable): IColumn[] {
         }
     ];
 
+    // We want to inject dates into the DetailsList column headers in order to assist with data consumption.
+    // Keep accumulators for the corresponding date strings.
+    let prevDay: string | null = null;
+    let prevChangelist: string | null = null;
+
     for (let i: number = 0; i < stepOutcomeTable.changeColHeaders.length; ++i) {
-        let change = stepOutcomeTable.changeColHeaders[i];
+
+        // If the date associated with the changelist header has changed since the last, we are on a new day. Set that as the colimn.
+        let changeHeader = stepOutcomeTable.changeColHeaders[i];
+        if (changeHeader.date !== undefined) {
+            const headerDate = new Date(changeHeader.date!);
+            const headerChangelist = changeHeader.change.toString();
+
+            // Extract the day portion as YYYY-MM-DD
+            const dayStr = `${headerDate.getUTCFullYear()}-${headerDate.getUTCMonth() + 1}-${headerDate.getUTCDate()}`;
+
+            if (prevDay !== dayStr) {
+                prevDay = dayStr;
+                prevChangelist = headerChangelist;
+                columns.push({
+                    key: `${prevDay}_${changeHeader.change}`,
+                    name: changeHeader.change.toString(),
+                    minWidth: 10,
+
+                    onRenderHeader: (props, defaultRender) => {
+                        if (!props || !defaultRender) return null;
+
+                        return (
+                            <div className={headerStylesHighlight.headerCell}>
+                                <TooltipHost>
+                                    {dayStr}
+                                </TooltipHost>
+                            </div>
+                        );
+                    }
+                });
+            }
+        }
+
+        // Push all columns that relate to changelists.
         columns.push({
-            key: change.toString(),
-            name: change.toString(),
-            fieldName: change.toString(),
+            key: changeHeader.change.toString(),
+            name: changeHeader.change.toString(),
+            fieldName: changeHeader.change.toString(),
             minWidth: 10,
+            onRenderHeader: (props, defaultRender) => {
+                if (!props) {
+                    return null;
+                }
+                let innerChange = stepOutcomeTable.changeColHeaders[props?.columnIndex];
+                return props && defaultRender ? (
+                    <TooltipHost content={
+                        <Stack>
+                            <Label>Change:  </Label>
+                            <Text style={{ paddingLeft: 8 }}>{changeHeader.change}</Text>
+                            <Label>Date (UTC): </Label>
+                            <Text style={{ paddingLeft: 8 }} >{new Date(changeHeader.date!).toISOString()}</Text>
+                            <Label>Date (Local): </Label>
+                            <Text style={{ paddingLeft: 8 }} >{new Date(changeHeader.date!).toString()}</Text>
+                        </Stack>
+                    }
+                    >
+                        {defaultRender(props)}
+                    </TooltipHost>
+                ) : null
+            },
             onRender: (item: any) => {
-                const entry = item[change.toString()] as StepOutcomeTableEntry | null;
+                const entry = item[changeHeader.change.toString()] as StepOutcomeTableEntry | null;
                 let bgColor = 'transparent';
 
                 if (entry) {
@@ -51,18 +151,32 @@ function getFluentColumns(stepOutcomeTable: StepOutcomeTable): IColumn[] {
                 }
 
                 return (
-                    <div
-                        title={
-                            `Change: ${change}\nJob: ${entry?.jobName}\nStep: ${entry?.stepResponse.name}`}
-                        className={stepOutcomeTableClasses.fillLink}
-                        style={{
-                            backgroundColor: bgColor,
-                            minHeight: 20, maxHeight: 20,
-                            minWidth: 20, maxWidth: 20,
-                        }}
-                    >
-                        &nbsp;
-                    </div>
+                    <TooltipHost content={
+                        <Stack>
+                            <Label>Change:</Label>
+                            <Text style={{ paddingLeft: 8 }}>{changeHeader.change}</Text>
+                            {entry && (
+                                <>
+                                    <Label>Stream:</Label>
+                                    <Text style={{ paddingLeft: 8 }}>{entry.streamId}</Text>
+                                    <Label>Job:</Label>
+                                    <Text style={{ paddingLeft: 8 }}>{entry.jobName}</Text>
+                                    <Label>Step:</Label>
+                                    <Text style={{ paddingLeft: 8 }}>{entry.stepResponse.name}</Text>
+                                </>)}
+                        </Stack>
+                    }>
+                        <div
+                            style={{
+                                backgroundColor: bgColor,
+                                minHeight: 20, maxHeight: 20,
+                                minWidth: 20, maxWidth: 20,
+                            }}
+                        >
+                            &nbsp;
+                        </div>
+                    </ TooltipHost>
+
                 );
             }
         });
@@ -76,12 +190,12 @@ function getFluentColumns(stepOutcomeTable: StepOutcomeTable): IColumn[] {
  * @param stepOutcomeTable The step outcome table to use in producing Fluent UI rows.
  * @returns The Fluent UI compatible rows.
  */
-function getFluentRows(stepOutcomeTable: StepOutcomeTable): any[] {
+function getFluentRows(stepOutcomeTable: StepOutcomeTable): StepOutcomeFluentRow[] {
     let items: StepOutcomeFluentRow[] = [];
 
     for (let i: number = 0; i < stepOutcomeTable.stepNameRowHeaders.length; ++i) {
-        let rows: StepOutcomeTableEntry[] = [...handler.stepOutcomeTableData.getStepOutputRow(stepOutcomeTable.stepNameRowHeaders[i].toLocaleLowerCase())];
-        let fluentRowEntry: StepOutcomeFluentRow = { 'stepName': stepOutcomeTable.stepNameRowHeaders[i] };
+        let rows: StepOutcomeTableEntry[] = [...handler.StepOutcomeTableData.getStepOutputRow(encodeStepNameFromStepNameHeader(stepOutcomeTable.stepNameRowHeaders[i]))];
+        let fluentRowEntry: StepOutcomeFluentRow = { [ROW_HEADER_NAME]: stepOutcomeTable.stepNameRowHeaders[i] };
         for (let y: number = 0; y < rows.length; ++y) {
             let tableEntry: StepOutcomeTableEntry = rows[y];
             if (tableEntry.change === undefined) {
@@ -100,8 +214,8 @@ function getFluentRows(stepOutcomeTable: StepOutcomeTable): any[] {
 // #region -- View Components --
 
 type StepOutcomeFluentRow = {
-    stepName: string;
-    [change: string]: StepOutcomeTableEntry | string | undefined;
+    [ROW_HEADER_NAME]: StepNameHeader;
+    [change: string]: StepOutcomeTableEntry | StepNameHeader | string | undefined;
 };
 
 type StepOutcomeTableProps = {
@@ -121,28 +235,40 @@ interface StepOutcomeViewProps {
  * Defines the style for column headers.
  */
 
+const headerStylesHighlight = mergeStyleSets({
+    headerCell: {
+        fontWeight: '700 !important',
+        fontSize: 'larger !important'
+    },
+});
+
 const stepOutcomeTableClasses = mergeStyleSets({
     stepOutcomeTableHeaders: {
         selectors: {
+            '.ms-DetailsHeader-cellName': {
+                fontWeight: 200
+            },
             '.ms-DetailsList-headerWrapper': {
                 position: 'sticky'
             },
             '.ms-DetailsHeader-cellTitle': {
-                writingMode: 'vertical-rl',
-                transform: 'rotate(190deg)',
+                writingMode: 'vertical-lr',
+                transform: 'rotate(180deg)',
                 whiteSpace: 'nowrap',
                 textAlign: 'center',
-                padding: '4px 2px',
                 paddingBottom: '10px',
                 width: '36px !important',
                 userSelect: 'text',
             },
             '.ms-DetailsHeader': {
-                minHeight: '80px',
+                minHeight: '100px',
                 lineHeight: 'normal',
+                paddingTop: '0 !important',
             },
             '.ms-DetailsHeader-cell': {
-                verticalAlign: 'bottom'
+                verticalAlign: 'bottom',
+                minHeight: '80px',
+                borderLeft: '1px solid rgba(41, 41, 41, 0.35)', borderRight: '1px solid  rgba(41, 41, 41, 0.35)'
             },
             '.ms-DetailsHeader-cell span': {
                 overflow: 'visible',
@@ -154,13 +280,12 @@ const stepOutcomeTableClasses = mergeStyleSets({
             '.ms-SelectionZone': {
                 display: 'none',
             },
-
         },
     },
     removeCellContentPadding: {
         selectors: {
             '.ms-DetailsRow-cell': {
-                paddingLeft: '0 !important'
+                paddingLeft: '0 !important',
             }
         }
     },
@@ -189,8 +314,14 @@ const stepOutcomeTableClasses = mergeStyleSets({
                 width: '0 !important',
             },
             '.ms-DetailsRow-cell': {
-                paddingLeft: '0 !important'
-            }
+                padding: '10px 4px',
+                // Cleanup work for proper squares
+                // maxHeight: '35px',
+                // width: '35px !important',
+                // maxWdith: '35px',
+                borderLeft: '1px solid  rgba(41, 41, 41, 0.35)', borderRight: '1px solid  rgba(41, 41, 41, 0.35)'
+            },
+
         },
     },
     disableGroupCollapseButton: {
@@ -206,7 +337,6 @@ const stepOutcomeTableClasses = mergeStyleSets({
         selectors: {
             '.ms-DetailsHeader-cellTitle': {
                 writingMode: 'vertical-rl',
-                transform: 'rotate(190deg)',
                 whiteSpace: 'nowrap',
                 textAlign: 'center',
                 padding: '4px 2px',
@@ -218,11 +348,6 @@ const stepOutcomeTableClasses = mergeStyleSets({
             },
         },
     },
-    fillLink: {
-        color: 'inherit',
-        padding: '2px',
-        boxSizing: 'border-box',
-    },
 });
 
 // #endregion -- Styles  --
@@ -410,7 +535,7 @@ function GenerateClassicTable(stepOutcomeTable: StepOutcomeTable): JSX.Element {
                                         padding: '8px 0',
                                     }}
                                 >
-                                    {col}
+                                    {col.change}
                                 </div>
                             </th>
                         ))}
@@ -420,7 +545,7 @@ function GenerateClassicTable(stepOutcomeTable: StepOutcomeTable): JSX.Element {
                     {stepOutcomeTable.stepNameRowHeaders.slice(0, 20).map((key, rIdx) => (
                         <tr key={rIdx} style={{ height: '40px' }}>
                             <td
-                                title={key}
+                                title={key.stepName}
                                 style={{
                                     position: 'sticky',
                                     left: 0,
@@ -435,9 +560,9 @@ function GenerateClassicTable(stepOutcomeTable: StepOutcomeTable): JSX.Element {
                                     height: '40px',
                                 }}
                             >
-                                {key}
+                                {`${key.stepName}(${key.streamId})`}
                             </td>
-                            {[...stepOutcomeTable.getStepOutputRow(key.toLocaleLowerCase())]
+                            {[...stepOutcomeTable.getStepOutputRow(encodeStepNameFromStepNameHeader(key))]
                                 .slice(0, 20)
                                 .map((entry, cIdx) => {
                                     const { bgColor } = getCellStyle(entry);
@@ -475,6 +600,24 @@ function GenerateClassicTable(stepOutcomeTable: StepOutcomeTable): JSX.Element {
     );
 }
 
+/**
+ * React Component that provides refresh spinner on incremental refresh. 
+ * Note: This is a separate component to prevent rerender of primary table @see GenerateFluentUITable upon refresh spinner.
+ * @returns React Component.
+ */
+const IncrementalRefreshIndicatorPanel: React.FC<{ handler: StepOutcomeDataHandler }> = observer(({ handler }) => (
+    <div
+        style={{
+            gridRow: '1 / span 2',
+            gridColumn: 3,
+            overflowY: 'auto',
+            width: 25,
+        }}
+    >
+        {handler.isInIncrementalDataRefresh && <Spinner size={SpinnerSize.small} />}
+    </div>
+));
+
 /**
  * React Component that represents the provided stepOutcomeTable. This is a Fluent UI table.
  * @param stepOutcomeTable The data to base the table off of.
@@ -505,10 +648,10 @@ function GenerateFluentUITable(stepOutcomeTable: StepOutcomeTable, onCellClick:
 
     const allColumns = getFluentColumns(stepOutcomeTable);
     const items = getFluentRows(stepOutcomeTable);
+    const groups = getFluentGroups(stepOutcomeTable);
 
     const firstColumn = allColumns.slice(0, 1);
     const restColumns = allColumns.slice(1);
-
     const [hasScrollbar, setHasScrollbar] = useState(false);
 
     useEffect(() => {
@@ -523,28 +666,41 @@ function GenerateFluentUITable(stepOutcomeTable: StepOutcomeTable, onCellClick:
         return () => resizeObserver.disconnect();
     }, []);
 
-    const handleScroll = (direction: 'vertical' | 'horizontal', source: 'left' | 'rightTop' | 'rightBottom') =>
-        (e: React.UIEvent<HTMLDivElement>) => {
-            const scrollValue = direction === 'vertical' ? e.currentTarget.scrollTop : e.currentTarget.scrollLeft;
-            const syncRef = direction === 'vertical' ? isSyncingVertical : isSyncingHorizontal;
+    // Stash data structure to not create on all scroll events.
+    const targetMap = {
+        vertical: {
+            left: rightBottomPaneRef,
+            rightBottom: leftPaneRef,
+        },
+        horizontal: {
+            rightTop: rightBottomPaneRef,
+            rightBottom: rightTopPaneRef,
+        },
+    };
+
+    const VERTICAL_SCROLL_AFFORDANCE = 8;
 
+    const handleScroll = (direction: 'vertical' | 'horizontal', source: 'left' | 'rightTop' | 'rightBottom') => {
+        const isVertical = direction === 'vertical';
+        const syncRef = isVertical ? isSyncingVertical : isSyncingHorizontal;
+        const targetRef = targetMap[direction][source];
+
+        return (e: React.UIEvent<HTMLDivElement>) => {
             if (syncRef.current) return;
             syncRef.current = true;
 
-            const targetMap = {
-                vertical: {
-                    left: rightBottomPaneRef,
-                    rightBottom: leftPaneRef,
-                },
-                horizontal: {
-                    rightTop: rightBottomPaneRef,
-                    rightBottom: rightTopPaneRef,
-                },
-            };
-
-            const targetRef = targetMap[direction][source];
+            const el = e.currentTarget;
+            let scrollValue = isVertical ? el.scrollTop : el.scrollLeft;
+
+            // Clamp for rightTop horizontal due to how we need to overlay the top DetailList as a "sticky row header".
+            if (!isVertical && source === 'rightTop') {
+                const maxAllowed = el.scrollWidth - el.clientWidth - VERTICAL_SCROLL_AFFORDANCE;
+                scrollValue = Math.min(scrollValue, maxAllowed);
+                el.scrollLeft = scrollValue; // keep UI in sync immediately
+            }
+
             if (targetRef?.current) {
-                if (direction === 'vertical') {
+                if (isVertical) {
                     targetRef.current.scrollTop = scrollValue;
                 } else {
                     targetRef.current.scrollLeft = scrollValue;
@@ -555,6 +711,7 @@ function GenerateFluentUITable(stepOutcomeTable: StepOutcomeTable, onCellClick:
                 syncRef.current = false;
             });
         };
+    };
 
     const gridStyle = {
         display: 'grid',
@@ -587,6 +744,7 @@ function GenerateFluentUITable(stepOutcomeTable: StepOutcomeTable, onCellClick:
                 <DetailsList
                     items={items}
                     columns={firstColumn}
+                    groups={groups}
                     selectionMode={SelectionMode.none}
                     layoutMode={DetailsListLayoutMode.fixedColumns}
                     constrainMode={ConstrainMode.unconstrained}
@@ -640,6 +798,7 @@ function GenerateFluentUITable(stepOutcomeTable: StepOutcomeTable, onCellClick:
                 <DetailsList
                     items={items}
                     columns={restColumns}
+                    groups={groups}
                     selectionMode={SelectionMode.none}
                     layoutMode={DetailsListLayoutMode.fixedColumns}
                     constrainMode={ConstrainMode.unconstrained}
@@ -663,6 +822,7 @@ function GenerateFluentUITable(stepOutcomeTable: StepOutcomeTable, onCellClick:
                     }}
                 />
             </div>
+            <IncrementalRefreshIndicatorPanel handler={handler} />
         </div>
     );
 };
@@ -672,7 +832,13 @@ function GenerateFluentUITable(stepOutcomeTable: StepOutcomeTable, onCellClick:
  */
 const StepOutcomeTableComponent: React.FC<StepOutcomeTableProps> = observer(({ handler, onCellSelected }) => {
     handler.updated;
-    return GenerateFluentUITable(handler.stepOutcomeTableData, onCellSelected);
+    const table = GenerateFluentUITable(handler.StepOutcomeTableData, onCellSelected);
+
+    return (
+        <Stack>
+            {handler.isInFullDataRefresh ? <Spinner size={SpinnerSize.large} /> : table}
+        </Stack>
+    );
 });
 
 /**
diff --git a/Engine/Source/Programs/Horde/HordeDashboard/src/components/buildhealth/stepoutcome/StepOutcomeDataHandler.ts b/Engine/Source/Programs/Horde/HordeDashboard/src/components/buildhealth/stepoutcome/StepOutcomeDataHandler.ts
--- a/Engine/Source/Programs/Horde/HordeDashboard/src/components/buildhealth/stepoutcome/StepOutcomeDataHandler.ts
+++ b/Engine/Source/Programs/Horde/HordeDashboard/src/components/buildhealth/stepoutcome/StepOutcomeDataHandler.ts
@@ -4,18 +4,21 @@ import { GetJobResponse, JobStreamQuery } from "horde/backend/Api";
 import { PollBase } from "horde/backend/PollBase";
 import { StepOutcomeTable, StepOutcomeTableEntry } from "./StepOutcomeDataTypes";
 import backend from "horde/backend";
+import { decodeStepKey, decodeTemplateKey } from "../BuildHealthUtilities";
+import { action, computed, makeObservable, observable, runInAction } from "mobx";
 
 // #region -- API Data Types & Utilities --
 
 /**
  * Data structure that controls the filtering & queries of the Data Handler.
+ * @todo UE-309726 - This data structure can be simplified to reduce duplication, and reduce the need for the internal container to decode.
  */
 export interface StepOutcomeFilters {
     streams: string[];
     jobHistorySpan: { start: Date; end?: Date };
     jobs?: string[];
     steps?: string[];
-    maxJobCount? : number;
+    maxJobCount?: number;
 }
 
 // #endregion -- API Data Types & Utilities --
@@ -30,10 +33,14 @@ export class StepOutcomeDataHandler extends PollBase {
 
     // #region -- Private Members -- 
 
-    activeRefresh : boolean = false;
-    activeFilter: StepOutcomeFilters;
-    jobData: GetJobResponse[] = [];
-    stepOutcomeTableData: StepOutcomeTable = new StepOutcomeTable();
+    @observable
+    private intialRefresh = true;
+    @observable
+    private activeRefresh: boolean = false;
+
+    private activeFilter: StepOutcomeFilters;
+    private jobData: GetJobResponse[] = [];
+    private stepOutcomeTableData: StepOutcomeTable = new StepOutcomeTable();
 
     //#endregion -- Private Members
 
@@ -45,14 +52,16 @@ export class StepOutcomeDataHandler extends PollBase {
      */
     constructor(pollTime: number) {
         super(pollTime);
+        makeObservable(this);
     }
 
     //#endregion -- Constructor --
 
     // #region -- Interface --
 
+    @action
     async poll(): Promise<void> {
-        if (this.activeRefresh || 
+        if (this.activeRefresh ||
             this.activeFilter === undefined ||
             !this.activeFilter.streams?.length ||
             !this.activeFilter.jobs?.length
@@ -61,21 +70,37 @@ export class StepOutcomeDataHandler extends PollBase {
             this.setUpdated();
             return;
         }
+        runInAction(() => {
+            this.activeRefresh = true;
+        });
 
-        this.activeRefresh = true;
         this.stepOutcomeTableData.reset();
 
-        const stepSet = new Set(this.activeFilter.steps?.map(s => s.toLowerCase()));
-        const jobList = this.activeFilter.jobs ?? [];
+        // @todo UE-309726 - The poll should not need to understand internal data structure; the filter object should be already decoded.
+        const stepSet = new Set(this.activeFilter.steps?.map(s => decodeStepKey(s).stepName.toLowerCase()));
         const streamList = this.activeFilter.streams ?? [];
-        const queries: { streamId: string; query: JobStreamQuery }[] = [];
         const modifiedAfter = this.activeFilter.jobHistorySpan.start.toISOString();
         const jobCount = this.activeFilter.maxJobCount ?? this.MAX_JOB_COUNT_DEFAULT;
 
+        // Group jobs by streamId
+        const jobsByStream: Record<string, string[]> = {};
+
+        (this.activeFilter.jobs ?? []).forEach(jobKey => {
+            // @todo UE-309726 - The poll should not need to understand internal data structure; the filter object should be already decoded.
+            const { streamId, templateId } = decodeTemplateKey(jobKey);
+            if (!jobsByStream[streamId]) {
+                jobsByStream[streamId] = [];
+            }
+            jobsByStream[streamId].push(templateId);
+        });
+
+        const queries: { streamId: string; query: JobStreamQuery }[] = [];
+
         for (const streamId of streamList) {
-            for (const job of jobList) {
+            const jobsForStream = jobsByStream[streamId] ?? [];
+            for (const templateId of jobsForStream) {
                 const query: JobStreamQuery = {
-                    template: [job],
+                    template: [templateId],
                     count: jobCount,
                     filter: this.STREAM_JOBS_FILTER,
                     modifiedAfter: modifiedAfter
@@ -89,21 +114,24 @@ export class StepOutcomeDataHandler extends PollBase {
             queries.map(({ streamId, query }) => backend.getStreamJobs(streamId, query))
         );
 
+
         // Flatten all jobData
         this.jobData = results.flat();
 
-        // Process all job responses
-        for (const jobResponse of this.jobData) {
-            const batchResponses = jobResponse.batches ?? [];
-            for (const batch of batchResponses) {
-                for (const step of batch.steps) {
-                    if (
-                        stepSet.size === 0 ||
-                        stepSet.has(step.name.toLowerCase())
-                    ) {
-                        this.stepOutcomeTableData.addEntry(
-                            new StepOutcomeTableEntry(step, jobResponse.id, jobResponse.name, jobResponse.streamId, jobResponse.change)
-                        );
+        for (const streamResult of results) {
+            // Process all job responses
+            for (const jobResponse of streamResult) {
+                const batchResponses = jobResponse.batches ?? [];
+                for (const batch of batchResponses) {
+                    for (const step of batch.steps) {
+                        if (
+                            stepSet.size === 0 ||
+                            stepSet.has(step.name.toLowerCase())
+                        ) {
+                            this.stepOutcomeTableData.addEntry(
+                                new StepOutcomeTableEntry(step, jobResponse.id, jobResponse.name, jobResponse.streamId, jobResponse.createTime.toString(), jobResponse.change)
+                            );
+                        }
                     }
                 }
             }
@@ -113,8 +141,12 @@ export class StepOutcomeDataHandler extends PollBase {
 
         this.stepOutcomeTableData.orderTableDataByChange(true);
         this.setUpdated();
-        
-        this.activeRefresh = false;
+
+        runInAction(() => {
+            this.activeRefresh = false;
+            this.intialRefresh = false;
+        });
+
         return Promise.resolve();
     }
 
@@ -122,12 +154,40 @@ export class StepOutcomeDataHandler extends PollBase {
 
     // #region -- Public API  --
 
+    /**
+     * Gets whether an full data refresh is occurring.
+     * @returns True if an full data refresh is occurring, false otherwise.
+     */
+    @computed
+    get isInFullDataRefresh(): boolean {
+        return this.activeRefresh && this.intialRefresh;
+    }
+
+    /**
+     * Gets whether an incremental data refresh is occurring.
+     * @returns True if an incremental data refresh is occurring, false otherwise.
+     */
+    @computed
+    get isInIncrementalDataRefresh(): boolean {
+        return this.activeRefresh && !this.intialRefresh;
+    }
+
+    /**
+     * Gets the step outcome table data.
+     * @returns The instance Step Outcome Data Table.
+     */
+    get StepOutcomeTableData(): StepOutcomeTable {
+        return this.stepOutcomeTableData;
+    }
+
     /**
      * Sets a new filter for the handler, which will control the underlying dataset.
      * @param filter The new filter to use.
      */
+    @action
     setFilter(filter: StepOutcomeFilters) {
         this.activeFilter = filter;
+        this.intialRefresh = true;
         this.poll();
     }
 
diff --git a/Engine/Source/Programs/Horde/HordeDashboard/src/components/buildhealth/stepoutcome/StepOutcomeDataTypes.ts b/Engine/Source/Programs/Horde/HordeDashboard/src/components/buildhealth/stepoutcome/StepOutcomeDataTypes.ts
--- a/Engine/Source/Programs/Horde/HordeDashboard/src/components/buildhealth/stepoutcome/StepOutcomeDataTypes.ts
+++ b/Engine/Source/Programs/Horde/HordeDashboard/src/components/buildhealth/stepoutcome/StepOutcomeDataTypes.ts
@@ -2,6 +2,8 @@
 
 import { GetStepResponse } from "horde/backend/Api";
 
+const KEY_SEPARATOR: string = "::";
+
 /**
  * Model for a table entry.
  */
@@ -13,6 +15,7 @@ export class StepOutcomeTableEntry {
     jobId: string;
     jobName: string;
     streamId: string;
+    jobCreateTime: string;
 
     // #endregion -- Public Members --
 
@@ -24,9 +27,10 @@ export class StepOutcomeTableEntry {
      * @param jobId The owning job's Id.
      * @param jobName The owning job's name.
      * @param streamId The owning streams's Id.
+     * @param jobCreateTime The owning job's create time.
      * @param change The change associated with the stepResponse.
      */
-    constructor(stepResponse: GetStepResponse, jobId: string, jobName: string, streamId: string, change?: number) {
+    constructor(stepResponse: GetStepResponse, jobId: string, jobName: string, streamId: string, jobCreateTime: string, change?: number) {
         this.stepResponse = stepResponse;
         if (change !== undefined) {
             this.change = change;
@@ -34,11 +38,48 @@ export class StepOutcomeTableEntry {
         this.jobId = jobId;
         this.jobName = jobName;
         this.streamId = streamId;
+        this.jobCreateTime = jobCreateTime
     }
 
     // #endregion -- Constructor --
 }
 
+/**
+ * Encodes a step outcome table entry to a stream qualified step name.
+ * Note: The template is not included in the hierarchy qualification since two separate templates within the same stream will collocate the same step row.
+ * @param stepOutcomeTableEntry The stepOutcomeTableEntry to encode.
+ * @returns The hierarchical stream qualified step name.
+ */
+function encodeStepName(stepOutcomeTableEntry: StepOutcomeTableEntry) {
+    return `${stepOutcomeTableEntry.streamId.toLocaleLowerCase()}${KEY_SEPARATOR}${stepOutcomeTableEntry.stepResponse.name.toLocaleLowerCase()}`;
+}
+
+/**
+ * Encodes a step outcome table entry to a stream qualified step name.
+ * Note: The template is not included in the hierarchy qualification since two separate templates within the same stream will collocate the same step row.
+ * @param stepNameHeader The stepNameHeader to encode.
+ * @returns The hierarchical stream qualified step name.
+ */
+export function encodeStepNameFromStepNameHeader(stepNameHeader: StepNameHeader) {
+    return `${stepNameHeader.streamId.toLocaleLowerCase()}${KEY_SEPARATOR}${stepNameHeader.stepName.toLocaleLowerCase()}`;
+}
+
+/**
+ * Data struct that represents a change header. This header contains associated data related to rendering & ordering the header.
+ */
+export type ChangeHeader = {
+    change: number;
+    date?: string;
+}
+
+/**
+ * Data struct that represents a step name header. This header contains associated data related to rendering & ordering the header.
+ */
+export type StepNameHeader = {
+    stepName: string;
+    streamId: string;
+}
+
 /**
  * Model for the StepOutcomeTable.
  */
@@ -49,13 +90,15 @@ export class StepOutcomeTable {
      * Change column headers.
      */
     changeColLookup: Map<number, number> = new Map<number, number>();
-    changeColHeaders: number[] = [];
+    changeColHeaders: ChangeHeader[] = [];
 
     /**
      * Step name row headers.
      */
     stepRowLookup: Map<string, number> = new Map<string, number>();
-    stepNameRowHeaders: string[] = [];
+    stepNameRowHeaders: StepNameHeader[] = [];
+
+    tableStreamMetadata: Map<string, number> = new Map<string, number>();
 
     /**
      * Table data structure
@@ -72,6 +115,15 @@ export class StepOutcomeTable {
         }
     }
 
+    private annotateStreamRow(streamId: string) {
+        let priorCount: number = this.tableStreamMetadata.get(streamId) ?? 0;
+        if (priorCount) {
+            this.tableStreamMetadata.set(streamId, priorCount);
+        }
+        priorCount = priorCount + 1;
+        this.tableStreamMetadata.set(streamId, priorCount);
+    }
+
     // #endregion -- Private API --
 
     // #region -- Public API --
@@ -85,6 +137,7 @@ export class StepOutcomeTable {
         this.stepRowLookup.clear();
         this.changeColHeaders = [];
         this.changeColLookup.clear();
+        this.tableStreamMetadata.clear();
     }
 
     /**
@@ -97,14 +150,15 @@ export class StepOutcomeTable {
             return;
         }
 
-        const normalizedStepName = tableEntry.stepResponse.name.toLocaleLowerCase();
+        const normalizedStepName = encodeStepName(tableEntry);
 
         // Resolve or create row index
         let insertionRow = this.stepRowLookup.get(normalizedStepName);
         if (insertionRow === undefined) {
+            this.annotateStreamRow(tableEntry.streamId);
             insertionRow = this.stepNameRowHeaders.length;
             this.stepRowLookup.set(normalizedStepName, insertionRow);
-            this.stepNameRowHeaders.push(tableEntry.stepResponse.name);
+            this.stepNameRowHeaders.push({ streamId: tableEntry.streamId, stepName: tableEntry.stepResponse.name });
         }
 
         // Resolve or create column index
@@ -112,7 +166,7 @@ export class StepOutcomeTable {
         if (insertionCol === undefined) {
             insertionCol = this.changeColHeaders.length;
             this.changeColLookup.set(tableEntry.change, insertionCol);
-            this.changeColHeaders.push(tableEntry.change);
+            this.changeColHeaders.push({ change: tableEntry.change, date: tableEntry.jobCreateTime });
 
             for (const row of this.tableEntries) {
                 if (row) {
@@ -130,7 +184,7 @@ export class StepOutcomeTable {
 
         const existingEntry = this.tableEntries[insertionRow][insertionCol];
         if (existingEntry !== null) {
-            console.warn(`Duplicate entry for step new (${tableEntry.jobName}::${tableEntry.jobId}::${tableEntry.stepResponse.name} (pre-existing: ${existingEntry.jobName}::${existingEntry.jobId}::${existingEntry.stepResponse.name})) and change ("${tableEntry.change}" (${existingEntry.change})). Ambiguous selector needed.`, {
+            console.warn(`Duplicate entry for step new (${tableEntry.streamId}::${tableEntry.jobName}::${tableEntry.jobId}::${tableEntry.stepResponse.name} (pre-existing: ${existingEntry.streamId}::${existingEntry.jobName}::${existingEntry.jobId}::${existingEntry.stepResponse.name})) and change ("${tableEntry.change}" (${existingEntry.change})). Ambiguous selector needed.`, {
                 existingEntry,
                 newEntry: tableEntry
             });
@@ -163,7 +217,7 @@ export class StepOutcomeTable {
         let rowEntries = this.tableEntries[rowIdx];
 
         for (const entry of this.changeColHeaders) {
-            let colIdx = this.changeColLookup.get(entry);
+            let colIdx = this.changeColLookup.get(entry.change);
             if (colIdx === undefined) {
                 continue;
             }
@@ -188,7 +242,7 @@ export class StepOutcomeTable {
         }
 
         for (const entry of this.stepNameRowHeaders) {
-            let rowIdx: (number | undefined) = this.stepRowLookup.get(entry);
+            let rowIdx: (number | undefined) = this.stepRowLookup.get(encodeStepNameFromStepNameHeader(entry));
             if (rowIdx === undefined) {
                 continue;
             }
@@ -205,7 +259,7 @@ export class StepOutcomeTable {
      * @param ascending True for ascending order; false for descending order.
      */
     orderTableDataByChange(ascending: boolean): void {
-        this.changeColHeaders.sort((a, b) => ascending ? a - b : b - a);
+        this.changeColHeaders.sort((a, b) => ascending ? a.change - b.change : b.change - a.change);
     }
 
     /**
