[Horde][BuildHeatlh][Dashboard] Issue Details in modal, and update cell styling to use Horde compliant colours.

diff --git a/Engine/Source/Programs/Horde/HordeDashboard/src/components/buildhealth/stepoutcome/StepOutcome.tsx b/Engine/Source/Programs/Horde/HordeDashboard/src/components/buildhealth/stepoutcome/StepOutcome.tsx
--- a/Engine/Source/Programs/Horde/HordeDashboard/src/components/buildhealth/stepoutcome/StepOutcome.tsx
+++ b/Engine/Source/Programs/Horde/HordeDashboard/src/components/buildhealth/stepoutcome/StepOutcome.tsx
@@ -1,12 +1,13 @@
 // Copyright Epic Games, Inc. All Rights Reserved.
 
-import { Stack, DetailsList, IColumn, mergeStyleSets, SelectionMode, TooltipHost, DetailsListLayoutMode, ConstrainMode, Dialog, DialogType, Label, IGroup, Spinner, SpinnerSize } from "@fluentui/react";
+import { Stack, DetailsList, IColumn, mergeStyleSets, SelectionMode, TooltipHost, DetailsListLayoutMode, ConstrainMode, Dialog, DialogType, Label, IGroup, ScrollablePane, Spinner, SpinnerSize } from "@fluentui/react";
 import { observer } from "mobx-react-lite";
 import { getHordeStyling } from 'horde/styles/Styles';
 import React, { useEffect, useRef, useState } from 'react';
 import { Text } from '@fluentui/react/lib/Text';
 import { encodeStepNameFromStepNameHeader, StepNameHeader, StepOutcomeTable, StepOutcomeTableEntry } from "./StepOutcomeDataTypes";
 import { StepOutcomeFilters, StepOutcomeDataHandler } from "./StepOutcomeDataHandler";
+import dashboard, { StatusColor } from "horde/backend/Dashboard";
 
 // #region -- UI Helper Methods --
 
@@ -235,6 +236,14 @@ interface StepOutcomeViewProps {
  * Defines the style for column headers.
  */
 
+const stepOutcomeModalClasses = mergeStyleSets(
+    {
+        issueModalDividers: {
+            borderBottom: `1px solid rgba(80, 80, 80, 1)`
+        }
+    }
+);
+
 const headerStylesHighlight = mergeStyleSets({
     headerCell: {
         fontWeight: '700 !important',
@@ -401,7 +410,15 @@ function formatDuration(startISO: string, finishISO: string): string {
 const StepOutcomeFocusModal: React.FC<StepOutcomeFocusModalProps> = ({ item, onDismiss }) => {
     const startDate = toDate(item.stepResponse.startTime);
     const finishDate = toDate(item.stepResponse.finishTime);
-    let duration: string = "Ongoing";
+    let duration: string = item.stepResponse.state === "Skipped" ? "N/A" : "Ongoing";
+
+    const [issueRequestComplete, setIssueRequestComplete] = useState(false);
+
+    useEffect(() => {
+        if (!issueRequestComplete) {
+            handler.populateStepOutcomeTableEntryIssueData(item, setIssueRequestComplete);
+        }
+    }, []);
 
     if (startDate && finishDate) {
         duration = formatDuration(startDate.toISOString(), finishDate.toISOString());
@@ -431,23 +448,64 @@ const StepOutcomeFocusModal: React.FC<StepOutcomeFocusModalProps> = ({ item, onD
                 onDismiss: onDismiss,
                 title: item.stepResponse.name,
             }}>
-            <Stack tokens={{ childrenGap: 8 }}>
-                {[
-                    ['Date', startDate?.toDateString()],
-                    ['Stream', item.streamId],
-                    ['Code Change', item.change],
-                    ['Job', item.jobName],
-                    ['Step', item.stepResponse.name],
-                    ['Duration', duration],
-                    ['Outcome', item.stepResponse.outcome],
-                    ['Job Link', <a href={`/job/${item.jobId}`} target="_blank" rel="noopener noreferrer">View Job</a>],
-                    ['Step Link', <a href={`/job/${item.jobId}?step=${item.stepResponse.id}`} target="_blank" rel="noopener noreferrer">View Step</a>]
-                ].map(([label, value]) => (
-                    <Stack horizontal tokens={{ childrenGap: 12 }}>
-                        <Text styles={{ root: { fontWeight: '600', minWidth: 100 } }}>{label}:</Text>
-                        <Text>{value}</Text>
+            <Stack>
+                <Stack horizontal tokens={{ childrenGap: 16 }} styles={{ root: { width: '100%', flexGrow: 0 } }}>
+
+                    {/* Left Column - Step Details*/}
+                    <Stack styles={{ root: { width: '35%' } }} tokens={{ childrenGap: 4 }}>
+                        {[
+                            ['Date', startDate ? startDate.toDateString() : "N/A"],
+                            ['Stream', item.streamId],
+                            ['Code Change', item.change],
+                            ['Job', item.jobName],
+                            ['Step', item.stepResponse.name],
+                            ['Duration', duration],
+                            ['State', item.stepResponse.state],
+                            ['Outcome', item.stepResponse.outcome],
+                            ['Job Link', <a href={`/job/${item.jobId}`} target="_blank" rel="noopener noreferrer">View Job</a>],
+                            ['Step Link', <a href={`/job/${item.jobId}?step=${item.stepResponse.id}`} target="_blank" rel="noopener noreferrer">View Step</a>]
+                        ].map(([label, value], index) => (
+                            <Stack key={`modal_step_item_${index}`} horizontal tokens={{ childrenGap: 12 }}>
+                                <Label styles={{ root: { minWidth: 100, padding: 0 } }}>{label}:</Label>
+                                <Text>{value}</Text>
+                            </Stack>
+                        ))}
                     </Stack>
-                ))}
+
+                    {/* Right Column - Issues*/}
+                    <div style={{ width: '65%', maxHeight: 400, overflow: 'auto', position: 'relative' }}>
+                        <ScrollablePane styles={{ root: { height: '100%' } }}>
+                            <Stack tokens={{ childrenGap: 8 }}>
+                                <Stack horizontal tokens={{ childrenGap: 4 }} className={stepOutcomeModalClasses.issueModalDividers}>
+                                    <Label styles={{ root: { padding: 0 } }}>Issue Count:</Label>
+                                    <Text>{item.issuesData.length}</Text>
+                                    <Text>[{item.issuesData.filter((issue) => !issue.resolvedAt).length} unresolved, {item.issuesData.filter((issue) => issue.resolvedAt).length} resolved]</Text>
+                                </Stack>
+                                {item.issuesData.length > 0 ?
+                                    <Stack tokens={{ childrenGap: 10 }} styles={{ root: { paddingLeft: 16 } }}>
+                                        {item.issuesData.map((issue, index) => (
+                                            <Stack key={`modal_issue_item_${index}`} tokens={{ childrenGap: 4 }} className={stepOutcomeModalClasses.issueModalDividers}>
+                                                {[
+                                                    ['Issue', issue.id],
+                                                    ['Issue Link', <a href={`/job/${item.jobId}?step=${item.stepResponse.id}&issue=${issue.id}`} target="_blank" rel="noopener noreferrer">View Issue</a>],
+                                                    ['Issue Summary', issue.summary],
+                                                    ['Status', issue.resolvedAt ? "Resolved" : "Unresolved"],
+                                                    ['Resolved Streams', `[${issue.resolvedStreams.join(', ')}]`],
+                                                    ['Unresolved Streams', `[${issue.unresolvedStreams.join(', ')}]`]
+                                                ].map(([label, value], j) => (
+                                                    <Stack key={j} horizontal tokens={{ childrenGap: 12 }}>
+                                                        <Label styles={{ root: { minWidth: 150, padding: 0 } }}>{label}:</Label>
+                                                        <Text>{value}</Text>
+                                                    </Stack>
+                                                ))}
+                                            </Stack>
+                                        ))}
+                                    </Stack> : null
+                                }
+                            </Stack>
+                        </ScrollablePane>
+                    </div>
+                </Stack>
             </Stack>
         </Dialog>
     );
@@ -461,19 +519,51 @@ const StepOutcomeFocusModal: React.FC<StepOutcomeFocusModalProps> = ({ item, onD
  * @returns The matching style given the state of the entry. 
  */
 function getCellStyle(entry: StepOutcomeTableEntry | null): { bgColor: string } {
-    if (!entry) return { bgColor: "grey" };
+    const scolors = dashboard.getStatusColors();
+    let resultColor: string = 'transparent';
+
+    if (!entry) {
+        return { bgColor: resultColor };
+    }
+
+    const WarningStr: string = "Warnings";
+    const SuccessStr: string = "Success";
+    const SkippedStr: string = "Skipped";
+    const FailureStr: string = "Failure";
 
     const { stepResponse } = entry;
     const { outcome, state } = stepResponse;
+    const colors: Record<string, string> = {
+        "Success": scolors.get(StatusColor.Success)!,
+        "Failure": scolors.get(StatusColor.Failure)!,
+        "Warnings": scolors.get(StatusColor.Warnings)!,
+        "Skipped": scolors.get(StatusColor.Skipped)!,
+        "Running": scolors.get(StatusColor.Running)!,
+        "Waiting": scolors.get(StatusColor.Waiting)!,
+        "Ready": scolors.get(StatusColor.Ready)!,
+        "Unspecified": scolors.get(StatusColor.Skipped)!,
+    };
 
-    let bgColor = "grey";
-    if (["Aborted", "Completed", "Skipped"].includes(state)) {
-        if (outcome === "Failure") bgColor = state === "Skipped" ? "black" : "red";
-        else if (outcome === "Warnings") bgColor = "yellow";
-        else if (outcome === "Success") bgColor = "lightgreen";
+    // There are some interesting mappings going on internally with respect to state & outcome. Outcome & State have overlapping "result view" space, and as such we need to 
+    // narrowly select the edge cases (e.g. outcome=success && state=completed uniquely signifies a completed run with no errors, otherwise outcome=success it could mean ready, waiting, or running).
+    // Then we can fall back to the state as the truth.
+    if (outcome === FailureStr) {
+        resultColor = state === SkippedStr ? colors[SkippedStr] : colors[FailureStr];
+    }
+    // A step can be emitting warnings whilst still running, so we narrowly select this colouring.
+    else if (outcome === WarningStr) {
+        resultColor = colors[WarningStr];
+    }
+    // A step can have a completed outcome, but merely be waiting to run, running, or ready.
+    else if (outcome === SuccessStr && state === "Completed") {
+        resultColor = colors[SuccessStr];
+    }
+    // The edge cases have been satisfied, yield to the state to define result view color.
+    else {
+        resultColor = colors[state];
     }
 
-    return { bgColor };
+    return { bgColor: resultColor };
 }
 
 /**
@@ -882,7 +972,7 @@ export const StepOutcomeView: React.FC<StepOutcomeViewProps> = ({ filter, refres
 
 // #region -- Script --
 
-const DEFAULT_POLL_TIME = 60000;
+const DEFAULT_POLL_TIME = 120000;
 const handler = new StepOutcomeDataHandler(DEFAULT_POLL_TIME);
 
 // #endregion -- Script --
\ No newline at end of file
diff --git a/Engine/Source/Programs/Horde/HordeDashboard/src/components/buildhealth/stepoutcome/StepOutcomeDataHandler.ts b/Engine/Source/Programs/Horde/HordeDashboard/src/components/buildhealth/stepoutcome/StepOutcomeDataHandler.ts
--- a/Engine/Source/Programs/Horde/HordeDashboard/src/components/buildhealth/stepoutcome/StepOutcomeDataHandler.ts
+++ b/Engine/Source/Programs/Horde/HordeDashboard/src/components/buildhealth/stepoutcome/StepOutcomeDataHandler.ts
@@ -1,6 +1,6 @@
 // Copyright Epic Games, Inc. All Rights Reserved.
 
-import { GetJobResponse, JobStreamQuery } from "horde/backend/Api";
+import { GetJobResponse, IssueQuery, JobStreamQuery } from "horde/backend/Api";
 import { PollBase } from "horde/backend/PollBase";
 import { StepOutcomeTable, StepOutcomeTableEntry } from "./StepOutcomeDataTypes";
 import backend from "horde/backend";
@@ -150,6 +150,23 @@ export class StepOutcomeDataHandler extends PollBase {
         return Promise.resolve();
     }
 
+    /**
+     * Will populate the provided stepOutcomeTableEntry's issue data field,
+     * @param stepOutcomeTableEntry The step outcome table entry to attempt to populate with issue data.
+     */
+    async populateStepOutcomeTableEntryIssueData(stepOutcomeTableEntry : StepOutcomeTableEntry, onCompleteCallback? : (boolean) => void) : Promise<void> {
+        let unresolvedIssueQuery : IssueQuery = {jobId: stepOutcomeTableEntry.jobId, stepId : stepOutcomeTableEntry.stepResponse.id, resolved:false}
+        const unresolvedIssueResults = await backend.getIssues(unresolvedIssueQuery);
+        stepOutcomeTableEntry.issuesData = unresolvedIssueResults;
+
+        let resolvedIssueQuery : IssueQuery = {jobId: stepOutcomeTableEntry.jobId, stepId : stepOutcomeTableEntry.stepResponse.id, resolved:true}
+        const resolvedIssueResults = await backend.getIssues(resolvedIssueQuery);
+        stepOutcomeTableEntry.issuesData.push(...resolvedIssueResults);
+
+        onCompleteCallback?.(true);
+        return Promise.resolve();
+    }
+
     // #endregion -- Interface --
 
     // #region -- Public API  --
diff --git a/Engine/Source/Programs/Horde/HordeDashboard/src/components/buildhealth/stepoutcome/StepOutcomeDataTypes.ts b/Engine/Source/Programs/Horde/HordeDashboard/src/components/buildhealth/stepoutcome/StepOutcomeDataTypes.ts
--- a/Engine/Source/Programs/Horde/HordeDashboard/src/components/buildhealth/stepoutcome/StepOutcomeDataTypes.ts
+++ b/Engine/Source/Programs/Horde/HordeDashboard/src/components/buildhealth/stepoutcome/StepOutcomeDataTypes.ts
@@ -1,6 +1,6 @@
 // Copyright Epic Games, Inc. All Rights Reserved.
 
-import { GetStepResponse } from "horde/backend/Api";
+import { GetStepResponse, IssueData } from "horde/backend/Api";
 
 const KEY_SEPARATOR: string = "::";
 
@@ -16,6 +16,7 @@ export class StepOutcomeTableEntry {
     jobName: string;
     streamId: string;
     jobCreateTime: string;
+    issuesData : IssueData[] = [];
 
     // #endregion -- Public Members --
 
