Horde: Update release notes

#rnx

diff --git a/Engine/Source/Programs/Horde/Docs/ReleaseNotes.md b/Engine/Source/Programs/Horde/Docs/ReleaseNotes.md
--- a/Engine/Source/Programs/Horde/Docs/ReleaseNotes.md
+++ b/Engine/Source/Programs/Horde/Docs/ReleaseNotes.md
@@ -1,5 +1,21 @@
 # Release Notes
 
+## 2025-08-21
+
+* [Horde] Adds a PUT function for updating devices to the HordeHttpClient (45045274)
+* [Horde][BuildHeatlh][Dashboard] Issue Details in modal, and update cell styling to use Horde compliant colours. (45036254)
+* [Horde][BuildHealth][Dashboard] Multiple Stream selections, date headers, loading spinners on refreshes, URL query encoding fixes to include hierarchy, scroll-locking optimization, and various UI cleanup. (45035672)
+* Add max retries to agent delete (45021946)
+* Filter out anonymous pools in pools view (44988418)
+* Cloud DDC - Sending parts as uniform object as the validation code on the C++ side prefers this (its a more optimal format of this object). Implemented a very explicit support for writing uniform objects, fixing so that the writer converts to uniform types if appopriate required a much more intrusive change. (44981589)
+* Cloud DDC - Dispose of JsonDocument when converting CbObject to Json, was not aware it was disposable. Thank you static code analysis. (44981577)
+* Cloud DDC - Job metadata field serialization fix (ObjectId types was not converted correctly to json). Removed 1 of the 2 json conversation paths we had as the one we now have did not have this issue. Also fixed issues with serializing Guids to cb object. Lastly this adds support from serializing a json object into a CbObject, there is quite a lot of type loss going to json so this is not perfect but its handy to have in tests. (44928734)
+* Reduce the timeout when retrieving the system capabilities on macOS (44927820)
+* Add raw log event data, job/step information, and issue grouping details to issue audit log. (44907754)
+* Horde Agent: Fail step when AutomationTool process has a non-zero exit code (44901827)
+* Support new EC2 tag for identifying agent's pool ID
+  Older tag "Horde_Autoscale_Pool" identifying by name will still work but should be considered deprecated. (44854631)
+
 ## 2025-08-14
 
 * Disable logging of UAT exit code in JobExecutor Reverts 44661300. (44854494)
