[Horde][BuildHealth][Dashboard] Minor UI enhancements for:
- Drop downs in Build Health View
- Step Outcome Table tool-tips

diff --git a/Engine/Source/Programs/Horde/HordeDashboard/src/components/buildhealth/BuildhealthView.tsx b/Engine/Source/Programs/Horde/HordeDashboard/src/components/buildhealth/BuildhealthView.tsx
--- a/Engine/Source/Programs/Horde/HordeDashboard/src/components/buildhealth/BuildhealthView.tsx
+++ b/Engine/Source/Programs/Horde/HordeDashboard/src/components/buildhealth/BuildhealthView.tsx
@@ -1,19 +1,19 @@
 // Copyright Epic Games, Inc. All Rights Reserved.
 
-import { ComboBox, Dropdown, DropdownMenuItemType, IDropdownOption, Label, Slider, Stack, TooltipHost } from "@fluentui/react";
+import { ComboBox, DirectionalHint, Dropdown, DropdownMenuItemType, IDropdownOption, Label, Slider, Stack, TooltipHost } from "@fluentui/react";
+import { GetProjectResponse, GetStreamResponse } from "horde/backend/Api";
+import { observer } from "mobx-react-lite";
+import { useEffect, useMemo, useState } from "react";
+import { useLocation, useNavigate } from "react-router";
 import { useWindowSize } from "../../base/utilities/hooks";
 import { getHordeStyling } from "../../styles/Styles";
 import { Breadcrumbs } from "../Breadcrumbs";
 import { TopNav } from "../TopNav";
 import { StepOutcomeView } from "../buildhealth/stepoutcome/StepOutcome";
-import { useEffect, useMemo, useState } from "react";
-import { GetProjectResponse, GetStreamResponse } from "horde/backend/Api";
-import { observer } from "mobx-react-lite";
-import { useLocation, useNavigate } from "react-router";
 import { BuildHealthDataHandler, DataHandlerRefreshRequest, expandRefreshRequest } from "./BuildHealthDataHandler";
-import { BuildHealthOptionsController, BuildHealthOptionsState, loadBuildHealthOptionsFromParams, parseBuildHealthQueryParams, TimeSpan, JobHistoryTimeSpans } from "./BuildHealthOptions";
-import { encodeStepKey, encodeTemplateKey } from "./BuildHealthUtilities";
 import { KEY_SEPARATOR, PARAMETER_KEY_PREFIX, StepRefData, TemplateRefData } from "./BuildHealthDataTypes";
+import { BuildHealthOptionsController, BuildHealthOptionsState, JobHistoryTimeSpans, loadBuildHealthOptionsFromParams, parseBuildHealthQueryParams, TimeSpan } from "./BuildHealthOptions";
+import { encodeStepKey, encodeTemplateKey } from "./BuildHealthUtilities";
 
 // #region -- Dropdown Components --
 
@@ -54,10 +54,13 @@ const BasicListParameter: React.FC<{ handler: BuildHealthDataHandler, buildHealt
     return (<Dropdown id={changeVersion.toString() + key}
         key={key}
         label={param.label}
+        placeholder={"Select option"}
         options={doptions}
-        dropdownWidth={200}
         selectedKey={selectedKey}
-        placeholder={"Select option"}
+        dropdownWidth={200}
+        calloutProps={{
+            directionalHint: DirectionalHint.rightCenter
+        }}
         onChange={(_ev, option) => {
             let castedOption = option as IDropdownOption<BuildHealthOptionData>;
             if (!castedOption) {
@@ -79,7 +82,7 @@ const BasicListParameter: React.FC<{ handler: BuildHealthDataHandler, buildHealt
 const MultiListParameter: React.FC<{ handler: BuildHealthDataHandler, buildHealthOptions: BuildHealthOptionsController, param: BuildHealthOptionList, disabled?: boolean, selectedKeys: string[], request: DataHandlerRefreshRequest, onChange: (option: IDropdownOption<BuildHealthOptionData>) => void }> = function ConstructMultilistParamter({ param, buildHealthOptions, selectedKeys, request, onChange }) {
     const changeVersion = buildHealthOptions.optionsChangeVersion;
     const { modeColors } = getHordeStyling();
-    const key = `${PARAMETER_KEY_PREFIX}${KEY_SEPARATOR}${param.label}`;// + `_${options.currentRenderKey}`;
+    const key = `${PARAMETER_KEY_PREFIX}${KEY_SEPARATOR}${param.label}`;
     const gset: Set<string> = new Set();
 
     param.items.forEach(item => {
@@ -93,10 +96,9 @@ const MultiListParameter: React.FC<{ handler: BuildHealthDataHandler, buildHealt
 
     const groups = Array.from(gset).sort((a, b) => a.localeCompare(b));
 
-    const doptions: IDropdownOption[] = [];
+    const doptions: IDropdownOption<BuildHealthOptionData>[] = [];
 
     groups.forEach(group => {
-
         if (group !== "__nogroup") {
             doptions.push({
                 key: `group_${group}`,
@@ -105,20 +107,13 @@ const MultiListParameter: React.FC<{ handler: BuildHealthDataHandler, buildHealt
             });
         }
 
-        const dupes = new Map<string, number>();
-
-        param.items.forEach(item => {
-            const v = dupes.get(item.text);
-            dupes.set(item.text, !v ? 1 : v + 1);
-        });
-
         param.items.forEach(item => {
             if (item.group === group) {
                 const key = item.id;
                 doptions.push({
                     key: key,
                     data: item,
-                    text: (item.group !== "__nogroup" && dupes.get(item.text)! > 1) ? `${item.text} - ${item.group}` : item.text,
+                    text: item.text
                 });
             }
         });
@@ -127,9 +122,16 @@ const MultiListParameter: React.FC<{ handler: BuildHealthDataHandler, buildHealt
     return (
         <Dropdown id={changeVersion.toString() + key}
             key={key}
+            label={param.label}
             placeholder={"Select options"}
-            dropdownWidth={300}
+            options={doptions}
+            selectedKeys={selectedKeys}
+            dropdownWidth={400}
+            calloutProps={{
+                directionalHint: DirectionalHint.rightCenter
+            }}
             styles={{
+                dropdown: { width: 300 },
                 callout: {
                     selectors: {
                         ".ms-Callout-main": {
@@ -160,8 +162,35 @@ const MultiListParameter: React.FC<{ handler: BuildHealthDataHandler, buildHealt
                     }
                 }
             }}
-            label={param.label}
-            selectedKeys={selectedKeys}
+            onRenderItem={(item, defaultRender) => {
+                let castedOption = item as IDropdownOption<BuildHealthOptionData>;
+                if (!castedOption || !defaultRender) return null;
+
+                if (!castedOption.data) {
+                    return defaultRender(item);
+                }
+                return (
+                    <TooltipHost key={`tooltip_${castedOption.text}`} content={
+                        <div>
+                            <Stack>
+                                <Label>{castedOption.data?.group}</Label>
+                            </Stack>
+                            <Stack>
+                                {castedOption.text}
+                            </Stack>
+                        </div>
+                    }
+                        directionalHint={DirectionalHint.rightCenter}
+                        calloutProps={{
+                            styles: {
+                                root: { maxWidth: 600 }
+                            }
+                        }}
+                    >
+                        {defaultRender(item)}
+                    </TooltipHost>
+                );
+            }}
             onChange={(_ev, option) => {
                 let castedOption = option as IDropdownOption<BuildHealthOptionData>;
                 if (!castedOption) {
@@ -174,7 +203,6 @@ const MultiListParameter: React.FC<{ handler: BuildHealthDataHandler, buildHealt
                 buildHealthOptions.synchronizeDerivedKeys();
             }}
             multiSelect
-            options={doptions}
         />);
 };
 
@@ -353,7 +381,7 @@ const StepDropdown: React.FC<{ handler: BuildHealthDataHandler, buildHealthOptio
         let listParam: BuildHealthOptionData = {
             id: encodeStepKey(step),
             text: step.name,
-            group: `${step.streamId}${KEY_SEPARATOR}${step.templateId}`
+            group: `${step.streamId} - ${step.templateId}`
         };
         return listParam;
     }).sort((a, b) => a.text.localeCompare(b.text));
@@ -396,7 +424,7 @@ const BuildHealthUrlSync: React.FC<{ buildHealthController: BuildHealthOptionsCo
 
         initializeFromUrl();
     }, []);
-    
+
     useEffect(() => {
         const desiredSearch = buildHealthController.toNavigationQuery().toString();
         if (desiredSearch && location.search !== `?${desiredSearch}`) {
diff --git a/Engine/Source/Programs/Horde/HordeDashboard/src/components/buildhealth/stepoutcome/StepOutcome.tsx b/Engine/Source/Programs/Horde/HordeDashboard/src/components/buildhealth/stepoutcome/StepOutcome.tsx
--- a/Engine/Source/Programs/Horde/HordeDashboard/src/components/buildhealth/stepoutcome/StepOutcome.tsx
+++ b/Engine/Source/Programs/Horde/HordeDashboard/src/components/buildhealth/stepoutcome/StepOutcome.tsx
@@ -1,13 +1,14 @@
 // Copyright Epic Games, Inc. All Rights Reserved.
 
-import { Stack, DetailsList, IColumn, mergeStyleSets, SelectionMode, TooltipHost, DetailsListLayoutMode, ConstrainMode, Dialog, DialogType, Label, IGroup, ScrollablePane, Spinner, SpinnerSize } from "@fluentui/react";
-import { observer } from "mobx-react-lite";
+import { ConstrainMode, DetailsList, DetailsListLayoutMode, Dialog, DialogType, IColumn, IGroup, Label, mergeStyleSets, ScrollablePane, SelectionMode, Spinner, SpinnerSize, Stack, TooltipHost } from "@fluentui/react";
+import { Text } from '@fluentui/react/lib/Text';
+import dashboard, { StatusColor } from "horde/backend/Dashboard";
 import { getHordeStyling } from 'horde/styles/Styles';
+import { observer } from "mobx-react-lite";
 import React, { useEffect, useRef, useState } from 'react';
-import { Text } from '@fluentui/react/lib/Text';
+import { StepOutcomeDataHandler, StepOutcomeFilters } from "./StepOutcomeDataHandler";
 import { encodeStepNameFromStepNameHeader, StepNameHeader, StepOutcomeTable, StepOutcomeTableEntry } from "./StepOutcomeDataTypes";
-import { StepOutcomeFilters, StepOutcomeDataHandler } from "./StepOutcomeDataHandler";
-import dashboard, { StatusColor } from "horde/backend/Dashboard";
+import { toDate } from "./StepOutcomeUtilities";
 
 // #region -- UI Helper Methods --
 
@@ -129,11 +130,11 @@ function getFluentColumns(stepOutcomeTable: StepOutcomeTable): IColumn[] {
                 return props && defaultRender ? (
                     <TooltipHost content={
                         <Stack>
-                            <Label>Change:  </Label>
+                            <Label>Change:</Label>
                             <Text style={{ paddingLeft: 8 }}>{changeHeader.change}</Text>
-                            <Label>Date (UTC): </Label>
+                            <Label>Date (UTC):</Label>
                             <Text style={{ paddingLeft: 8 }} >{new Date(changeHeader.date!).toISOString()}</Text>
-                            <Label>Date (Local): </Label>
+                            <Label>Date (Local):</Label>
                             <Text style={{ paddingLeft: 8 }} >{new Date(changeHeader.date!).toString()}</Text>
                         </Stack>
                     }
@@ -155,9 +156,13 @@ function getFluentColumns(stepOutcomeTable: StepOutcomeTable): IColumn[] {
                     <TooltipHost content={
                         <Stack>
                             <Label>Change:</Label>
-                            <Text style={{ paddingLeft: 8 }}>{changeHeader.change}</Text>
+                            <Text className={stepOutcomeTooltipClasses.bottomBorder} style={{ paddingLeft: 8 }}>{changeHeader.change}</Text>
                             {entry && (
                                 <>
+                                    <Label>Duration:</Label>
+                                    <Text style={{ paddingLeft: 8 }}>{entry.getDurationString()}</Text>
+                                    <Label>State:</Label>
+                                    <Text className={stepOutcomeTooltipClasses.bottomBorder} style={{ paddingLeft: 8 }}>{entry.stepResponse.state}</Text>
                                     <Label>Stream:</Label>
                                     <Text style={{ paddingLeft: 8 }}>{entry.streamId}</Text>
                                     <Label>Job:</Label>
@@ -236,13 +241,15 @@ interface StepOutcomeViewProps {
  * Defines the style for column headers.
  */
 
-const stepOutcomeModalClasses = mergeStyleSets(
-    {
-        issueModalDividers: {
-            borderBottom: `1px solid rgba(80, 80, 80, 1)`
-        }
+const sharedBorderStyle = {
+    bottomBorder: {
+        borderBottom: `1px solid rgba(80, 80, 80, 1)`
     }
-);
+};
+
+const stepOutcomeTooltipClasses = mergeStyleSets(sharedBorderStyle);
+
+const stepOutcomeModalClasses = mergeStyleSets(sharedBorderStyle);
 
 const headerStylesHighlight = mergeStyleSets({
     headerCell: {
@@ -361,7 +368,6 @@ const stepOutcomeTableClasses = mergeStyleSets({
 
 // #endregion -- Styles  --
 
-
 // #region -- Modal Popup View --
 
 type StepOutcomeFocusModalProps = {
@@ -369,48 +375,11 @@ type StepOutcomeFocusModalProps = {
     onDismiss?: (_ev?: React.MouseEvent<HTMLButtonElement>) => any;
 };
 
-/**
- * Normalizes an ambiguous input (Date | string) to a Date object. 
- * @param input 
- * @returns 
- */
-function toDate(input?: Date | string): Date | undefined {
-    if (!input) return undefined;
-
-    if (input instanceof Date) {
-        return input;
-    }
-
-    const parsed = new Date(input);
-    return isNaN(parsed.getTime()) ? undefined : parsed;
-}
-
-/**
- * Generates a formatted duration provided a start and finish time.
- * @param startISO The start time, in string representation.
- * @param finishISO The finish time, in string representation.
- * @returns The formated duration.
- */
-function formatDuration(startISO: string, finishISO: string): string {
-    const start = new Date(startISO);
-    const finish = new Date(finishISO);
-    const diffMs = finish.getTime() - start.getTime();
-    const diffSeconds = Math.floor(diffMs / 1000);
-
-    const hours = Math.floor(diffSeconds / 3600);
-    const minutes = Math.floor((diffSeconds % 3600) / 60);
-    const seconds = diffSeconds % 60;
-
-    return `${hours}h ${minutes}m ${seconds}s`;
-}
-
 /**
  * React Component representing the Step Outcome focus.
  */
 const StepOutcomeFocusModal: React.FC<StepOutcomeFocusModalProps> = ({ item, onDismiss }) => {
     const startDate = toDate(item.stepResponse.startTime);
-    const finishDate = toDate(item.stepResponse.finishTime);
-    let duration: string = item.stepResponse.state === "Skipped" ? "N/A" : "Ongoing";
 
     const [issueRequestComplete, setIssueRequestComplete] = useState(false);
 
@@ -420,10 +389,6 @@ const StepOutcomeFocusModal: React.FC<StepOutcomeFocusModalProps> = ({ item, onD
         }
     }, []);
 
-    if (startDate && finishDate) {
-        duration = formatDuration(startDate.toISOString(), finishDate.toISOString());
-    }
-
     return (
         <Dialog
             modalProps={{
@@ -459,7 +424,7 @@ const StepOutcomeFocusModal: React.FC<StepOutcomeFocusModalProps> = ({ item, onD
                             ['Code Change', item.change],
                             ['Job', item.jobName],
                             ['Step', item.stepResponse.name],
-                            ['Duration', duration],
+                            ['Duration', item.getDurationString()],
                             ['State', item.stepResponse.state],
                             ['Outcome', item.stepResponse.outcome],
                             ['Job Link', <a href={`/job/${item.jobId}`} target="_blank" rel="noopener noreferrer">View Job</a>],
@@ -476,7 +441,7 @@ const StepOutcomeFocusModal: React.FC<StepOutcomeFocusModalProps> = ({ item, onD
                     <div style={{ width: '65%', maxHeight: 400, overflow: 'auto', position: 'relative' }}>
                         <ScrollablePane styles={{ root: { height: '100%' } }}>
                             <Stack tokens={{ childrenGap: 8 }}>
-                                <Stack horizontal tokens={{ childrenGap: 4 }} className={stepOutcomeModalClasses.issueModalDividers}>
+                                <Stack horizontal tokens={{ childrenGap: 4 }} className={stepOutcomeModalClasses.bottomBorder}>
                                     <Label styles={{ root: { padding: 0 } }}>Issue Count:</Label>
                                     <Text>{item.issuesData.length}</Text>
                                     <Text>[{item.issuesData.filter((issue) => !issue.resolvedAt).length} unresolved, {item.issuesData.filter((issue) => issue.resolvedAt).length} resolved]</Text>
@@ -484,7 +449,7 @@ const StepOutcomeFocusModal: React.FC<StepOutcomeFocusModalProps> = ({ item, onD
                                 {item.issuesData.length > 0 ?
                                     <Stack tokens={{ childrenGap: 10 }} styles={{ root: { paddingLeft: 16 } }}>
                                         {item.issuesData.map((issue, index) => (
-                                            <Stack key={`modal_issue_item_${index}`} tokens={{ childrenGap: 4 }} className={stepOutcomeModalClasses.issueModalDividers}>
+                                            <Stack key={`modal_issue_item_${index}`} tokens={{ childrenGap: 4 }} className={stepOutcomeModalClasses.bottomBorder}>
                                                 {[
                                                     ['Issue', issue.id],
                                                     ['Issue Link', <a href={`/job/${item.jobId}?step=${item.stepResponse.id}&issue=${issue.id}`} target="_blank" rel="noopener noreferrer">View Issue</a>],
diff --git a/Engine/Source/Programs/Horde/HordeDashboard/src/components/buildhealth/stepoutcome/StepOutcomeDataTypes.ts b/Engine/Source/Programs/Horde/HordeDashboard/src/components/buildhealth/stepoutcome/StepOutcomeDataTypes.ts
--- a/Engine/Source/Programs/Horde/HordeDashboard/src/components/buildhealth/stepoutcome/StepOutcomeDataTypes.ts
+++ b/Engine/Source/Programs/Horde/HordeDashboard/src/components/buildhealth/stepoutcome/StepOutcomeDataTypes.ts
@@ -1,9 +1,17 @@
 // Copyright Epic Games, Inc. All Rights Reserved.
 
 import { GetStepResponse, IssueData } from "horde/backend/Api";
+import { computeDuration, DURATION_NOT_SET, formatDurationFromMs, toDate } from "./StepOutcomeUtilities";
 
 const KEY_SEPARATOR: string = "::";
 
+/*
+ * Convenience type to fold in extra metadata relevant for the StepResponse.
+ */
+export type StepReponseData = GetStepResponse & {
+    stepDurationMs: number;
+}
+
 /**
  * Model for a table entry.
  */
@@ -11,12 +19,12 @@ export class StepOutcomeTableEntry {
     // #region -- Public Members --
 
     change?: number;
-    stepResponse: GetStepResponse;
+    stepResponse: StepReponseData;
     jobId: string;
     jobName: string;
     streamId: string;
     jobCreateTime: string;
-    issuesData : IssueData[] = [];
+    issuesData: IssueData[] = [];
 
     // #endregion -- Public Members --
 
@@ -32,10 +40,11 @@ export class StepOutcomeTableEntry {
      * @param change The change associated with the stepResponse.
      */
     constructor(stepResponse: GetStepResponse, jobId: string, jobName: string, streamId: string, jobCreateTime: string, change?: number) {
-        this.stepResponse = stepResponse;
         if (change !== undefined) {
             this.change = change;
         }
+
+        this.stepResponse = { ...stepResponse, stepDurationMs: computeDuration(stepResponse.startTime, stepResponse.finishTime) };
         this.jobId = jobId;
         this.jobName = jobName;
         this.streamId = streamId;
@@ -43,6 +52,23 @@ export class StepOutcomeTableEntry {
     }
 
     // #endregion -- Constructor --
+
+    // #region -- Public API --
+
+    /**
+     * Gets the duration of the underyling @see StepReponseData in string representation.
+     * @returns The string representation of the step duration.
+     */
+    getDurationString(): string {
+        let durationStr: string = this.stepResponse.state === "Skipped" ? "N/A" : "Ongoing";
+
+        if (this.stepResponse.stepDurationMs != DURATION_NOT_SET) {
+            durationStr = formatDurationFromMs(this.stepResponse.stepDurationMs);
+        }
+
+        return durationStr;
+    }
+    // #endregion -- Public API --
 }
 
 /**
diff --git a/Engine/Source/Programs/Horde/HordeDashboard/src/components/buildhealth/stepoutcome/StepOutcomeUtilities.ts b/Engine/Source/Programs/Horde/HordeDashboard/src/components/buildhealth/stepoutcome/StepOutcomeUtilities.ts
new file mode 100644
--- /dev/null
+++ b/Engine/Source/Programs/Horde/HordeDashboard/src/components/buildhealth/stepoutcome/StepOutcomeUtilities.ts
@@ -0,0 +1,63 @@
+// Copyright Epic Games, Inc. All Rights Reserved.
+
+export const DURATION_NOT_SET: number = -1;
+
+/**
+ * Normalizes an ambiguous input (Date | string) to a Date object. 
+ * @param input the input string union type.
+ * @returns The @see Date standard representation for the time. 
+ */
+export function toDate(input?: Date | string): Date | undefined {
+    if (!input) {
+        return undefined;
+    }
+
+    if (input instanceof Date) {
+        return input;
+    }
+
+    const parsed = new Date(input);
+    return isNaN(parsed.getTime()) ? undefined : parsed;
+}
+
+/**
+ * Computes the duration of a window in ms from a common time response format (Date | string).
+ * @param startTime The start time of the window.
+ * @param finishTime The finish time of the window.
+ * @returns The duration of the window in ms, @see DURATION_NOT_SET if an invalid window is provided.
+ */
+export function computeDuration(startTime?: Date | string, finishTime?: Date | string): number {
+    const startDate = toDate(startTime);
+    const finishDate = toDate(finishTime);
+    return startDate && finishDate
+        ? finishDate.getTime() - startDate.getTime()
+        : DURATION_NOT_SET;
+}
+
+/**
+ * Generates a formatted duration provided a time span in milliseconds. 
+ * @param ms The duration in milliseconds.
+ * @returns The formated duration.
+ */
+export function formatDurationFromMs(ms: number): string {
+    const diffSeconds = Math.floor(ms / 1000);
+    const hours = Math.floor(diffSeconds / 3600);
+    const minutes = Math.floor((diffSeconds % 3600) / 60);
+    const seconds = diffSeconds % 60;
+
+    return `${hours}h ${minutes}m ${seconds}s`;
+}
+
+/**
+ * Generates a formatted duration provided a start and finish time.
+ * @param startISO The start time, in string representation.
+ * @param finishISO The finish time, in string representation.
+ * @returns The formated duration.
+ */
+export function formatDuration(startISO: string, finishISO: string): string {
+    const start = new Date(startISO);
+    const finish = new Date(finishISO);
+    const diffMs = finish.getTime() - start.getTime();
+
+    return formatDurationFromMs(diffMs);
+}
\ No newline at end of file
