Added support for StreamTags to determine if notifications should be sent

#rnx

diff --git a/Engine/Source/Programs/Horde/Plugins/Experimental/HordeServer.Experimental.Tests/Notifications/ExperimentalNotificationTests.cs b/Engine/Source/Programs/Horde/Plugins/Experimental/HordeServer.Experimental.Tests/Notifications/ExperimentalNotificationTests.cs
--- a/Engine/Source/Programs/Horde/Plugins/Experimental/HordeServer.Experimental.Tests/Notifications/ExperimentalNotificationTests.cs
+++ b/Engine/Source/Programs/Horde/Plugins/Experimental/HordeServer.Experimental.Tests/Notifications/ExperimentalNotificationTests.cs
@@ -224,6 +224,8 @@ namespace HordeServer.Experimental.Tests.Notifications
 		StreamConfig CreateStream(StreamId streamId, string streamName)
 		{
 			StreamConfig streamConfig = new StreamConfig { Id = streamId, Name = streamName };
+			streamConfig.StreamTags.Add(new StreamTag { Name = "NotificationEnabledTag", Enabled = true });
+			streamConfig.StreamTags.Add(new StreamTag { Name = "DisabledTag", Enabled = false });
 			streamConfig.Tabs.Add(new TabConfig { Title = "General", Templates = new List<TemplateId> { _buildAndTestTemplateId, _buildTemplateId, _spawnedTestTemplateId } });
 			streamConfig.Templates.Add(_testBuildTestJobTemplate);
 			streamConfig.Templates.Add(_testBuildJobTemplate);
@@ -1286,5 +1288,63 @@ namespace HordeServer.Experimental.Tests.Notifications
 			consolidatedJobStepNotificationStates = await jobNotificationCollection.GetJobStepNotificationStatesAsync(consolidatedJobStepNotificationQueryBuilder);
 			Assert.IsNull(consolidatedJobStepNotificationStates);
 		}
+
+		[TestMethod]
+		public async Task NotificationConfigurationsByStreamTag()
+		{
+			// Scenario: Job matched with the configuration using associated StreamTags enabled in the StreamConfig
+			// Expected: Single record should exist in both the JobNotificationCollection and JobStepNotificationCollection
+
+			int commitId = 50;
+			NotificationFormatConfig spawnNotificationFormatConfig = CreateNotificationFormatConfig(SpawnJobTestGroup, "^Spawn\\s");
+
+			JobNotificationConfig jobNotificationConfig = new JobNotificationConfig();
+			jobNotificationConfig.Enabled = true;
+			jobNotificationConfig.Template = _buildTemplateId;
+			jobNotificationConfig.NamePattern = "Platform1";
+			jobNotificationConfig.NotificationFormats.Add(spawnNotificationFormatConfig);
+			jobNotificationConfig.Channels.Add(ChannelId);
+
+			NotificationStreamConfig notificationStreamConfig = new NotificationStreamConfig();
+			notificationStreamConfig.EnabledStreamTags.Add("NotificationEnabledTag");
+			notificationStreamConfig.JobNotifications.Add(jobNotificationConfig);
+
+			NotificationConfig notificationConfig = new NotificationConfig();
+			notificationConfig.Id = new NotificationConfigId("test-notification");
+			notificationConfig.NotificationStreams.Add(notificationStreamConfig);
+
+			ExperimentalConfig experimentalConfig = new ExperimentalConfig();
+			experimentalConfig.Notifications.Add(notificationConfig);
+
+			UpdateGlobalConfig(experimentalConfig);
+
+			JobNotificationCollection jobNotificationCollection = ServiceProvider.GetRequiredService<JobNotificationCollection>();
+
+			FakeExperimentalNotificationSink fakeSink = ServiceProvider.GetRequiredService<FakeExperimentalNotificationSink>();
+
+			NotificationService service = (NotificationService)ServiceProvider.GetRequiredService<INotificationService>();
+			await service.StartAsync(cancellationToken: default);
+
+			IJob job = await CreateNightlyBuildJobAsync(_mainStreamConfig, commitId);
+			_testJobNotifications.Add(new JobNotificationState(job.Id, job.TemplateId, ChannelId));
+
+			job = await RunBatchAsync(_mainStreamConfig, job, 0);
+			job = await RunStepAsync(_mainStreamConfig, job, 0, 0, JobStepOutcome.Success); // Setup Build
+
+			await WaitForNotificationAsync(fakeSink, job, job.Batches[0].Steps[0]);
+
+			// There should be items in the collections if we successfully picked up the job from the tag
+			IJobNotificationStateQueryBuilder jobNotificationQueryBuilder = JobNotificationCollection.CreateJobNotificationStateQueryBuilder();
+			jobNotificationQueryBuilder.AddJobFilter(job.Id);
+			IReadOnlyList<IJobNotificationStateRef>? jobNotificationStates = await jobNotificationCollection.GetJobNotificationStatesAsync(jobNotificationQueryBuilder);
+			Assert.IsNotNull(jobNotificationStates);
+			Assert.IsTrue(jobNotificationStates.Count == 1);
+
+			IJobStepNotificationStateQueryBuilder jobStepNotificationQueryBuilder = JobNotificationCollection.CreateJobStepNotificationStateQueryBuilder();
+			jobStepNotificationQueryBuilder.AddJobFilter(job.Id);
+			IReadOnlyList<IJobStepNotificationStateRef>? jobStepNotificationStates = await jobNotificationCollection.GetJobStepNotificationStatesAsync(jobStepNotificationQueryBuilder);
+			Assert.IsNotNull(jobStepNotificationStates);
+			Assert.IsTrue(jobStepNotificationStates.Count == 1);
+		}
 	}
 }
diff --git a/Engine/Source/Programs/Horde/Plugins/Experimental/HordeServer.Experimental/ExperimentalConfig.cs b/Engine/Source/Programs/Horde/Plugins/Experimental/HordeServer.Experimental/ExperimentalConfig.cs
--- a/Engine/Source/Programs/Horde/Plugins/Experimental/HordeServer.Experimental/ExperimentalConfig.cs
+++ b/Engine/Source/Programs/Horde/Plugins/Experimental/HordeServer.Experimental/ExperimentalConfig.cs
@@ -3,6 +3,7 @@
 using EpicGames.Horde.Jobs.Templates;
 using EpicGames.Horde.Streams;
 using HordeServer.Plugins;
+using HordeServer.Streams;
 using System.Diagnostics.CodeAnalysis;
 
 namespace HordeServer
@@ -36,16 +37,18 @@ namespace HordeServer
 		/// <param name="streamId">Stream id</param>
 		/// <param name="templateId">Template id</param>
 		/// <param name="jobNotificationConfigs">Receives the list of config objects on success</param>
+		/// <param name="streamTags">List of tags for a stream</param>
 		/// <param name="includeDisabledConfigs">Flag on whether or not disabled configs should be included</param>
 		/// <returns>True if the job notification configs were found</returns>
-		public bool TryGetJobNofications(StreamId streamId, TemplateId templateId, [NotNullWhen(true)] out HashSet<JobNotificationConfig>? jobNotificationConfigs, bool includeDisabledConfigs = false)
+		public bool TryGetJobNofications(StreamId streamId, TemplateId templateId, [NotNullWhen(true)] out HashSet<JobNotificationConfig>? jobNotificationConfigs, List<StreamTag>? streamTags = null, bool includeDisabledConfigs = false)
 		{
 			jobNotificationConfigs = new HashSet<JobNotificationConfig>();
 
 			foreach (NotificationConfig config in Notifications)
 			{
 				HashSet<JobNotificationConfig>? localJobNotificationConfigs;
-				if (config.TryGetJobNofications(streamId, templateId, out localJobNotificationConfigs, includeDisabledConfigs))
+				if (config.TryGetJobNofications(streamId, templateId, out localJobNotificationConfigs, includeDisabledConfigs) ||
+					(streamTags is not null && config.TryGetJobNofications(streamTags, templateId, out localJobNotificationConfigs, includeDisabledConfigs)))
 				{
 					jobNotificationConfigs.UnionWith(localJobNotificationConfigs);
 				}
diff --git a/Engine/Source/Programs/Horde/Plugins/Experimental/HordeServer.Experimental/Notifications/NotificationConfig.cs b/Engine/Source/Programs/Horde/Plugins/Experimental/HordeServer.Experimental/Notifications/NotificationConfig.cs
--- a/Engine/Source/Programs/Horde/Plugins/Experimental/HordeServer.Experimental/Notifications/NotificationConfig.cs
+++ b/Engine/Source/Programs/Horde/Plugins/Experimental/HordeServer.Experimental/Notifications/NotificationConfig.cs
@@ -4,6 +4,7 @@ using EpicGames.Core;
 using EpicGames.Horde.Jobs.Templates;
 using EpicGames.Horde.Streams;
 using HordeServer.Configuration;
+using HordeServer.Streams;
 using HordeServer.Utilities;
 using System.Diagnostics.CodeAnalysis;
 using System.Text.RegularExpressions;
@@ -105,6 +106,11 @@ namespace HordeServer
 		/// </summary>
 		public List<StreamId> Streams { get; set; } = new List<StreamId>();
 
+		/// <summary>
+		/// List of enabled tags for a stream for this notification
+		/// </summary>
+		public List<string> EnabledStreamTags { get; set; } = new List<string>();
+
 		/// <summary>
 		/// Configuration for sending experimental Job Summary Slack notifications
 		/// </summary>
@@ -167,6 +173,7 @@ namespace HordeServer
 		public List<ConfigInclude> Include { get; set; } = new List<ConfigInclude>();
 
 		readonly Dictionary<StreamId, HashSet<NotificationStreamConfig>> _streamLookup = new Dictionary<StreamId, HashSet<NotificationStreamConfig>>();
+		readonly Dictionary<string, HashSet<NotificationStreamConfig>> _streamTagLookup = new Dictionary<string, HashSet<NotificationStreamConfig>>();
 
 		/// <summary>
 		/// Called after the store has been deserialized to compute cached values
@@ -189,6 +196,18 @@ namespace HordeServer
 					configs.Add(notificationStream);
 					_streamLookup[streamId] = configs;
 				}
+
+				foreach (string streamTag in notificationStream.EnabledStreamTags)
+				{
+					HashSet<NotificationStreamConfig>? configs;
+					if (!_streamTagLookup.TryGetValue(streamTag, out configs))
+					{
+						configs = new HashSet<NotificationStreamConfig>();
+					}
+
+					configs.Add(notificationStream);
+					_streamTagLookup[streamTag] = configs;
+				}
 			}
 		}
 
@@ -227,5 +246,44 @@ namespace HordeServer
 
 			return jobNotificationConfigs.Count != 0;
 		}
+
+		/// <summary>
+		/// Attempt to get config for a given list of stream tags and template
+		/// </summary>
+		/// <param name="streamTags">List of tags for a stream</param>
+		/// <param name="templateId">Template id</param>
+		/// <param name="jobNotificationConfigs">Receives the list of config objects on success</param>
+		/// <param name="includeDisabledConfigs">Flag on whether or not disabled configs should be included</param>
+		/// <returns>True if the job notification configs were found</returns>
+		public bool TryGetJobNofications(List<StreamTag> streamTags, TemplateId templateId, [NotNullWhen(true)] out HashSet<JobNotificationConfig>? jobNotificationConfigs, bool includeDisabledConfigs = false)
+		{
+			jobNotificationConfigs = new HashSet<JobNotificationConfig>();
+
+			HashSet<NotificationStreamConfig>? streamConfigList;
+			foreach (StreamTag streamTag in streamTags)
+			{
+				if (!streamTag.Enabled || !_streamTagLookup.TryGetValue(streamTag.Name, out streamConfigList))
+				{
+					continue;
+				}
+
+				foreach (NotificationStreamConfig streamConfig in streamConfigList)
+				{
+					List<JobNotificationConfig>? notificationConfigList;
+					if (streamConfig.TryGetTemplateConfigs(templateId, out notificationConfigList))
+					{
+						foreach (JobNotificationConfig notificationConfig in notificationConfigList)
+						{
+							if (notificationConfig.Enabled || includeDisabledConfigs)
+							{
+								jobNotificationConfigs.Add(notificationConfig);
+							}
+						}
+					}
+				}
+			}
+
+			return jobNotificationConfigs.Count != 0;
+		}
 	}
 }
diff --git a/Engine/Source/Programs/Horde/Plugins/Experimental/HordeServer.Experimental/Notifications/SlackNotificationProcessor.cs b/Engine/Source/Programs/Horde/Plugins/Experimental/HordeServer.Experimental/Notifications/SlackNotificationProcessor.cs
--- a/Engine/Source/Programs/Horde/Plugins/Experimental/HordeServer.Experimental/Notifications/SlackNotificationProcessor.cs
+++ b/Engine/Source/Programs/Horde/Plugins/Experimental/HordeServer.Experimental/Notifications/SlackNotificationProcessor.cs
@@ -339,7 +339,8 @@ namespace HordeServer.Experimental.Notifications
 
 			// From this point forward the current job/step will be handled as a spawned job
 
-			string branch = GetBranchFromStream(_buildConfig.CurrentValue, parentJob.StreamId);
+			StreamConfig? streamConfig = GetStreamConfig(_buildConfig.CurrentValue, parentJob.StreamId);
+			string branch = GetBranchFromStream(streamConfig);
 
 			// Check if the Setup step failed to notify that the job can't be executed
 			bool wasSetupSuccess = (step.Outcome == JobStepOutcome.Success || step.Outcome == JobStepOutcome.Warnings);
@@ -361,7 +362,7 @@ namespace HordeServer.Experimental.Notifications
 			ExperimentalConfig currentExperimentalConfig = _experimentalConfig.CurrentValue;
 
 			HashSet<JobNotificationConfig>? jobNotificationConfigs;
-			if (!currentExperimentalConfig.TryGetJobNofications(parentJob.StreamId, parentJob.TemplateId, out jobNotificationConfigs))
+			if (!currentExperimentalConfig.TryGetJobNofications(parentJob.StreamId, parentJob.TemplateId, out jobNotificationConfigs, streamConfig?.StreamTags ?? null))
 			{
 				_logger.LogInformation("Parent configurations not found for parent step '{ParentJobStep}' in parent job '{ParentJob}' for '{JobId}' in stream '{StreamId}' and for template '{TemplateId}'", job.ParentJobStepId, job.ParentJobId, job.Id, parentJob.StreamId, parentJob.TemplateId);
 				return true;
@@ -393,11 +394,13 @@ namespace HordeServer.Experimental.Notifications
 		/// <param name="cancellationToken">Cancellation token for the operation</param>
 		async Task ProcessJobSetupNotificationAsync(IJob job, IJobStep step, List<LogData> jobStepEventData, CancellationToken cancellationToken = default)
 		{
+			StreamConfig? streamConfig = GetStreamConfig(_buildConfig.CurrentValue, job.StreamId);
+
 			// Fetch our notification configurations for the parent job to append our current job's steps
 			ExperimentalConfig currentExperimentalConfig = _experimentalConfig.CurrentValue;
 
 			HashSet<JobNotificationConfig>? jobNotificationConfigs;
-			if (!currentExperimentalConfig.TryGetJobNofications(job.StreamId, job.TemplateId, out jobNotificationConfigs))
+			if (!currentExperimentalConfig.TryGetJobNofications(job.StreamId, job.TemplateId, out jobNotificationConfigs, streamConfig?.StreamTags ?? null))
 			{
 				_logger.LogInformation("No configurations found for '{JobId}' in stream '{StreamId}' and for template '{TemplateId}'", job.Id, job.StreamId, job.TemplateId);
 				return;
@@ -733,6 +736,28 @@ namespace HordeServer.Experimental.Notifications
 			}
 		}
 
+		/// <summary>
+		/// Fetch the latest stream config
+		/// </summary>
+		/// <param name="buildConfig">Build plugin's global config</param>
+		/// <param name="streamId">Identifier of the stream</param>
+		/// <returns>Latest <see cref="StreamConfig"/> or null if one is not found</returns>
+		static StreamConfig? GetStreamConfig(BuildConfig buildConfig, StreamId streamId)
+		{
+			buildConfig.TryGetStream(streamId, out StreamConfig? streamConfig);
+			return streamConfig;
+		}
+
+		/// <summary>
+		/// Fetch the branch name from the stream
+		/// </summary>
+		/// <param name="streamConfig">Config for the current stream</param>
+		/// <returns>Name of the corresponding branch or 'Unknown Branch' if no stream was found</returns>
+		static string GetBranchFromStream(StreamConfig? streamConfig)
+		{
+			return streamConfig?.Name ?? "Unknown Branch";
+		}
+
 		/// <summary>
 		/// Fetch the branch name from the stream
 		/// </summary>
@@ -741,12 +766,9 @@ namespace HordeServer.Experimental.Notifications
 		/// <returns>Name of the corresponding branch or 'Unknown Branch' if no stream was found</returns>
 		static string GetBranchFromStream(BuildConfig buildConfig, StreamId streamId)
 		{
-			if (buildConfig.TryGetStream(streamId, out StreamConfig? streamConfig))
-			{
-				return streamConfig.Name;
-			}
+			StreamConfig? streamConfig = GetStreamConfig(buildConfig, streamId);
 
-			return "Unknown Branch";
+			return GetBranchFromStream(streamConfig);
 		}
 
 		/// <summary>
