Addressing comments caught post submission

#rnx

diff --git a/Engine/Source/Programs/Horde/Plugins/Experimental/HordeServer.Experimental/Notifications/NotificationConfig.cs b/Engine/Source/Programs/Horde/Plugins/Experimental/HordeServer.Experimental/Notifications/NotificationConfig.cs
--- a/Engine/Source/Programs/Horde/Plugins/Experimental/HordeServer.Experimental/Notifications/NotificationConfig.cs
+++ b/Engine/Source/Programs/Horde/Plugins/Experimental/HordeServer.Experimental/Notifications/NotificationConfig.cs
@@ -126,14 +126,13 @@ namespace HordeServer
 			_notificationLookup.Clear();
 			foreach (JobNotificationConfig notificationConfig in JobNotifications)
 			{
-				List<JobNotificationConfig>? configs;
-				if (!_notificationLookup.TryGetValue(notificationConfig.Template, out configs))
+				if (!_notificationLookup.TryGetValue(notificationConfig.Template, out List<JobNotificationConfig>? configs))
 				{
 					configs = new List<JobNotificationConfig>();
+					_notificationLookup[notificationConfig.Template] = configs;
 				}
 
 				configs.Add(notificationConfig);
-				_notificationLookup[notificationConfig.Template] = configs;
 			}
 		}
 
@@ -187,26 +186,24 @@ namespace HordeServer
 
 				foreach (StreamId streamId in notificationStream.Streams)
 				{
-					HashSet<NotificationStreamConfig>? configs;
-					if (!_streamLookup.TryGetValue(streamId, out configs))
+					if (!_streamLookup.TryGetValue(streamId, out HashSet<NotificationStreamConfig>? configs))
 					{
 						configs = new HashSet<NotificationStreamConfig>();
+						_streamLookup[streamId] = configs;
 					}
 
 					configs.Add(notificationStream);
-					_streamLookup[streamId] = configs;
 				}
 
 				foreach (string streamTag in notificationStream.EnabledStreamTags)
 				{
-					HashSet<NotificationStreamConfig>? configs;
-					if (!_streamTagLookup.TryGetValue(streamTag, out configs))
+					if (!_streamTagLookup.TryGetValue(streamTag, out HashSet<NotificationStreamConfig>? configs))
 					{
 						configs = new HashSet<NotificationStreamConfig>();
+						_streamTagLookup[streamTag] = configs;
 					}
 
 					configs.Add(notificationStream);
-					_streamTagLookup[streamTag] = configs;
 				}
 			}
 		}
