Horde Dashboard: Add Epic copyright notice to files missing it

diff --git a/Engine/Source/Programs/Horde/HordeDashboard/plugins/analytics/api.ts b/Engine/Source/Programs/Horde/HordeDashboard/plugins/analytics/api.ts
--- a/Engine/Source/Programs/Horde/HordeDashboard/plugins/analytics/api.ts
+++ b/Engine/Source/Programs/Horde/HordeDashboard/plugins/analytics/api.ts
@@ -1,3 +1,5 @@
+// Copyright Epic Games, Inc. All Rights Reserved.
+
 import backend from "horde/backend";
 
 export type MetricsQuery = {
diff --git a/Engine/Source/Programs/Horde/HordeDashboard/plugins/analytics/index.tsx b/Engine/Source/Programs/Horde/HordeDashboard/plugins/analytics/index.tsx
--- a/Engine/Source/Programs/Horde/HordeDashboard/plugins/analytics/index.tsx
+++ b/Engine/Source/Programs/Horde/HordeDashboard/plugins/analytics/index.tsx
@@ -1,3 +1,5 @@
+// Copyright Epic Games, Inc. All Rights Reserved.
+
 import { registerHordePlugin, MountType } from "..";
 import { TelemetryView } from "./telemetryView";
 
diff --git a/Engine/Source/Programs/Horde/HordeDashboard/plugins/analytics/telemetryData.ts b/Engine/Source/Programs/Horde/HordeDashboard/plugins/analytics/telemetryData.ts
--- a/Engine/Source/Programs/Horde/HordeDashboard/plugins/analytics/telemetryData.ts
+++ b/Engine/Source/Programs/Horde/HordeDashboard/plugins/analytics/telemetryData.ts
@@ -1,3 +1,5 @@
+// Copyright Epic Games, Inc. All Rights Reserved.
+
 import { getMetrics, GetTelemetryMetricsResponse, GetTelemetryViewResponse } from "./api";
 
 export type MetricId = string;
diff --git a/Engine/Source/Programs/Horde/HordeDashboard/plugins/analytics/telemetryLineGraph.ts b/Engine/Source/Programs/Horde/HordeDashboard/plugins/analytics/telemetryLineGraph.ts
--- a/Engine/Source/Programs/Horde/HordeDashboard/plugins/analytics/telemetryLineGraph.ts
+++ b/Engine/Source/Programs/Horde/HordeDashboard/plugins/analytics/telemetryLineGraph.ts
@@ -1,3 +1,5 @@
+// Copyright Epic Games, Inc. All Rights Reserved.
+
 import * as d3 from "d3";
 import moment from "moment";
 import { GetTelemetryMetricResponse, GetTelemetryMetricsResponse, GetTelemetryChartResponse } from "./api";
diff --git a/Engine/Source/Programs/Horde/HordeDashboard/plugins/analytics/telemetryView.tsx b/Engine/Source/Programs/Horde/HordeDashboard/plugins/analytics/telemetryView.tsx
--- a/Engine/Source/Programs/Horde/HordeDashboard/plugins/analytics/telemetryView.tsx
+++ b/Engine/Source/Programs/Horde/HordeDashboard/plugins/analytics/telemetryView.tsx
@@ -1,3 +1,5 @@
+// Copyright Epic Games, Inc. All Rights Reserved.
+
 import { ComboBox, DefaultButton, DirectionalHint, FontIcon, IComboBox, IComboBoxOption, IComboBoxStyles, ITooltipHostStyles, ITooltipProps, Icon, Pivot, PivotItem, SelectableOptionMenuItemType, Spinner, SpinnerSize, Stack, Text, TooltipHost, mergeStyleSets, mergeStyles } from "@fluentui/react";
 import { useConst } from '@fluentui/react-hooks';
 import dashboard, { StatusColor } from "horde/backend/Dashboard";
diff --git a/Engine/Source/Programs/Horde/HordeDashboard/plugins/index.ts b/Engine/Source/Programs/Horde/HordeDashboard/plugins/index.ts
--- a/Engine/Source/Programs/Horde/HordeDashboard/plugins/index.ts
+++ b/Engine/Source/Programs/Horde/HordeDashboard/plugins/index.ts
@@ -1,3 +1,5 @@
+// Copyright Epic Games, Inc. All Rights Reserved.
+
 import { RouteObject } from "react-router-dom";
 
 // Routes => components
diff --git a/Engine/Source/Programs/Horde/HordeDashboard/plugins/registry.ts b/Engine/Source/Programs/Horde/HordeDashboard/plugins/registry.ts
--- a/Engine/Source/Programs/Horde/HordeDashboard/plugins/registry.ts
+++ b/Engine/Source/Programs/Horde/HordeDashboard/plugins/registry.ts
@@ -1,3 +1,4 @@
+// Copyright Epic Games, Inc. All Rights Reserved.
 
 import "./analytics"
 
diff --git a/Engine/Source/Programs/Horde/HordeDashboard/plugins/test-automation/api.ts b/Engine/Source/Programs/Horde/HordeDashboard/plugins/test-automation/api.ts
--- a/Engine/Source/Programs/Horde/HordeDashboard/plugins/test-automation/api.ts
+++ b/Engine/Source/Programs/Horde/HordeDashboard/plugins/test-automation/api.ts
@@ -1,3 +1,5 @@
+// Copyright Epic Games, Inc. All Rights Reserved.
+
 import backend from "horde/backend";
 
 /// A test data response
diff --git a/Engine/Source/Programs/Horde/HordeDashboard/plugins/test-automation/index.tsx b/Engine/Source/Programs/Horde/HordeDashboard/plugins/test-automation/index.tsx
--- a/Engine/Source/Programs/Horde/HordeDashboard/plugins/test-automation/index.tsx
+++ b/Engine/Source/Programs/Horde/HordeDashboard/plugins/test-automation/index.tsx
@@ -1,3 +1,5 @@
+// Copyright Epic Games, Inc. All Rights Reserved.
+
 import { registerHordePlugin, MountType } from "..";
 import { TestAutomationView } from "./testAutomationView";
 
diff --git a/Engine/Source/Programs/Horde/HordeDashboard/plugins/test-automation/phaseCompare.ts b/Engine/Source/Programs/Horde/HordeDashboard/plugins/test-automation/phaseCompare.ts
--- a/Engine/Source/Programs/Horde/HordeDashboard/plugins/test-automation/phaseCompare.ts
+++ b/Engine/Source/Programs/Horde/HordeDashboard/plugins/test-automation/phaseCompare.ts
@@ -1,3 +1,5 @@
+// Copyright Epic Games, Inc. All Rights Reserved.
+
 import { action, makeObservable, observable } from "mobx";
 import { TestDataHandler, TestPhaseStatus, TestSessionDetails } from "./testData";
 import { JobParameters } from "horde/components/ChangeButton";
diff --git a/Engine/Source/Programs/Horde/HordeDashboard/plugins/test-automation/phaseHistoryGraph.ts b/Engine/Source/Programs/Horde/HordeDashboard/plugins/test-automation/phaseHistoryGraph.ts
--- a/Engine/Source/Programs/Horde/HordeDashboard/plugins/test-automation/phaseHistoryGraph.ts
+++ b/Engine/Source/Programs/Horde/HordeDashboard/plugins/test-automation/phaseHistoryGraph.ts
@@ -1,3 +1,5 @@
+// Copyright Epic Games, Inc. All Rights Reserved.
+
 import * as d3 from "d3";
 import { TestDataHandler, TestNameRef, MetadataRef, PhaseSessionResult, TestMetaStatus } from './testData';
 import dashboard, { StatusColor } from 'horde/backend/Dashboard';
diff --git a/Engine/Source/Programs/Horde/HordeDashboard/plugins/test-automation/testAutomationCommon.tsx b/Engine/Source/Programs/Horde/HordeDashboard/plugins/test-automation/testAutomationCommon.tsx
--- a/Engine/Source/Programs/Horde/HordeDashboard/plugins/test-automation/testAutomationCommon.tsx
+++ b/Engine/Source/Programs/Horde/HordeDashboard/plugins/test-automation/testAutomationCommon.tsx
@@ -1,3 +1,5 @@
+// Copyright Epic Games, Inc. All Rights Reserved.
+
 import { mergeStyles, mergeStyleSets, Text, ComboBox, IComboBoxOption, SelectableOptionMenuItemType, IContextualMenuProps, IContextualMenuItem, IconButton, IContextualMenuListProps, Stack, Checkbox, FontIcon, Modal, DefaultButton, ScrollablePane, SelectionZone, DetailsList, Selection, DetailsListLayoutMode, SelectionMode, IColumn, PrimaryButton } from "@fluentui/react";
 import { MetadataRef, PhaseSessionResult, TestDataHandler, TestNameRef, TestPhaseStatus, TestSessionDetails, TestSessionResult } from "./testData";
 import dashboard, { StatusColor } from "horde/backend/Dashboard";
diff --git a/Engine/Source/Programs/Horde/HordeDashboard/plugins/test-automation/testAutomationEvents.tsx b/Engine/Source/Programs/Horde/HordeDashboard/plugins/test-automation/testAutomationEvents.tsx
--- a/Engine/Source/Programs/Horde/HordeDashboard/plugins/test-automation/testAutomationEvents.tsx
+++ b/Engine/Source/Programs/Horde/HordeDashboard/plugins/test-automation/testAutomationEvents.tsx
@@ -1,3 +1,5 @@
+// Copyright Epic Games, Inc. All Rights Reserved.
+
 import { Checkbox, DefaultButton, DirectionalHint, FontIcon, HoverCard, HoverCardType, Image, IPlainCardProps, Link, mergeStyles, mergeStyleSets, Spinner, SpinnerSize, Stack, Text } from "@fluentui/react";
 import { Artifacts, EmbeddedHTML, EventEntry, EventLevel, ImageCompare, LogPropertyType, URLLink } from "./testEventsModel";
 import { TestPhaseOutcome } from "./api";
diff --git a/Engine/Source/Programs/Horde/HordeDashboard/plugins/test-automation/testAutomationPhaseCompare.tsx b/Engine/Source/Programs/Horde/HordeDashboard/plugins/test-automation/testAutomationPhaseCompare.tsx
--- a/Engine/Source/Programs/Horde/HordeDashboard/plugins/test-automation/testAutomationPhaseCompare.tsx
+++ b/Engine/Source/Programs/Horde/HordeDashboard/plugins/test-automation/testAutomationPhaseCompare.tsx
@@ -1,3 +1,5 @@
+// Copyright Epic Games, Inc. All Rights Reserved.
+
 import { getHordeTheme } from "horde/styles/theme";
 import { getJobParams, PhaseComparisonContext } from "./phaseCompare";
 import { MetadataRef, PhaseSessionResult, TestDataHandler, TestNameRef, TestPhaseStatus, TestSessionStatus } from "./testData";
diff --git a/Engine/Source/Programs/Horde/HordeDashboard/plugins/test-automation/testAutomationPhaseView.tsx b/Engine/Source/Programs/Horde/HordeDashboard/plugins/test-automation/testAutomationPhaseView.tsx
--- a/Engine/Source/Programs/Horde/HordeDashboard/plugins/test-automation/testAutomationPhaseView.tsx
+++ b/Engine/Source/Programs/Horde/HordeDashboard/plugins/test-automation/testAutomationPhaseView.tsx
@@ -1,3 +1,5 @@
+// Copyright Epic Games, Inc. All Rights Reserved.
+
 import { IconButton, Stack, Text, Spinner, SpinnerSize, Link, FontIcon, DefaultButton, TooltipHost, Icon, ITag, TagPicker, IPickerItemProps, Label, List } from "@fluentui/react";
 import { useState, useEffect, useCallback, useRef, memo, useMemo } from "react";
 import { observer } from "mobx-react-lite";
diff --git a/Engine/Source/Programs/Horde/HordeDashboard/plugins/test-automation/testAutomationSummary.tsx b/Engine/Source/Programs/Horde/HordeDashboard/plugins/test-automation/testAutomationSummary.tsx
--- a/Engine/Source/Programs/Horde/HordeDashboard/plugins/test-automation/testAutomationSummary.tsx
+++ b/Engine/Source/Programs/Horde/HordeDashboard/plugins/test-automation/testAutomationSummary.tsx
@@ -1,3 +1,5 @@
+// Copyright Epic Games, Inc. All Rights Reserved.
+
 import { DefaultButton, FontIcon, IconButton, mergeStyleSets, Modal, Stack, Text, HoverCard, HoverCardType, IPlainCardProps, DirectionalHint, Spinner, SpinnerSize, IContextualMenuItem, IContextualMenuProps, IComboBoxOption } from "@fluentui/react";
 import { useState, useEffect, useCallback, memo } from "react";
 import { observer } from "mobx-react-lite";
diff --git a/Engine/Source/Programs/Horde/HordeDashboard/plugins/test-automation/testAutomationView.tsx b/Engine/Source/Programs/Horde/HordeDashboard/plugins/test-automation/testAutomationView.tsx
--- a/Engine/Source/Programs/Horde/HordeDashboard/plugins/test-automation/testAutomationView.tsx
+++ b/Engine/Source/Programs/Horde/HordeDashboard/plugins/test-automation/testAutomationView.tsx
@@ -1,3 +1,5 @@
+// Copyright Epic Games, Inc. All Rights Reserved.
+
 import { Stack, FontIcon, IComboBoxOption, IContextualMenuItem, IContextualMenuProps, Text, Spinner, SpinnerSize, Label, DefaultButton, ITag, TagPicker, IPickerItemProps, PrimaryButton, Modal, IconButton } from "@fluentui/react";
 import { Breadcrumbs } from "horde/components/Breadcrumbs";
 import { TopNav } from "horde/components/TopNav";
diff --git a/Engine/Source/Programs/Horde/HordeDashboard/plugins/test-automation/testData.ts b/Engine/Source/Programs/Horde/HordeDashboard/plugins/test-automation/testData.ts
--- a/Engine/Source/Programs/Horde/HordeDashboard/plugins/test-automation/testData.ts
+++ b/Engine/Source/Programs/Horde/HordeDashboard/plugins/test-automation/testData.ts
@@ -1,3 +1,5 @@
+// Copyright Epic Games, Inc. All Rights Reserved.
+
 import { isNumber } from '@datadog/browser-core';
 import { action, makeObservable, observable, reaction, when } from 'mobx';
 import { JobStepTestDataItem, TestOutcome, TestPhaseOutcome } from './api';
diff --git a/Engine/Source/Programs/Horde/HordeDashboard/plugins/test-automation/testDataV2Fetcher.ts b/Engine/Source/Programs/Horde/HordeDashboard/plugins/test-automation/testDataV2Fetcher.ts
--- a/Engine/Source/Programs/Horde/HordeDashboard/plugins/test-automation/testDataV2Fetcher.ts
+++ b/Engine/Source/Programs/Horde/HordeDashboard/plugins/test-automation/testDataV2Fetcher.ts
@@ -1,3 +1,5 @@
+// Copyright Epic Games, Inc. All Rights Reserved.
+
 import { getJobStepTestDataV2, GetPhaseResponse, getStreamTestsV2, getTestDataV2, GetTestMetadataResponse, GetTestNameResponse, GetTestPhaseSessionResponse, getTestPhaseSessionsV2, getTestPhasesV2, GetTestSessionResponse, getTestSessionsV2, GetTestTagResponse, JobStepTestDataItem, PhaseSessionDetailsResponse, TestData, TestOutcome, TestSessionDetailsResponse } from "./api";
 import { TestStreamRef, TestStatus, TestSessionResult, TestNameRef, MetadataRef, TestTagRef, PhaseSessionResult, TestPhaseRef, TestMetaStatus, TestSessionDetails, TestPhaseStatus, DeviceRef, TestVersionFetcher } from "./testData"
 import { EventEntry } from "./testEventsModel";
diff --git a/Engine/Source/Programs/Horde/HordeDashboard/plugins/test-automation/testEventsModel.ts b/Engine/Source/Programs/Horde/HordeDashboard/plugins/test-automation/testEventsModel.ts
--- a/Engine/Source/Programs/Horde/HordeDashboard/plugins/test-automation/testEventsModel.ts
+++ b/Engine/Source/Programs/Horde/HordeDashboard/plugins/test-automation/testEventsModel.ts
@@ -1,3 +1,5 @@
+// Copyright Epic Games, Inc. All Rights Reserved.
+
 /// An event entry
 export type EventEntry = {
 
diff --git a/Engine/Source/Programs/Horde/HordeDashboard/plugins/test-automation/testHistoryGraph.ts b/Engine/Source/Programs/Horde/HordeDashboard/plugins/test-automation/testHistoryGraph.ts
--- a/Engine/Source/Programs/Horde/HordeDashboard/plugins/test-automation/testHistoryGraph.ts
+++ b/Engine/Source/Programs/Horde/HordeDashboard/plugins/test-automation/testHistoryGraph.ts
@@ -1,3 +1,5 @@
+// Copyright Epic Games, Inc. All Rights Reserved.
+
 import * as d3 from "d3";
 import { TestDataHandler, TestNameRef, MetadataRef, TestMetaStatus, TestSessionResult } from './testData';
 import dashboard, { StatusColor } from 'horde/backend/Dashboard';
diff --git a/Engine/Source/Programs/Horde/HordeDashboard/src/App.css b/Engine/Source/Programs/Horde/HordeDashboard/src/App.css
--- a/Engine/Source/Programs/Horde/HordeDashboard/src/App.css
+++ b/Engine/Source/Programs/Horde/HordeDashboard/src/App.css
@@ -1,3 +1,5 @@
+/* Copyright Epic Games, Inc. All Rights Reserved. */
+
 .App {
   text-align: center;
 }
diff --git a/Engine/Source/Programs/Horde/HordeDashboard/src/backend/AutomationTestData.ts b/Engine/Source/Programs/Horde/HordeDashboard/src/backend/AutomationTestData.ts
--- a/Engine/Source/Programs/Horde/HordeDashboard/src/backend/AutomationTestData.ts
+++ b/Engine/Source/Programs/Horde/HordeDashboard/src/backend/AutomationTestData.ts
@@ -1,3 +1,5 @@
+// Copyright Epic Games, Inc. All Rights Reserved.
+
 import { isNumber } from '@datadog/browser-core';
 import { action, makeObservable, observable } from 'mobx';
 import backend from '../backend';
diff --git a/Engine/Source/Programs/Horde/HordeDashboard/src/backend/DeviceHandler.ts b/Engine/Source/Programs/Horde/HordeDashboard/src/backend/DeviceHandler.ts
--- a/Engine/Source/Programs/Horde/HordeDashboard/src/backend/DeviceHandler.ts
+++ b/Engine/Source/Programs/Horde/HordeDashboard/src/backend/DeviceHandler.ts
@@ -1,3 +1,5 @@
+// Copyright Epic Games, Inc. All Rights Reserved.
+
 import moment from "moment";
 import backend from "../backend";
 import { GetDevicePlatformResponse, GetDevicePoolResponse, GetDeviceReservationResponse, GetDevicePoolTelemetryResponse, GetDevicePlatformTelemetryResponse, GetDeviceResponse, GetUserResponse, DevicePoolTelemetryQuery } from "../backend/Api";
diff --git a/Engine/Source/Programs/Horde/HordeDashboard/src/backend/DeviceTelemetryHandler.ts b/Engine/Source/Programs/Horde/HordeDashboard/src/backend/DeviceTelemetryHandler.ts
--- a/Engine/Source/Programs/Horde/HordeDashboard/src/backend/DeviceTelemetryHandler.ts
+++ b/Engine/Source/Programs/Horde/HordeDashboard/src/backend/DeviceTelemetryHandler.ts
@@ -1,3 +1,5 @@
+// Copyright Epic Games, Inc. All Rights Reserved.
+
 import { action, makeObservable, observable } from "mobx";
 import moment from "moment";
 import backend from "../backend";
diff --git a/Engine/Source/Programs/Horde/HordeDashboard/src/backend/Notices.ts b/Engine/Source/Programs/Horde/HordeDashboard/src/backend/Notices.ts
--- a/Engine/Source/Programs/Horde/HordeDashboard/src/backend/Notices.ts
+++ b/Engine/Source/Programs/Horde/HordeDashboard/src/backend/Notices.ts
@@ -1,3 +1,5 @@
+// Copyright Epic Games, Inc. All Rights Reserved.
+
 import moment from "moment";
 import backend from ".";
 import { displayTimeZone } from "../base/utilities/timeUtils";
diff --git a/Engine/Source/Programs/Horde/HordeDashboard/src/backend/UserCache.ts b/Engine/Source/Programs/Horde/HordeDashboard/src/backend/UserCache.ts
--- a/Engine/Source/Programs/Horde/HordeDashboard/src/backend/UserCache.ts
+++ b/Engine/Source/Programs/Horde/HordeDashboard/src/backend/UserCache.ts
@@ -1,3 +1,4 @@
+// Copyright Epic Games, Inc. All Rights Reserved.
 
 import { action, makeObservable, observable } from 'mobx';
 import backend from '.';
diff --git a/Engine/Source/Programs/Horde/HordeDashboard/src/backend/UserInactivity.ts b/Engine/Source/Programs/Horde/HordeDashboard/src/backend/UserInactivity.ts
--- a/Engine/Source/Programs/Horde/HordeDashboard/src/backend/UserInactivity.ts
+++ b/Engine/Source/Programs/Horde/HordeDashboard/src/backend/UserInactivity.ts
@@ -1,3 +1,5 @@
+// Copyright Epic Games, Inc. All Rights Reserved.
+
 /**
  * Tracks a timestamp for when the user last performed a UI activity
  */
diff --git a/Engine/Source/Programs/Horde/HordeDashboard/src/base/components/Markdown/Markdown.tsx b/Engine/Source/Programs/Horde/HordeDashboard/src/base/components/Markdown/Markdown.tsx
--- a/Engine/Source/Programs/Horde/HordeDashboard/src/base/components/Markdown/Markdown.tsx
+++ b/Engine/Source/Programs/Horde/HordeDashboard/src/base/components/Markdown/Markdown.tsx
@@ -1,3 +1,5 @@
+// Copyright Epic Games, Inc. All Rights Reserved.
+
 import * as React from 'react';
 import type MarkdownComponentType from 'markdown-to-jsx';
 import type { MarkdownToJSX } from 'markdown-to-jsx';
diff --git a/Engine/Source/Programs/Horde/HordeDashboard/src/base/components/Markdown/Markdown.types.ts b/Engine/Source/Programs/Horde/HordeDashboard/src/base/components/Markdown/Markdown.types.ts
--- a/Engine/Source/Programs/Horde/HordeDashboard/src/base/components/Markdown/Markdown.types.ts
+++ b/Engine/Source/Programs/Horde/HordeDashboard/src/base/components/Markdown/Markdown.types.ts
@@ -1,3 +1,5 @@
+// Copyright Epic Games, Inc. All Rights Reserved.
+
 import type { ITheme, IStyleFunctionOrObject, IStyle, ILinkStyleProps, IImageStyleProps } from '@fluentui/react';
 import type { IMarkdownHeaderStyleProps } from './MarkdownHeader';
 import type { IMarkdownParagraphStyleProps } from './MarkdownParagraph';
diff --git a/Engine/Source/Programs/Horde/HordeDashboard/src/base/components/Markdown/MarkdownHeader.tsx b/Engine/Source/Programs/Horde/HordeDashboard/src/base/components/Markdown/MarkdownHeader.tsx
--- a/Engine/Source/Programs/Horde/HordeDashboard/src/base/components/Markdown/MarkdownHeader.tsx
+++ b/Engine/Source/Programs/Horde/HordeDashboard/src/base/components/Markdown/MarkdownHeader.tsx
@@ -1,3 +1,5 @@
+// Copyright Epic Games, Inc. All Rights Reserved.
+
 import * as React from 'react';
 import { classNamesFunction, IStyleFunctionOrObject, styled, IStyleFunction } from '@fluentui/react/lib/Utilities';
 import { ITheme, IStyle } from '@fluentui/react/lib/Styling';
diff --git a/Engine/Source/Programs/Horde/HordeDashboard/src/base/components/Markdown/MarkdownImage.tsx b/Engine/Source/Programs/Horde/HordeDashboard/src/base/components/Markdown/MarkdownImage.tsx
--- a/Engine/Source/Programs/Horde/HordeDashboard/src/base/components/Markdown/MarkdownImage.tsx
+++ b/Engine/Source/Programs/Horde/HordeDashboard/src/base/components/Markdown/MarkdownImage.tsx
@@ -1,3 +1,5 @@
+// Copyright Epic Games, Inc. All Rights Reserved.
+
 import { IImageProps, Image } from "@fluentui/react"
 import dashboard from "../../../backend/Dashboard";
 
diff --git a/Engine/Source/Programs/Horde/HordeDashboard/src/base/components/Markdown/MarkdownLink.tsx b/Engine/Source/Programs/Horde/HordeDashboard/src/base/components/Markdown/MarkdownLink.tsx
--- a/Engine/Source/Programs/Horde/HordeDashboard/src/base/components/Markdown/MarkdownLink.tsx
+++ b/Engine/Source/Programs/Horde/HordeDashboard/src/base/components/Markdown/MarkdownLink.tsx
@@ -1,3 +1,5 @@
+// Copyright Epic Games, Inc. All Rights Reserved.
+
 import * as React from 'react';
 import { Link } from 'react-router-dom';
 import { ILinkProps } from '@fluentui/react/lib/Link';
diff --git a/Engine/Source/Programs/Horde/HordeDashboard/src/base/components/Markdown/MarkdownParagraph.tsx b/Engine/Source/Programs/Horde/HordeDashboard/src/base/components/Markdown/MarkdownParagraph.tsx
--- a/Engine/Source/Programs/Horde/HordeDashboard/src/base/components/Markdown/MarkdownParagraph.tsx
+++ b/Engine/Source/Programs/Horde/HordeDashboard/src/base/components/Markdown/MarkdownParagraph.tsx
@@ -1,3 +1,5 @@
+// Copyright Epic Games, Inc. All Rights Reserved.
+
 import * as React from 'react';
 import {
    IClassNames,
diff --git a/Engine/Source/Programs/Horde/HordeDashboard/src/base/components/Markdown/MarkdownPre.tsx b/Engine/Source/Programs/Horde/HordeDashboard/src/base/components/Markdown/MarkdownPre.tsx
--- a/Engine/Source/Programs/Horde/HordeDashboard/src/base/components/Markdown/MarkdownPre.tsx
+++ b/Engine/Source/Programs/Horde/HordeDashboard/src/base/components/Markdown/MarkdownPre.tsx
@@ -1,3 +1,5 @@
+// Copyright Epic Games, Inc. All Rights Reserved.
+
 import * as React from 'react';
 import dashboard from '../../../backend/Dashboard';
 
diff --git a/Engine/Source/Programs/Horde/HordeDashboard/src/base/components/Markdown/index.ts b/Engine/Source/Programs/Horde/HordeDashboard/src/base/components/Markdown/index.ts
--- a/Engine/Source/Programs/Horde/HordeDashboard/src/base/components/Markdown/index.ts
+++ b/Engine/Source/Programs/Horde/HordeDashboard/src/base/components/Markdown/index.ts
@@ -1,3 +1,5 @@
+// Copyright Epic Games, Inc. All Rights Reserved.
+
 export * from './Markdown.types';
 export * from './Markdown';
 export * from './MarkdownHeader';
diff --git a/Engine/Source/Programs/Horde/HordeDashboard/src/base/components/MarkdownTable/MarkdownCell.tsx b/Engine/Source/Programs/Horde/HordeDashboard/src/base/components/MarkdownTable/MarkdownCell.tsx
--- a/Engine/Source/Programs/Horde/HordeDashboard/src/base/components/MarkdownTable/MarkdownCell.tsx
+++ b/Engine/Source/Programs/Horde/HordeDashboard/src/base/components/MarkdownTable/MarkdownCell.tsx
@@ -1,3 +1,5 @@
+// Copyright Epic Games, Inc. All Rights Reserved.
+
 import * as React from 'react';
 import { classNamesFunction, styled } from '@fluentui/react';
 import { IMarkdownTableCellProps, IMarkdownTableStyleProps, IMarkdownTableStyles } from './MarkdownTable.types';
diff --git a/Engine/Source/Programs/Horde/HordeDashboard/src/base/components/MarkdownTable/MarkdownTBody.tsx b/Engine/Source/Programs/Horde/HordeDashboard/src/base/components/MarkdownTable/MarkdownTBody.tsx
--- a/Engine/Source/Programs/Horde/HordeDashboard/src/base/components/MarkdownTable/MarkdownTBody.tsx
+++ b/Engine/Source/Programs/Horde/HordeDashboard/src/base/components/MarkdownTable/MarkdownTBody.tsx
@@ -1,3 +1,5 @@
+// Copyright Epic Games, Inc. All Rights Reserved.
+
 import * as React from 'react';
 import { classNamesFunction, styled } from '@fluentui/react';
 import { IMarkdownTableProps, IMarkdownTableStyleProps, IMarkdownTableStyles } from './MarkdownTable.types';
diff --git a/Engine/Source/Programs/Horde/HordeDashboard/src/base/components/MarkdownTable/MarkdownTHead.tsx b/Engine/Source/Programs/Horde/HordeDashboard/src/base/components/MarkdownTable/MarkdownTHead.tsx
--- a/Engine/Source/Programs/Horde/HordeDashboard/src/base/components/MarkdownTable/MarkdownTHead.tsx
+++ b/Engine/Source/Programs/Horde/HordeDashboard/src/base/components/MarkdownTable/MarkdownTHead.tsx
@@ -1,3 +1,5 @@
+// Copyright Epic Games, Inc. All Rights Reserved.
+
 import * as React from 'react';
 import { classNamesFunction, styled } from '@fluentui/react';
 import { IMarkdownTableProps, IMarkdownTableStyleProps, IMarkdownTableStyles } from './MarkdownTable.types';
diff --git a/Engine/Source/Programs/Horde/HordeDashboard/src/base/components/MarkdownTable/MarkdownTable.styles.ts b/Engine/Source/Programs/Horde/HordeDashboard/src/base/components/MarkdownTable/MarkdownTable.styles.ts
--- a/Engine/Source/Programs/Horde/HordeDashboard/src/base/components/MarkdownTable/MarkdownTable.styles.ts
+++ b/Engine/Source/Programs/Horde/HordeDashboard/src/base/components/MarkdownTable/MarkdownTable.styles.ts
@@ -1,3 +1,5 @@
+// Copyright Epic Games, Inc. All Rights Reserved.
+
 import { getGlobalClassNames, IStyle } from '@fluentui/react';
 import { IMarkdownTableStyleProps, IMarkdownTableStyles } from './MarkdownTable.types';
 
diff --git a/Engine/Source/Programs/Horde/HordeDashboard/src/base/components/MarkdownTable/MarkdownTable.tsx b/Engine/Source/Programs/Horde/HordeDashboard/src/base/components/MarkdownTable/MarkdownTable.tsx
--- a/Engine/Source/Programs/Horde/HordeDashboard/src/base/components/MarkdownTable/MarkdownTable.tsx
+++ b/Engine/Source/Programs/Horde/HordeDashboard/src/base/components/MarkdownTable/MarkdownTable.tsx
@@ -1,3 +1,5 @@
+// Copyright Epic Games, Inc. All Rights Reserved.
+
 import * as React from 'react';
 import { classNamesFunction, styled } from '@fluentui/react';
 import { IMarkdownTableProps, IMarkdownTableStyleProps, IMarkdownTableStyles } from './MarkdownTable.types';
diff --git a/Engine/Source/Programs/Horde/HordeDashboard/src/base/components/MarkdownTable/MarkdownTable.types.ts b/Engine/Source/Programs/Horde/HordeDashboard/src/base/components/MarkdownTable/MarkdownTable.types.ts
--- a/Engine/Source/Programs/Horde/HordeDashboard/src/base/components/MarkdownTable/MarkdownTable.types.ts
+++ b/Engine/Source/Programs/Horde/HordeDashboard/src/base/components/MarkdownTable/MarkdownTable.types.ts
@@ -1,3 +1,5 @@
+// Copyright Epic Games, Inc. All Rights Reserved.
+
 import { IStyle, ITheme, IStyleFunctionOrObject } from '@fluentui/react';
 
 export interface IMarkdownTableProps {
diff --git a/Engine/Source/Programs/Horde/HordeDashboard/src/base/components/MarkdownTable/MarkdownTr.tsx b/Engine/Source/Programs/Horde/HordeDashboard/src/base/components/MarkdownTable/MarkdownTr.tsx
--- a/Engine/Source/Programs/Horde/HordeDashboard/src/base/components/MarkdownTable/MarkdownTr.tsx
+++ b/Engine/Source/Programs/Horde/HordeDashboard/src/base/components/MarkdownTable/MarkdownTr.tsx
@@ -1,3 +1,5 @@
+// Copyright Epic Games, Inc. All Rights Reserved.
+
 import * as React from 'react';
 import { classNamesFunction, styled } from '@fluentui/react';
 import { IMarkdownTableProps, IMarkdownTableStyleProps, IMarkdownTableStyles } from './MarkdownTable.types';
diff --git a/Engine/Source/Programs/Horde/HordeDashboard/src/base/components/MarkdownTable/index.ts b/Engine/Source/Programs/Horde/HordeDashboard/src/base/components/MarkdownTable/index.ts
--- a/Engine/Source/Programs/Horde/HordeDashboard/src/base/components/MarkdownTable/index.ts
+++ b/Engine/Source/Programs/Horde/HordeDashboard/src/base/components/MarkdownTable/index.ts
@@ -1,3 +1,5 @@
+// Copyright Epic Games, Inc. All Rights Reserved.
+
 export * from './MarkdownCell';
 export * from './MarkdownTable.types';
 export * from './MarkdownTable';
diff --git a/Engine/Source/Programs/Horde/HordeDashboard/src/base/components/SideRail/SideRail.styles.ts b/Engine/Source/Programs/Horde/HordeDashboard/src/base/components/SideRail/SideRail.styles.ts
--- a/Engine/Source/Programs/Horde/HordeDashboard/src/base/components/SideRail/SideRail.styles.ts
+++ b/Engine/Source/Programs/Horde/HordeDashboard/src/base/components/SideRail/SideRail.styles.ts
@@ -1,3 +1,5 @@
+// Copyright Epic Games, Inc. All Rights Reserved.
+
 //import { FontWeights, getFocusOutlineStyle, IStyleFunction } from '@fluentui/react/lib/index';
 import { FontWeights, getFocusOutlineStyle, IStyleFunction } from '@fluentui/react';
 import { ISideRailStyleProps, ISideRailStyles } from './SideRail.types';
diff --git a/Engine/Source/Programs/Horde/HordeDashboard/src/base/components/SideRail/SideRail.tsx b/Engine/Source/Programs/Horde/HordeDashboard/src/base/components/SideRail/SideRail.tsx
--- a/Engine/Source/Programs/Horde/HordeDashboard/src/base/components/SideRail/SideRail.tsx
+++ b/Engine/Source/Programs/Horde/HordeDashboard/src/base/components/SideRail/SideRail.tsx
@@ -1,3 +1,5 @@
+// Copyright Epic Games, Inc. All Rights Reserved.
+
 import { classNamesFunction, css, FocusZone, FocusZoneDirection, IProcessedStyleSet, Link, Stack, styled, Text } from '@fluentui/react';
 import * as React from 'react';
 import { Link as ReactRouterLink } from 'react-router-dom';
diff --git a/Engine/Source/Programs/Horde/HordeDashboard/src/base/components/SideRail/SideRail.types.ts b/Engine/Source/Programs/Horde/HordeDashboard/src/base/components/SideRail/SideRail.types.ts
--- a/Engine/Source/Programs/Horde/HordeDashboard/src/base/components/SideRail/SideRail.types.ts
+++ b/Engine/Source/Programs/Horde/HordeDashboard/src/base/components/SideRail/SideRail.types.ts
@@ -1,3 +1,5 @@
+// Copyright Epic Games, Inc. All Rights Reserved.
+
 import { ITheme, IStyleFunctionOrObject, IStyle } from '@fluentui/react';
 
 export interface ISideRailProps {
diff --git a/Engine/Source/Programs/Horde/HordeDashboard/src/base/components/SideRail/index.ts b/Engine/Source/Programs/Horde/HordeDashboard/src/base/components/SideRail/index.ts
--- a/Engine/Source/Programs/Horde/HordeDashboard/src/base/components/SideRail/index.ts
+++ b/Engine/Source/Programs/Horde/HordeDashboard/src/base/components/SideRail/index.ts
@@ -1,2 +1,4 @@
+// Copyright Epic Games, Inc. All Rights Reserved.
+
 export * from './SideRail';
 export * from './SideRail.types';
diff --git a/Engine/Source/Programs/Horde/HordeDashboard/src/base/components/ThemeTester/ThemeTester.tsx b/Engine/Source/Programs/Horde/HordeDashboard/src/base/components/ThemeTester/ThemeTester.tsx
--- a/Engine/Source/Programs/Horde/HordeDashboard/src/base/components/ThemeTester/ThemeTester.tsx
+++ b/Engine/Source/Programs/Horde/HordeDashboard/src/base/components/ThemeTester/ThemeTester.tsx
@@ -1,3 +1,5 @@
+// Copyright Epic Games, Inc. All Rights Reserved.
+
 import { Checkbox, DefaultButton, IContextualMenuProps, ITextStyles, Pivot, PivotItem, PrimaryButton, SearchBox, SpinButton, Stack, Text, TextField, Toggle } from "@fluentui/react";
 import { useState } from "react";
 import { Breadcrumbs } from "../../../components/Breadcrumbs";
diff --git a/Engine/Source/Programs/Horde/HordeDashboard/src/base/components/ThemeTester/stickers/ComboBox.tsx b/Engine/Source/Programs/Horde/HordeDashboard/src/base/components/ThemeTester/stickers/ComboBox.tsx
--- a/Engine/Source/Programs/Horde/HordeDashboard/src/base/components/ThemeTester/stickers/ComboBox.tsx
+++ b/Engine/Source/Programs/Horde/HordeDashboard/src/base/components/ThemeTester/stickers/ComboBox.tsx
@@ -1,3 +1,5 @@
+// Copyright Epic Games, Inc. All Rights Reserved.
+
 import {
    ComboBox,
    IComboBoxOption,
diff --git a/Engine/Source/Programs/Horde/HordeDashboard/src/base/components/ThemeTester/stickers/CommandBar.tsx b/Engine/Source/Programs/Horde/HordeDashboard/src/base/components/ThemeTester/stickers/CommandBar.tsx
--- a/Engine/Source/Programs/Horde/HordeDashboard/src/base/components/ThemeTester/stickers/CommandBar.tsx
+++ b/Engine/Source/Programs/Horde/HordeDashboard/src/base/components/ThemeTester/stickers/CommandBar.tsx
@@ -1,3 +1,5 @@
+// Copyright Epic Games, Inc. All Rights Reserved.
+
 import { setVirtualParent } from '@fluentui/dom-utilities';
 import { IButtonProps } from '@fluentui/react/lib/Button';
 import { CommandBar, ICommandBarItemProps } from '@fluentui/react/lib/CommandBar';
diff --git a/Engine/Source/Programs/Horde/HordeDashboard/src/base/components/ThemeTester/stickers/ContextualMenu.tsx b/Engine/Source/Programs/Horde/HordeDashboard/src/base/components/ThemeTester/stickers/ContextualMenu.tsx
--- a/Engine/Source/Programs/Horde/HordeDashboard/src/base/components/ThemeTester/stickers/ContextualMenu.tsx
+++ b/Engine/Source/Programs/Horde/HordeDashboard/src/base/components/ThemeTester/stickers/ContextualMenu.tsx
@@ -1,4 +1,4 @@
-
+// Copyright Epic Games, Inc. All Rights Reserved.
 
 import * as React from 'react';
 import { ContextualMenuItemType, IContextualMenuProps } from '@fluentui/react/lib/ContextualMenu';
diff --git a/Engine/Source/Programs/Horde/HordeDashboard/src/base/components/ThemeTester/stickers/Dropdown.tsx b/Engine/Source/Programs/Horde/HordeDashboard/src/base/components/ThemeTester/stickers/Dropdown.tsx
--- a/Engine/Source/Programs/Horde/HordeDashboard/src/base/components/ThemeTester/stickers/Dropdown.tsx
+++ b/Engine/Source/Programs/Horde/HordeDashboard/src/base/components/ThemeTester/stickers/Dropdown.tsx
@@ -1,3 +1,5 @@
+// Copyright Epic Games, Inc. All Rights Reserved.
+
 import * as React from 'react';
 import { IStackTokens, Stack } from '@fluentui/react/lib/Stack';
 import { Dropdown, DropdownMenuItemType, IDropdownStyles, IDropdownOption } from '@fluentui/react/lib/Dropdown';
diff --git a/Engine/Source/Programs/Horde/HordeDashboard/src/base/components/ThemeTester/stickers/Progress.tsx b/Engine/Source/Programs/Horde/HordeDashboard/src/base/components/ThemeTester/stickers/Progress.tsx
--- a/Engine/Source/Programs/Horde/HordeDashboard/src/base/components/ThemeTester/stickers/Progress.tsx
+++ b/Engine/Source/Programs/Horde/HordeDashboard/src/base/components/ThemeTester/stickers/Progress.tsx
@@ -1,3 +1,5 @@
+// Copyright Epic Games, Inc. All Rights Reserved.
+
 import * as React from 'react';
 import { ProgressIndicator } from '@fluentui/react/lib/ProgressIndicator';
 
diff --git a/Engine/Source/Programs/Horde/HordeDashboard/src/base/components/ThemeTester/stickers/Slider.tsx b/Engine/Source/Programs/Horde/HordeDashboard/src/base/components/ThemeTester/stickers/Slider.tsx
--- a/Engine/Source/Programs/Horde/HordeDashboard/src/base/components/ThemeTester/stickers/Slider.tsx
+++ b/Engine/Source/Programs/Horde/HordeDashboard/src/base/components/ThemeTester/stickers/Slider.tsx
@@ -1,3 +1,5 @@
+// Copyright Epic Games, Inc. All Rights Reserved.
+
 import * as React from 'react';
 import { Slider, IStackTokens, Stack, IStackStyles } from '@fluentui/react';
 
diff --git a/Engine/Source/Programs/Horde/HordeDashboard/src/base/components/ThemeTester/stickers/Spinner.tsx b/Engine/Source/Programs/Horde/HordeDashboard/src/base/components/ThemeTester/stickers/Spinner.tsx
--- a/Engine/Source/Programs/Horde/HordeDashboard/src/base/components/ThemeTester/stickers/Spinner.tsx
+++ b/Engine/Source/Programs/Horde/HordeDashboard/src/base/components/ThemeTester/stickers/Spinner.tsx
@@ -1,3 +1,5 @@
+// Copyright Epic Games, Inc. All Rights Reserved.
+
 import * as React from 'react';
 import { Spinner, SpinnerSize } from '@fluentui/react/lib/Spinner';
 import { Label } from '@fluentui/react/lib/Label';
diff --git a/Engine/Source/Programs/Horde/HordeDashboard/src/base/utilities/clipboard.ts b/Engine/Source/Programs/Horde/HordeDashboard/src/base/utilities/clipboard.ts
--- a/Engine/Source/Programs/Horde/HordeDashboard/src/base/utilities/clipboard.ts
+++ b/Engine/Source/Programs/Horde/HordeDashboard/src/base/utilities/clipboard.ts
@@ -1,4 +1,4 @@
-
+// Copyright Epic Games, Inc. All Rights Reserved.
 
 export const copyToClipboard = (value: string | undefined) => {
 
diff --git a/Engine/Source/Programs/Horde/HordeDashboard/src/base/utilities/colors.ts b/Engine/Source/Programs/Horde/HordeDashboard/src/base/utilities/colors.ts
--- a/Engine/Source/Programs/Horde/HordeDashboard/src/base/utilities/colors.ts
+++ b/Engine/Source/Programs/Horde/HordeDashboard/src/base/utilities/colors.ts
@@ -1,4 +1,4 @@
-
+// Copyright Epic Games, Inc. All Rights Reserved.
 
 export const hexToRgb = (hex: string): string => {
 
diff --git a/Engine/Source/Programs/Horde/HordeDashboard/src/base/utilities/extractAnchorLink.ts b/Engine/Source/Programs/Horde/HordeDashboard/src/base/utilities/extractAnchorLink.ts
--- a/Engine/Source/Programs/Horde/HordeDashboard/src/base/utilities/extractAnchorLink.ts
+++ b/Engine/Source/Programs/Horde/HordeDashboard/src/base/utilities/extractAnchorLink.ts
@@ -1,3 +1,5 @@
+// Copyright Epic Games, Inc. All Rights Reserved.
+
 /**
  * Given a URL containing a route path (first hash) and possibly an anchor (second hash),
  * returns only the anchor part, if it exists. (Does not include the query string, if any.)
diff --git a/Engine/Source/Programs/Horde/HordeDashboard/src/base/utilities/filter.ts b/Engine/Source/Programs/Horde/HordeDashboard/src/base/utilities/filter.ts
--- a/Engine/Source/Programs/Horde/HordeDashboard/src/base/utilities/filter.ts
+++ b/Engine/Source/Programs/Horde/HordeDashboard/src/base/utilities/filter.ts
@@ -1,3 +1,4 @@
+// Copyright Epic Games, Inc. All Rights Reserved.
 
 import { GetJobResponse } from "../../backend/Api";
 
diff --git a/Engine/Source/Programs/Horde/HordeDashboard/src/base/utilities/hooks.ts b/Engine/Source/Programs/Horde/HordeDashboard/src/base/utilities/hooks.ts
--- a/Engine/Source/Programs/Horde/HordeDashboard/src/base/utilities/hooks.ts
+++ b/Engine/Source/Programs/Horde/HordeDashboard/src/base/utilities/hooks.ts
@@ -1,3 +1,5 @@
+// Copyright Epic Games, Inc. All Rights Reserved.
+
 import { useState, useEffect } from 'react';
 import { useSearchParams } from 'react-router-dom';
 import LZString from "lz-string"
diff --git a/Engine/Source/Programs/Horde/HordeDashboard/src/base/utilities/index.ts b/Engine/Source/Programs/Horde/HordeDashboard/src/base/utilities/index.ts
--- a/Engine/Source/Programs/Horde/HordeDashboard/src/base/utilities/index.ts
+++ b/Engine/Source/Programs/Horde/HordeDashboard/src/base/utilities/index.ts
@@ -1,3 +1,5 @@
+// Copyright Epic Games, Inc. All Rights Reserved.
+
 export * from './jumpToAnchor';
 export * from './removeAnchorLink';
 export * from './extractAnchorLink';
diff --git a/Engine/Source/Programs/Horde/HordeDashboard/src/base/utilities/issueUtils.ts b/Engine/Source/Programs/Horde/HordeDashboard/src/base/utilities/issueUtils.ts
--- a/Engine/Source/Programs/Horde/HordeDashboard/src/base/utilities/issueUtils.ts
+++ b/Engine/Source/Programs/Horde/HordeDashboard/src/base/utilities/issueUtils.ts
@@ -1,3 +1,5 @@
+// Copyright Epic Games, Inc. All Rights Reserved.
+
 import moment from "moment";
 import { GetIssueResponse, StreamData } from "../../backend/Api";
 import dashboard from "../../backend/Dashboard";
diff --git a/Engine/Source/Programs/Horde/HordeDashboard/src/base/utilities/jumpToAnchor.ts b/Engine/Source/Programs/Horde/HordeDashboard/src/base/utilities/jumpToAnchor.ts
--- a/Engine/Source/Programs/Horde/HordeDashboard/src/base/utilities/jumpToAnchor.ts
+++ b/Engine/Source/Programs/Horde/HordeDashboard/src/base/utilities/jumpToAnchor.ts
@@ -1,3 +1,5 @@
+// Copyright Epic Games, Inc. All Rights Reserved.
+
 import { extractAnchorLink } from './extractAnchorLink';
 
 const SCROLL_DISTANCE = 52;
diff --git a/Engine/Source/Programs/Horde/HordeDashboard/src/base/utilities/removeAnchorLink.ts b/Engine/Source/Programs/Horde/HordeDashboard/src/base/utilities/removeAnchorLink.ts
--- a/Engine/Source/Programs/Horde/HordeDashboard/src/base/utilities/removeAnchorLink.ts
+++ b/Engine/Source/Programs/Horde/HordeDashboard/src/base/utilities/removeAnchorLink.ts
@@ -1,3 +1,5 @@
+// Copyright Epic Games, Inc. All Rights Reserved.
+
 /**
  * Remove the anchor link from a full page URL or a hash.
  * Preserves any route path specified in the URL.
diff --git a/Engine/Source/Programs/Horde/HordeDashboard/src/base/utilities/streamUtils.ts b/Engine/Source/Programs/Horde/HordeDashboard/src/base/utilities/streamUtils.ts
--- a/Engine/Source/Programs/Horde/HordeDashboard/src/base/utilities/streamUtils.ts
+++ b/Engine/Source/Programs/Horde/HordeDashboard/src/base/utilities/streamUtils.ts
@@ -1,5 +1,6 @@
-import { projectStore } from "../../backend/ProjectStore";
+// Copyright Epic Games, Inc. All Rights Reserved.
 
+import { projectStore } from "../../backend/ProjectStore";
 
 export const getActiveStreamId = (): string | undefined => {
 
diff --git a/Engine/Source/Programs/Horde/HordeDashboard/src/base/utilities/stringUtills.ts b/Engine/Source/Programs/Horde/HordeDashboard/src/base/utilities/stringUtills.ts
--- a/Engine/Source/Programs/Horde/HordeDashboard/src/base/utilities/stringUtills.ts
+++ b/Engine/Source/Programs/Horde/HordeDashboard/src/base/utilities/stringUtills.ts
@@ -1,3 +1,4 @@
+// Copyright Epic Games, Inc. All Rights Reserved.
 
 export function formatBytes(bytes: number, decimals = 2) {
     if (!+bytes) return '0 Bytes'
diff --git a/Engine/Source/Programs/Horde/HordeDashboard/src/base/utilities/timeUtils.ts b/Engine/Source/Programs/Horde/HordeDashboard/src/base/utilities/timeUtils.ts
--- a/Engine/Source/Programs/Horde/HordeDashboard/src/base/utilities/timeUtils.ts
+++ b/Engine/Source/Programs/Horde/HordeDashboard/src/base/utilities/timeUtils.ts
@@ -1,3 +1,5 @@
+// Copyright Epic Games, Inc. All Rights Reserved.
+
 import moment, { Moment } from "moment-timezone";
 import { BatchData, GetAgentLeaseResponse, JobData, LabelState, StepData, TimingInfo } from "../../backend/Api";
 import { JobLabel } from "../../backend/JobDetails";
diff --git a/Engine/Source/Programs/Horde/HordeDashboard/src/components/AutomationCommon.tsx b/Engine/Source/Programs/Horde/HordeDashboard/src/components/AutomationCommon.tsx
--- a/Engine/Source/Programs/Horde/HordeDashboard/src/components/AutomationCommon.tsx
+++ b/Engine/Source/Programs/Horde/HordeDashboard/src/components/AutomationCommon.tsx
@@ -1,3 +1,4 @@
+// Copyright Epic Games, Inc. All Rights Reserved.
 
 import { mergeStyles, mergeStyleSets } from "@fluentui/react";
 
diff --git a/Engine/Source/Programs/Horde/HordeDashboard/src/components/AutomationReport.tsx b/Engine/Source/Programs/Horde/HordeDashboard/src/components/AutomationReport.tsx
--- a/Engine/Source/Programs/Horde/HordeDashboard/src/components/AutomationReport.tsx
+++ b/Engine/Source/Programs/Horde/HordeDashboard/src/components/AutomationReport.tsx
@@ -1,3 +1,5 @@
+// Copyright Epic Games, Inc. All Rights Reserved.
+
 import { IconButton, Label, List, Modal, ScrollablePane, ScrollbarVisibility, Spinner, SpinnerSize, Stack, Text, getFocusStyle, mergeStyleSets } from "@fluentui/react";
 import { action, makeObservable, observable } from "mobx";
 import { observer } from "mobx-react-lite";
diff --git a/Engine/Source/Programs/Horde/HordeDashboard/src/components/AutomationSuiteDetails.tsx b/Engine/Source/Programs/Horde/HordeDashboard/src/components/AutomationSuiteDetails.tsx
--- a/Engine/Source/Programs/Horde/HordeDashboard/src/components/AutomationSuiteDetails.tsx
+++ b/Engine/Source/Programs/Horde/HordeDashboard/src/components/AutomationSuiteDetails.tsx
@@ -1,3 +1,5 @@
+// Copyright Epic Games, Inc. All Rights Reserved.
+
 import { DetailsList, DetailsListLayoutMode, DetailsRow, Dropdown, FontIcon, IColumn, IconButton, IDetailsGroupRenderProps, IDetailsListProps, IDropdownOption, IGroup, mergeStyleSets, Modal, SelectionMode, Spinner, SpinnerSize, Stack, Text, TextField } from "@fluentui/react";
 import * as d3 from "d3";
 import { action, makeObservable, observable } from 'mobx';
diff --git a/Engine/Source/Programs/Horde/HordeDashboard/src/components/AutomationSuiteTest.tsx b/Engine/Source/Programs/Horde/HordeDashboard/src/components/AutomationSuiteTest.tsx
--- a/Engine/Source/Programs/Horde/HordeDashboard/src/components/AutomationSuiteTest.tsx
+++ b/Engine/Source/Programs/Horde/HordeDashboard/src/components/AutomationSuiteTest.tsx
@@ -1,3 +1,4 @@
+// Copyright Epic Games, Inc. All Rights Reserved.
 
 import { DefaultButton, FontIcon, IContextualMenuItem, IContextualMenuItemProps, IContextualMenuProps, Image, Label, ScrollablePane, ScrollbarVisibility, Spinner, SpinnerSize, Stack, Text } from "@fluentui/react";
 import { action, makeObservable, observable } from 'mobx';
diff --git a/Engine/Source/Programs/Horde/HordeDashboard/src/components/AutomationViewSummary.tsx b/Engine/Source/Programs/Horde/HordeDashboard/src/components/AutomationViewSummary.tsx
--- a/Engine/Source/Programs/Horde/HordeDashboard/src/components/AutomationViewSummary.tsx
+++ b/Engine/Source/Programs/Horde/HordeDashboard/src/components/AutomationViewSummary.tsx
@@ -1,3 +1,5 @@
+// Copyright Epic Games, Inc. All Rights Reserved.
+
 import { DefaultButton, FontIcon, IconButton, mergeStyleSets, Modal, ScrollablePane, ScrollbarVisibility, Stack, Text } from "@fluentui/react";
 import { observer } from "mobx-react-lite";
 import { useState } from "react";
diff --git a/Engine/Source/Programs/Horde/HordeDashboard/src/components/IssueViewV2.tsx b/Engine/Source/Programs/Horde/HordeDashboard/src/components/IssueViewV2.tsx
--- a/Engine/Source/Programs/Horde/HordeDashboard/src/components/IssueViewV2.tsx
+++ b/Engine/Source/Programs/Horde/HordeDashboard/src/components/IssueViewV2.tsx
@@ -1,3 +1,5 @@
+// Copyright Epic Games, Inc. All Rights Reserved.
+
 import { Checkbox, CommandBar, CommandBarButton, DefaultButton, Dialog, DialogFooter, DialogType, getFocusStyle, ICommandBarItemProps, Icon, IconButton, IContextualMenuItem, IContextualMenuItemProps, IContextualMenuProps, IGroup, Label, List, mergeStyleSets, MessageBar, MessageBarType, Modal, Persona, PersonaSize, Pivot, PivotItem, PrimaryButton, Spinner, SpinnerSize, Stack, TagPicker, Text, TextField } from "@fluentui/react";
 import { action, makeObservable, observable } from 'mobx';
 import { observer } from "mobx-react-lite";
diff --git a/Engine/Source/Programs/Horde/HordeDashboard/src/components/JobSearchSimple.tsx b/Engine/Source/Programs/Horde/HordeDashboard/src/components/JobSearchSimple.tsx
--- a/Engine/Source/Programs/Horde/HordeDashboard/src/components/JobSearchSimple.tsx
+++ b/Engine/Source/Programs/Horde/HordeDashboard/src/components/JobSearchSimple.tsx
@@ -1,3 +1,5 @@
+// Copyright Epic Games, Inc. All Rights Reserved.
+
 import { Checkbox, ContextualMenuItemType, DatePicker, DefaultButton, DetailsList, DetailsRow, IColumn, IconButton, IContextualMenuItem, IContextualMenuItemProps, IContextualMenuProps, IDetailsGroupRenderProps, IDetailsListProps, IGroup, Label, Modal, PrimaryButton, SelectionMode, Spinner, SpinnerSize, Stack, Text, TextField } from "@fluentui/react";
 import { observer } from "mobx-react-lite";
 import moment from "moment";
diff --git a/Engine/Source/Programs/Horde/HordeDashboard/src/components/PoolsView.tsx b/Engine/Source/Programs/Horde/HordeDashboard/src/components/PoolsView.tsx
--- a/Engine/Source/Programs/Horde/HordeDashboard/src/components/PoolsView.tsx
+++ b/Engine/Source/Programs/Horde/HordeDashboard/src/components/PoolsView.tsx
@@ -1,3 +1,5 @@
+// Copyright Epic Games, Inc. All Rights Reserved.
+
 import { DefaultButton, DetailsHeader, DetailsList, FontIcon, IColumn, IDetailsHeaderStyles, IDetailsListProps, ITag, Pivot, PivotItem, SelectionMode, Spinner, SpinnerSize, Stack, Sticky, StickyPositionType, TagPicker, Text, mergeStyleSets, mergeStyles } from "@fluentui/react";
 import { observer } from "mobx-react-lite";
 import React, { useEffect, useState } from "react";
diff --git a/Engine/Source/Programs/Horde/HordeDashboard/src/components/PreviewChanges.tsx b/Engine/Source/Programs/Horde/HordeDashboard/src/components/PreviewChanges.tsx
--- a/Engine/Source/Programs/Horde/HordeDashboard/src/components/PreviewChanges.tsx
+++ b/Engine/Source/Programs/Horde/HordeDashboard/src/components/PreviewChanges.tsx
@@ -1,3 +1,5 @@
+// Copyright Epic Games, Inc. All Rights Reserved.
+
 import { DetailsList, DetailsListLayoutMode, IColumn, IconButton, Modal, PrimaryButton, SelectionMode, Stack, Text } from "@fluentui/react";
 import { Link as ReactRouterLink } from 'react-router-dom';
 import { getHordeStyling } from "../styles/Styles";
diff --git a/Engine/Source/Programs/Horde/HordeDashboard/src/components/UserSelect.tsx b/Engine/Source/Programs/Horde/HordeDashboard/src/components/UserSelect.tsx
--- a/Engine/Source/Programs/Horde/HordeDashboard/src/components/UserSelect.tsx
+++ b/Engine/Source/Programs/Horde/HordeDashboard/src/components/UserSelect.tsx
@@ -1,3 +1,5 @@
+// Copyright Epic Games, Inc. All Rights Reserved.
+
 import { CompactPeoplePicker, IBasePickerSuggestionsProps, IInputProps, IPersonaProps, Label, Stack } from "@fluentui/react";
 import { observer } from "mobx-react-lite";
 import React, { useEffect, useState } from "react";
diff --git a/Engine/Source/Programs/Horde/HordeDashboard/src/components/accounts/HordeLoginView.tsx b/Engine/Source/Programs/Horde/HordeDashboard/src/components/accounts/HordeLoginView.tsx
--- a/Engine/Source/Programs/Horde/HordeDashboard/src/components/accounts/HordeLoginView.tsx
+++ b/Engine/Source/Programs/Horde/HordeDashboard/src/components/accounts/HordeLoginView.tsx
@@ -1,3 +1,5 @@
+// Copyright Epic Games, Inc. All Rights Reserved.
+
 import { ThemeProvider, Stack, Image, TextField, PrimaryButton, MessageBar, MessageBarType } from "@fluentui/react";
 import dashboard from "../../backend/Dashboard";
 import { darkTheme } from "../../styles/darkTheme";
diff --git a/Engine/Source/Programs/Horde/HordeDashboard/src/components/accounts/UserAccountPanel.tsx b/Engine/Source/Programs/Horde/HordeDashboard/src/components/accounts/UserAccountPanel.tsx
--- a/Engine/Source/Programs/Horde/HordeDashboard/src/components/accounts/UserAccountPanel.tsx
+++ b/Engine/Source/Programs/Horde/HordeDashboard/src/components/accounts/UserAccountPanel.tsx
@@ -1,3 +1,5 @@
+// Copyright Epic Games, Inc. All Rights Reserved.
+
 import { DetailsList, DetailsListLayoutMode, ScrollablePane, ScrollbarVisibility, SelectionMode, Stack, Text } from "@fluentui/react";
 import dashboard from "../../backend/Dashboard";
 import { useState } from "react";
diff --git a/Engine/Source/Programs/Horde/HordeDashboard/src/components/agents/AgentTelemetrySparkline.tsx b/Engine/Source/Programs/Horde/HordeDashboard/src/components/agents/AgentTelemetrySparkline.tsx
--- a/Engine/Source/Programs/Horde/HordeDashboard/src/components/agents/AgentTelemetrySparkline.tsx
+++ b/Engine/Source/Programs/Horde/HordeDashboard/src/components/agents/AgentTelemetrySparkline.tsx
@@ -1,3 +1,5 @@
+// Copyright Epic Games, Inc. All Rights Reserved.
+
 import { Label, Spinner, SpinnerSize, Stack, Text } from "@fluentui/react";
 import * as d3 from "d3";
 import { observer } from "mobx-react-lite";
diff --git a/Engine/Source/Programs/Horde/HordeDashboard/src/components/agents/HistoryModalStore.ts b/Engine/Source/Programs/Horde/HordeDashboard/src/components/agents/HistoryModalStore.ts
--- a/Engine/Source/Programs/Horde/HordeDashboard/src/components/agents/HistoryModalStore.ts
+++ b/Engine/Source/Programs/Horde/HordeDashboard/src/components/agents/HistoryModalStore.ts
@@ -1,3 +1,5 @@
+// Copyright Epic Games, Inc. All Rights Reserved.
+
 import { action, makeObservable, observable } from "mobx";
 import { AgentData, GetAgentLeaseResponse, GetAgentSessionResponse, GetAgentTelemetrySampleResponse, GetJobResponse, LeaseData, SessionData } from "../../backend/Api" 
 import backend from "../../backend";
diff --git a/Engine/Source/Programs/Horde/HordeDashboard/src/components/artifacts/ArtifactButton.tsx b/Engine/Source/Programs/Horde/HordeDashboard/src/components/artifacts/ArtifactButton.tsx
--- a/Engine/Source/Programs/Horde/HordeDashboard/src/components/artifacts/ArtifactButton.tsx
+++ b/Engine/Source/Programs/Horde/HordeDashboard/src/components/artifacts/ArtifactButton.tsx
@@ -1,3 +1,5 @@
+// Copyright Epic Games, Inc. All Rights Reserved.
+
 import { IContextualMenuProps, DirectionalHint, PrimaryButton, Modal, Stack, Text, DefaultButton, Checkbox } from "@fluentui/react";
 import backend from "horde/backend";
 import { ArtifactBrowserItem, ArtifactBrowserType, ArtifactSelection, CreateZipRequest, DownloadFormat, GetArtifactResponse } from "horde/backend/Api";
diff --git a/Engine/Source/Programs/Horde/HordeDashboard/src/components/artifacts/ArtifactsSearch.tsx b/Engine/Source/Programs/Horde/HordeDashboard/src/components/artifacts/ArtifactsSearch.tsx
--- a/Engine/Source/Programs/Horde/HordeDashboard/src/components/artifacts/ArtifactsSearch.tsx
+++ b/Engine/Source/Programs/Horde/HordeDashboard/src/components/artifacts/ArtifactsSearch.tsx
@@ -1,3 +1,5 @@
+// Copyright Epic Games, Inc. All Rights Reserved.
+
 import { ComboBox, DefaultButton, DetailsList, DetailsListLayoutMode, DirectionalHint, IComboBox, IComboBoxOption, IContextualMenuProps, IDetailsListProps, IconButton, Label, Modal, PrimaryButton, ScrollablePane, ScrollbarVisibility, SelectionMode, Spinner, SpinnerSize, Stack, Text, TextField } from "@fluentui/react";
 import { useConst } from '@fluentui/react-hooks';
 import React, { useState } from "react";
diff --git a/Engine/Source/Programs/Horde/HordeDashboard/src/components/bisection/CreateModal.tsx b/Engine/Source/Programs/Horde/HordeDashboard/src/components/bisection/CreateModal.tsx
--- a/Engine/Source/Programs/Horde/HordeDashboard/src/components/bisection/CreateModal.tsx
+++ b/Engine/Source/Programs/Horde/HordeDashboard/src/components/bisection/CreateModal.tsx
@@ -1,3 +1,5 @@
+// Copyright Epic Games, Inc. All Rights Reserved.
+
 import { Modal, Stack, Text, IconButton, TextField, PrimaryButton, DefaultButton, MessageBar, MessageBarType } from "@fluentui/react";
 import { useState } from "react";
 import backend from "../../backend";
diff --git a/Engine/Source/Programs/Horde/HordeDashboard/src/components/build/NewBuildV2.tsx b/Engine/Source/Programs/Horde/HordeDashboard/src/components/build/NewBuildV2.tsx
--- a/Engine/Source/Programs/Horde/HordeDashboard/src/components/build/NewBuildV2.tsx
+++ b/Engine/Source/Programs/Horde/HordeDashboard/src/components/build/NewBuildV2.tsx
@@ -1,5 +1,6 @@
+// Copyright Epic Games, Inc. All Rights Reserved.
+
 import { Checkbox, ComboBox, ContextualMenuItemType, DefaultButton, DirectionalHint, Dropdown, DropdownMenuItemType, IComboBoxOption, IContextualMenuItem, IContextualMenuProps, IDropdownOption, ITag, Icon, IconButton, Label, MessageBar, MessageBarType, Modal, Pivot, PivotItem, PrimaryButton, ScrollablePane, ScrollbarVisibility, Spinner, SpinnerSize, Stack, TagPicker, Text, TextField, TooltipHost, ValidationState } from "@fluentui/react";
-import { useConst } from '@fluentui/react-hooks';
 import { ITextField } from "@fluentui/react/lib-commonjs/TextField";
 import Markdown from "markdown-to-jsx";
 import { action, makeObservable, observable } from "mobx";
diff --git a/Engine/Source/Programs/Horde/HordeDashboard/src/components/devices/DeviceEditor.tsx b/Engine/Source/Programs/Horde/HordeDashboard/src/components/devices/DeviceEditor.tsx
--- a/Engine/Source/Programs/Horde/HordeDashboard/src/components/devices/DeviceEditor.tsx
+++ b/Engine/Source/Programs/Horde/HordeDashboard/src/components/devices/DeviceEditor.tsx
@@ -1,3 +1,4 @@
+// Copyright Epic Games, Inc. All Rights Reserved.
 import { SpinnerSize, DialogType, IDropdownOption, MessageBarType, Text, Checkbox, DefaultButton, Dialog, DialogFooter, Dropdown, MessageBar, Modal, PrimaryButton, Spinner, Stack, TextField, Label, ITag, TagPicker, IPickerItemProps } from "@fluentui/react";
 import React, { useState } from "react";
 import backend from "../../backend";
diff --git a/Engine/Source/Programs/Horde/HordeDashboard/src/components/devices/DeviceHealthWidgets.tsx b/Engine/Source/Programs/Horde/HordeDashboard/src/components/devices/DeviceHealthWidgets.tsx
--- a/Engine/Source/Programs/Horde/HordeDashboard/src/components/devices/DeviceHealthWidgets.tsx
+++ b/Engine/Source/Programs/Horde/HordeDashboard/src/components/devices/DeviceHealthWidgets.tsx
@@ -1,3 +1,5 @@
+// Copyright Epic Games, Inc. All Rights Reserved.
+
 import { Stack, Label, Text, TooltipHost, ITooltipProps, DirectionalHint, IGroup } from "@fluentui/react";
 import { useState, useRef } from "react";
 import dashboard, { StatusColor } from "../../backend/Dashboard";
diff --git a/Engine/Source/Programs/Horde/HordeDashboard/src/components/devices/DeviceInfoView.tsx b/Engine/Source/Programs/Horde/HordeDashboard/src/components/devices/DeviceInfoView.tsx
--- a/Engine/Source/Programs/Horde/HordeDashboard/src/components/devices/DeviceInfoView.tsx
+++ b/Engine/Source/Programs/Horde/HordeDashboard/src/components/devices/DeviceInfoView.tsx
@@ -1,3 +1,5 @@
+// Copyright Epic Games, Inc. All Rights Reserved.
+
 import { DefaultButton, DetailsList, DetailsRow, IColumn, IconButton, IDetailsListProps, Label, Modal, PrimaryButton, SelectionMode, Spinner, SpinnerSize, Stack, Text, TextField } from "@fluentui/react";
 import moment from "moment";
 import { useState } from "react";
diff --git a/Engine/Source/Programs/Horde/HordeDashboard/src/components/devices/DevicePoolTelemetry.tsx b/Engine/Source/Programs/Horde/HordeDashboard/src/components/devices/DevicePoolTelemetry.tsx
--- a/Engine/Source/Programs/Horde/HordeDashboard/src/components/devices/DevicePoolTelemetry.tsx
+++ b/Engine/Source/Programs/Horde/HordeDashboard/src/components/devices/DevicePoolTelemetry.tsx
@@ -1,3 +1,5 @@
+// Copyright Epic Games, Inc. All Rights Reserved.
+
 import { DefaultButton, DetailsList, DetailsListLayoutMode, Dropdown, FocusZone, FocusZoneDirection, IColumn, IconButton, IDropdownOption, Modal, ScrollablePane, ScrollbarVisibility, SelectionMode, Spinner, SpinnerSize, Stack, Text } from "@fluentui/react";
 import { useState, useRef, useEffect, useCallback, useMemo } from "react";
 import { Link } from "react-router-dom";
diff --git a/Engine/Source/Programs/Horde/HordeDashboard/src/components/devices/DeviceTelemetryGraphs.tsx b/Engine/Source/Programs/Horde/HordeDashboard/src/components/devices/DeviceTelemetryGraphs.tsx
--- a/Engine/Source/Programs/Horde/HordeDashboard/src/components/devices/DeviceTelemetryGraphs.tsx
+++ b/Engine/Source/Programs/Horde/HordeDashboard/src/components/devices/DeviceTelemetryGraphs.tsx
@@ -1,3 +1,5 @@
+// Copyright Epic Games, Inc. All Rights Reserved.
+
 import { PoolTelemetryHandler, UsagePlatformTelemetry, PlatformTelemetry } from "../../backend/DeviceTelemetryHandler"
 import { getHordeStyling } from "../../styles/Styles";
 import { getHumanTime, msecToElapsed } from "../../base/utilities/timeUtils";
diff --git a/Engine/Source/Programs/Horde/HordeDashboard/src/components/devices/DeviceView.tsx b/Engine/Source/Programs/Horde/HordeDashboard/src/components/devices/DeviceView.tsx
--- a/Engine/Source/Programs/Horde/HordeDashboard/src/components/devices/DeviceView.tsx
+++ b/Engine/Source/Programs/Horde/HordeDashboard/src/components/devices/DeviceView.tsx
@@ -1,3 +1,5 @@
+// Copyright Epic Games, Inc. All Rights Reserved.
+
 import { DefaultButton, DetailsList, DetailsListLayoutMode, Dropdown, FocusZone, FocusZoneDirection, IColumn, Icon, IconButton, IGroup, GroupHeader, Label, Link as FluentLink, mergeStyleSets, Modal, Pivot, PivotItem, PrimaryButton, ScrollablePane, ScrollbarVisibility, SelectionMode, Spinner, SpinnerSize, Stack, Text, TextField, IDetailsGroupDividerProps, IGroupHeaderProps } from "@fluentui/react";
 import { observer } from "mobx-react-lite";
 import { observable, action, makeObservable } from "mobx";
diff --git a/Engine/Source/Programs/Horde/HordeDashboard/src/components/docs/DocView.tsx b/Engine/Source/Programs/Horde/HordeDashboard/src/components/docs/DocView.tsx
--- a/Engine/Source/Programs/Horde/HordeDashboard/src/components/docs/DocView.tsx
+++ b/Engine/Source/Programs/Horde/HordeDashboard/src/components/docs/DocView.tsx
@@ -1,3 +1,5 @@
+// Copyright Epic Games, Inc. All Rights Reserved.
+
 import { Stack, mergeStyleSets } from "@fluentui/react";
 import { action, makeObservable, observable } from "mobx";
 import { observer } from "mobx-react-lite";
diff --git a/Engine/Source/Programs/Horde/HordeDashboard/src/components/error/ErrorBoundary.tsx b/Engine/Source/Programs/Horde/HordeDashboard/src/components/error/ErrorBoundary.tsx
--- a/Engine/Source/Programs/Horde/HordeDashboard/src/components/error/ErrorBoundary.tsx
+++ b/Engine/Source/Programs/Horde/HordeDashboard/src/components/error/ErrorBoundary.tsx
@@ -1,3 +1,5 @@
+// Copyright Epic Games, Inc. All Rights Reserved.
+
 import { Stack, Text } from "@fluentui/react";
 import { Component, ErrorInfo, ReactNode } from "react";
 import { errorDialogStore } from "./ErrorStore";
diff --git a/Engine/Source/Programs/Horde/HordeDashboard/src/components/error/ErrorDialogStyles.tsx b/Engine/Source/Programs/Horde/HordeDashboard/src/components/error/ErrorDialogStyles.tsx
--- a/Engine/Source/Programs/Horde/HordeDashboard/src/components/error/ErrorDialogStyles.tsx
+++ b/Engine/Source/Programs/Horde/HordeDashboard/src/components/error/ErrorDialogStyles.tsx
@@ -1,3 +1,5 @@
+// Copyright Epic Games, Inc. All Rights Reserved.
+
 import { mergeStyleSets } from "@fluentui/react"
 import dashboard, { StatusColor } from "horde/backend/Dashboard";
 
diff --git a/Engine/Source/Programs/Horde/HordeDashboard/src/components/error/ErrorStore.ts b/Engine/Source/Programs/Horde/HordeDashboard/src/components/error/ErrorStore.ts
--- a/Engine/Source/Programs/Horde/HordeDashboard/src/components/error/ErrorStore.ts
+++ b/Engine/Source/Programs/Horde/HordeDashboard/src/components/error/ErrorStore.ts
@@ -1,3 +1,5 @@
+// Copyright Epic Games, Inc. All Rights Reserved.
+
 import { getSiteConfig } from "../../backend/Config";
 import { action, makeObservable, observable, override, runInAction } from "mobx";
 
diff --git a/Engine/Source/Programs/Horde/HordeDashboard/src/components/jobDetailsV2/JobDetailLabels.tsx b/Engine/Source/Programs/Horde/HordeDashboard/src/components/jobDetailsV2/JobDetailLabels.tsx
--- a/Engine/Source/Programs/Horde/HordeDashboard/src/components/jobDetailsV2/JobDetailLabels.tsx
+++ b/Engine/Source/Programs/Horde/HordeDashboard/src/components/jobDetailsV2/JobDetailLabels.tsx
@@ -1,3 +1,5 @@
+// Copyright Epic Games, Inc. All Rights Reserved.
+
 import { ConstrainMode, DefaultButton, DetailsList, DetailsListLayoutMode, DirectionalHint, IColumn, SelectionMode, Stack, Text, TooltipHost } from "@fluentui/react";
 import { observer } from "mobx-react-lite";
 import { useNavigate, useLocation } from "react-router-dom";
diff --git a/Engine/Source/Programs/Horde/HordeDashboard/src/components/jobDetailsV2/JobDetailViewSteps.tsx b/Engine/Source/Programs/Horde/HordeDashboard/src/components/jobDetailsV2/JobDetailViewSteps.tsx
--- a/Engine/Source/Programs/Horde/HordeDashboard/src/components/jobDetailsV2/JobDetailViewSteps.tsx
+++ b/Engine/Source/Programs/Horde/HordeDashboard/src/components/jobDetailsV2/JobDetailViewSteps.tsx
@@ -1,3 +1,5 @@
+// Copyright Epic Games, Inc. All Rights Reserved.
+
 import { CollapseAllVisibility, ConstrainMode, DetailsHeader, DetailsList, DetailsListLayoutMode, DetailsRow, DirectionalHint, IColumn, Icon, IDetailsListProps, ProgressIndicator, SelectionMode, Stack, Text, TooltipHost } from "@fluentui/react";
 import { observer } from "mobx-react-lite";
 import React, { useEffect, useId, useState } from "react";
diff --git a/Engine/Source/Programs/Horde/HordeDashboard/src/components/jobDetailsV2/JobDetailViewV2.tsx b/Engine/Source/Programs/Horde/HordeDashboard/src/components/jobDetailsV2/JobDetailViewV2.tsx
--- a/Engine/Source/Programs/Horde/HordeDashboard/src/components/jobDetailsV2/JobDetailViewV2.tsx
+++ b/Engine/Source/Programs/Horde/HordeDashboard/src/components/jobDetailsV2/JobDetailViewV2.tsx
@@ -1,3 +1,5 @@
+// Copyright Epic Games, Inc. All Rights Reserved.
+
 import { mergeStyleSets, Pivot, PivotItem, Spinner, SpinnerSize, Stack } from "@fluentui/react";
 import { observer } from "mobx-react-lite";
 import { useEffect, useRef, useState } from "react";
diff --git a/Engine/Source/Programs/Horde/HordeDashboard/src/components/jobDetailsV2/JobDetailsViewCommon.tsx b/Engine/Source/Programs/Horde/HordeDashboard/src/components/jobDetailsV2/JobDetailsViewCommon.tsx
--- a/Engine/Source/Programs/Horde/HordeDashboard/src/components/jobDetailsV2/JobDetailsViewCommon.tsx
+++ b/Engine/Source/Programs/Horde/HordeDashboard/src/components/jobDetailsV2/JobDetailsViewCommon.tsx
@@ -1,3 +1,5 @@
+// Copyright Epic Games, Inc. All Rights Reserved.
+
 import { Dropdown, FontIcon, IDropdownOption, IPickerItemProps, ISuggestionItemProps, ITag, PrimaryButton, Stack, TagPicker, Text, ValidationState } from "@fluentui/react";
 import { action, makeObservable, observable } from "mobx";
 import { observer } from "mobx-react-lite";
diff --git a/Engine/Source/Programs/Horde/HordeDashboard/src/components/jobDetailsV2/StepDetailSummary.tsx b/Engine/Source/Programs/Horde/HordeDashboard/src/components/jobDetailsV2/StepDetailSummary.tsx
--- a/Engine/Source/Programs/Horde/HordeDashboard/src/components/jobDetailsV2/StepDetailSummary.tsx
+++ b/Engine/Source/Programs/Horde/HordeDashboard/src/components/jobDetailsV2/StepDetailSummary.tsx
@@ -1,3 +1,4 @@
+// Copyright Epic Games, Inc. All Rights Reserved.
 
 import { Label, Stack, Text } from '@fluentui/react';
 import { observer } from 'mobx-react-lite';
diff --git a/Engine/Source/Programs/Horde/HordeDashboard/src/components/preflights/PreflightConfigCheck.tsx b/Engine/Source/Programs/Horde/HordeDashboard/src/components/preflights/PreflightConfigCheck.tsx
--- a/Engine/Source/Programs/Horde/HordeDashboard/src/components/preflights/PreflightConfigCheck.tsx
+++ b/Engine/Source/Programs/Horde/HordeDashboard/src/components/preflights/PreflightConfigCheck.tsx
@@ -1,3 +1,5 @@
+// Copyright Epic Games, Inc. All Rights Reserved.
+
 import { ComboBox, FontIcon, IComboBoxOption, IconButton, MaskedTextField, Modal, PrimaryButton, Spinner, SpinnerSize, Stack, Text } from "@fluentui/react";
 import { useState } from "react";
 import { useNavigate } from "react-router-dom";
diff --git a/Engine/Source/Programs/Horde/HordeDashboard/src/components/projects/StreamChooser.tsx b/Engine/Source/Programs/Horde/HordeDashboard/src/components/projects/StreamChooser.tsx
--- a/Engine/Source/Programs/Horde/HordeDashboard/src/components/projects/StreamChooser.tsx
+++ b/Engine/Source/Programs/Horde/HordeDashboard/src/components/projects/StreamChooser.tsx
@@ -1,3 +1,5 @@
+// Copyright Epic Games, Inc. All Rights Reserved.
+
 import { DefaultButton, IContextualMenuItem, IContextualMenuProps, IRefObject, Stack } from "@fluentui/react"
 import { projectStore } from "../../backend/ProjectStore"
 import { useImperativeHandle, useState } from "react";
diff --git a/Engine/Source/Programs/Horde/HordeDashboard/src/components/setup/HordeSetupView.tsx b/Engine/Source/Programs/Horde/HordeDashboard/src/components/setup/HordeSetupView.tsx
--- a/Engine/Source/Programs/Horde/HordeDashboard/src/components/setup/HordeSetupView.tsx
+++ b/Engine/Source/Programs/Horde/HordeDashboard/src/components/setup/HordeSetupView.tsx
@@ -1,3 +1,5 @@
+// Copyright Epic Games, Inc. All Rights Reserved.
+
 import { ThemeProvider, Stack, Image, TextField, PrimaryButton, MessageBar, MessageBarType } from "@fluentui/react";
 import dashboard from "../../backend/Dashboard";
 import { darkTheme } from "../../styles/darkTheme";
diff --git a/Engine/Source/Programs/Horde/HordeDashboard/src/components/tools/ToolManagement.tsx b/Engine/Source/Programs/Horde/HordeDashboard/src/components/tools/ToolManagement.tsx
--- a/Engine/Source/Programs/Horde/HordeDashboard/src/components/tools/ToolManagement.tsx
+++ b/Engine/Source/Programs/Horde/HordeDashboard/src/components/tools/ToolManagement.tsx
@@ -1,3 +1,5 @@
+// Copyright Epic Games, Inc. All Rights Reserved.
+
 import { DefaultButton, DetailsList, DetailsListLayoutMode, IContextualMenuProps, IDetailsListProps, IconButton, Modal, PrimaryButton, ProgressIndicator, SelectionMode, Spinner, SpinnerSize, Stack, Text } from "@fluentui/react";
 import { useConst } from '@fluentui/react-hooks';
 import Markdown from "markdown-to-jsx";
diff --git a/Engine/Source/Programs/Horde/HordeDashboard/src/index.css b/Engine/Source/Programs/Horde/HordeDashboard/src/index.css
--- a/Engine/Source/Programs/Horde/HordeDashboard/src/index.css
+++ b/Engine/Source/Programs/Horde/HordeDashboard/src/index.css
@@ -1,3 +1,5 @@
+/* Copyright Epic Games, Inc. All Rights Reserved. */
+
 body {
   margin: 0;
   font-family: "Horde Open Sans Regular";
diff --git a/Engine/Source/Programs/Horde/HordeDashboard/src/main.tsx b/Engine/Source/Programs/Horde/HordeDashboard/src/main.tsx
--- a/Engine/Source/Programs/Horde/HordeDashboard/src/main.tsx
+++ b/Engine/Source/Programs/Horde/HordeDashboard/src/main.tsx
@@ -1,3 +1,5 @@
+// Copyright Epic Games, Inc. All Rights Reserved.
+
 import React from 'react'
 import ReactDOM from 'react-dom/client'
 import App from './App.tsx'
diff --git a/Engine/Source/Programs/Horde/HordeDashboard/src/styles/darkTheme.ts b/Engine/Source/Programs/Horde/HordeDashboard/src/styles/darkTheme.ts
--- a/Engine/Source/Programs/Horde/HordeDashboard/src/styles/darkTheme.ts
+++ b/Engine/Source/Programs/Horde/HordeDashboard/src/styles/darkTheme.ts
@@ -1,3 +1,5 @@
+// Copyright Epic Games, Inc. All Rights Reserved.
+
 import { createTheme } from '@fluentui/react';
 import { HordeTheme, HordeThemeExtensions } from './themeTypes';
 
diff --git a/Engine/Source/Programs/Horde/HordeDashboard/src/styles/lightTheme.ts b/Engine/Source/Programs/Horde/HordeDashboard/src/styles/lightTheme.ts
--- a/Engine/Source/Programs/Horde/HordeDashboard/src/styles/lightTheme.ts
+++ b/Engine/Source/Programs/Horde/HordeDashboard/src/styles/lightTheme.ts
@@ -1,3 +1,4 @@
+// Copyright Epic Games, Inc. All Rights Reserved.
 
 import { createTheme } from '@fluentui/react';
 import { HordeTheme } from './themeTypes';
diff --git a/Engine/Source/Programs/Horde/HordeDashboard/src/styles/theme.ts b/Engine/Source/Programs/Horde/HordeDashboard/src/styles/theme.ts
--- a/Engine/Source/Programs/Horde/HordeDashboard/src/styles/theme.ts
+++ b/Engine/Source/Programs/Horde/HordeDashboard/src/styles/theme.ts
@@ -1,3 +1,4 @@
+// Copyright Epic Games, Inc. All Rights Reserved.
 
 import dashboard from "../backend/Dashboard";
 import { darkTheme } from "./darkTheme";
diff --git a/Engine/Source/Programs/Horde/HordeDashboard/src/styles/themeTypes.ts b/Engine/Source/Programs/Horde/HordeDashboard/src/styles/themeTypes.ts
--- a/Engine/Source/Programs/Horde/HordeDashboard/src/styles/themeTypes.ts
+++ b/Engine/Source/Programs/Horde/HordeDashboard/src/styles/themeTypes.ts
@@ -1,3 +1,5 @@
+// Copyright Epic Games, Inc. All Rights Reserved.
+
 import { Theme } from "@fluentui/react";
 
 export type HordeThemeExtensions = {
diff --git a/Engine/Source/Programs/Horde/HordeDashboard/src/vite-env.d.ts b/Engine/Source/Programs/Horde/HordeDashboard/src/vite-env.d.ts
--- a/Engine/Source/Programs/Horde/HordeDashboard/src/vite-env.d.ts
+++ b/Engine/Source/Programs/Horde/HordeDashboard/src/vite-env.d.ts
@@ -1 +1,3 @@
+// Copyright Epic Games, Inc. All Rights Reserved.
+
 /// <reference types="vite/client" />
diff --git a/Engine/Source/Programs/Horde/HordeDashboard/types/react-flame-graph/index.d.ts b/Engine/Source/Programs/Horde/HordeDashboard/types/react-flame-graph/index.d.ts
--- a/Engine/Source/Programs/Horde/HordeDashboard/types/react-flame-graph/index.d.ts
+++ b/Engine/Source/Programs/Horde/HordeDashboard/types/react-flame-graph/index.d.ts
@@ -1 +1,3 @@
+// Copyright Epic Games, Inc. All Rights Reserved.
+
 declare module 'react-flame-graph'
\ No newline at end of file
diff --git a/Engine/Source/Programs/Horde/HordeDashboard/types/react-text-marquee/index.d.ts b/Engine/Source/Programs/Horde/HordeDashboard/types/react-text-marquee/index.d.ts
--- a/Engine/Source/Programs/Horde/HordeDashboard/types/react-text-marquee/index.d.ts
+++ b/Engine/Source/Programs/Horde/HordeDashboard/types/react-text-marquee/index.d.ts
@@ -1 +1,3 @@
+// Copyright Epic Games, Inc. All Rights Reserved.
+
 declare module 'react-text-marquee'
\ No newline at end of file
