Horde: Updating schema documentation

#rnx

diff --git a/Engine/Source/Programs/Horde/Docs/Config/Schema/Notification.md b/Engine/Source/Programs/Horde/Docs/Config/Schema/Notification.md
--- a/Engine/Source/Programs/Horde/Docs/Config/Schema/Notification.md
+++ b/Engine/Source/Programs/Horde/Docs/Config/Schema/Notification.md
@@ -17,6 +17,7 @@ Config for the notifications based on the streams
 Name | Description
 ---- | -----------
 `streams` | `string[]`<br>List of streams for this notification
+`enabledStreamTags` | `string[]`<br>List of enabled tags for a stream for this notification
 `jobNotifications` | [JobNotificationConfig](#jobnotificationconfig)`[]`<br>Configuration for sending experimental Job Summary Slack notifications
 
 ## JobNotificationConfig
