Avoid incrementing history when adding `stream` query when loading target report session
+ fix issue with stream filter not getting updated on navigation event

#rnx

diff --git a/Engine/Source/Programs/Horde/HordeDashboard/plugins/test-automation/testAutomationView.tsx b/Engine/Source/Programs/Horde/HordeDashboard/plugins/test-automation/testAutomationView.tsx
--- a/Engine/Source/Programs/Horde/HordeDashboard/plugins/test-automation/testAutomationView.tsx
+++ b/Engine/Source/Programs/Horde/HordeDashboard/plugins/test-automation/testAutomationView.tsx
@@ -25,6 +25,7 @@ TestDataVersionRegistrar.register(new TestDataV2);
 const StreamChooser: React.FC = observer(() => {
 
     handler.subscribeToFilter();
+    handler.subscribeToSearch();
 
     const stream = handler.filterState.stream;
 
@@ -60,6 +61,7 @@ const getTextFromTag = (item: ITag) => item.name;
 const TestChooser: React.FC = observer(() => {
     
     handler.subscribeToFilter();
+    handler.subscribeToSearch();
 
     const ctests: Set<string> = new Set<string>(handler.filterState.tests ?? []);
     const streamTests = handler.allStreamTestNames;
@@ -81,7 +83,7 @@ const TestChooser: React.FC = observer(() => {
         const isTagExist = closestTags.find((tag) => tag.key === lowerText);
         if (!isTagExist) closestTags.splice(0, 0, {key: lowerText, name: filterText});
         return closestTags;
-    }, [handler.queryLoading]);
+    }, [handler.queryLoading, handler.searchUpdated]);
 
     if (!streamTests.length) {
         return null;
@@ -119,6 +121,7 @@ const TestChooser: React.FC = observer(() => {
 const MetaKeyChooser: React.FC<{ onUpdateKey: (key: string) => void }> = observer(({ onUpdateKey }) => {
 
     handler.subscribeToFilter();
+    handler.subscribeToSearch();
 
    const cmetakeys: Set<string> = new Set(handler.filterState.metadata?.keys() ?? []);
    const streamMetadataKeys = handler.allStreamMetadataKeys;
@@ -358,7 +361,7 @@ export const TestAutomationView: React.FC = () => {
     const stepId = searchParams.get('step');
 
     useEffect(() => {
-        handler.initialize(searchParams.toString(), (search: string) => {setSearchParams(search)}).then(() => {setInitialized(true)});
+        handler.initialize(searchParams.toString(), (search: string, replace?: boolean) => {setSearchParams(search, {replace: replace})}).then(() => {setInitialized(true)});
         return () => {
             handler.clear();
             setInitialized(false); // for dev env
diff --git a/Engine/Source/Programs/Horde/HordeDashboard/plugins/test-automation/testData.ts b/Engine/Source/Programs/Horde/HordeDashboard/plugins/test-automation/testData.ts
--- a/Engine/Source/Programs/Horde/HordeDashboard/plugins/test-automation/testData.ts
+++ b/Engine/Source/Programs/Horde/HordeDashboard/plugins/test-automation/testData.ts
@@ -868,7 +868,7 @@ export class TestDataHandler {
         );
     }
 
-    async initialize(search: string, searchCallback?: (search: string) => void) {
+    async initialize(search: string, searchCallback?: (search: string, replace?: boolean) => void) {
         this.filterState = this.filterFromSearch(search);
         this.searchCallback = searchCallback;
         this.checkSearchCallback(search);
@@ -898,15 +898,17 @@ export class TestDataHandler {
         }
     }
 
-    setSearchParam(key: string, value: string) {
+    setSearchParam(key: string, value: string, replace?: boolean) {
         const previous = this.search.toString();
         this.search.set(key, null, value);
+        replace && this.queueHistoryReplace();
         this.checkSearchCallback(previous);
     }
 
-    removeSearchParam(key: string) {
+    removeSearchParam(key: string, replace?: boolean) {
         const previous = this.search.toString();
         this.search.delete(key);
+        replace && this.queueHistoryReplace();
         this.checkSearchCallback(previous);
     }
 
@@ -1015,6 +1017,12 @@ export class TestDataHandler {
             return;
         }
 
+        if (!state.stream && !!this.search.get('session')) {
+            // detect straight-to-session-view to then replace history instead of adding to it
+            // as in the case selectStream is called only because we want to load the stream related test data for comparison
+            this.queueHistoryReplace();
+        }
+
         state.stream = streamId;
 
         this.setFilterUpdated();
@@ -1124,6 +1132,16 @@ export class TestDataHandler {
         return testStreamRef!.getMatchingMeta(state.metadata).keys().toArray().sort((a, b) => a.orderedValues.join().localeCompare(b.orderedValues.join()));
     }
 
+    private queueHistoryReplace() {
+        this.nextHistoryReplace = true;
+    }
+
+    private popHistoryReplace(): boolean | undefined {
+        const replace = this.nextHistoryReplace;
+        this.nextHistoryReplace = undefined;
+        return replace;
+    }
+
     private updateFilter() {
         if (this.checkUpdateFilter()) {
             this.onSearchUpdated();
@@ -1133,7 +1151,7 @@ export class TestDataHandler {
     private onSearchUpdated() {
         this.setSearchUpdated();
         if (this.searchCallback) {
-            this.searchCallback(this.search.toString());
+            this.searchCallback(this.search.toString(), this.popHistoryReplace());
         }
     }
 
@@ -1483,10 +1501,12 @@ export class TestDataHandler {
         if (this.subPhaseQueryLoading) { }
     }
 
-    private searchCallback?: (search: string) => void;
+    private searchCallback?: (search: string, replace?: boolean) => void;
     private search: TestDataFilterSearchParams = new TestDataFilterSearchParams();
     filterState: TestDataState = {};
 
+    private nextHistoryReplace?: boolean;
+
     testStreams: Map<StreamId, TestStreamRef> = new Map();
 
     testTestData: Map<TestDataId, TestSessionDetails> = new Map();
