Differentiate better Skipped from Unspecified
+ isolate test-automation plugin theme colors from Horde default
+ move skipped at last position to phase outcome filter
+ add Interrupted phase dedicated navigation widget

#rnx

diff --git a/Engine/Source/Programs/Horde/HordeDashboard/plugins/test-automation/phaseHistoryGraph.ts b/Engine/Source/Programs/Horde/HordeDashboard/plugins/test-automation/phaseHistoryGraph.ts
--- a/Engine/Source/Programs/Horde/HordeDashboard/plugins/test-automation/phaseHistoryGraph.ts
+++ b/Engine/Source/Programs/Horde/HordeDashboard/plugins/test-automation/phaseHistoryGraph.ts
@@ -6,7 +6,7 @@ import dashboard, { StatusColor } from 'horde/backend/Dashboard';
 import { getHumanTime, getShortNiceTime } from 'horde/base/utilities/timeUtils';
 import { getHordeStyling } from 'horde/styles/Styles';
 import { TestOutcome } from "./api";
-import { getPhaseSessionStatusColor } from "./testAutomationCommon";
+import { getPhaseSessionStatusColor, getStatusColors } from "./testAutomationCommon";
 
 // Handle bad "@types/d3" types, fix if addressed upstream
 const _d3 = d3 as any;
@@ -47,7 +47,7 @@ export class PhaseHistoryGraph {
         const sessions = this.sessions;
         const width = 1000
 
-        const scolors = dashboard.getStatusColors();
+        const scolors = getStatusColors();
         const colors = new Map<TestOutcome, string>([
             [TestOutcome.Success, scolors.get(StatusColor.Success)!],
             [TestOutcome.Failure, scolors.get(StatusColor.Failure)!],
diff --git a/Engine/Source/Programs/Horde/HordeDashboard/plugins/test-automation/testAutomationCommon.tsx b/Engine/Source/Programs/Horde/HordeDashboard/plugins/test-automation/testAutomationCommon.tsx
--- a/Engine/Source/Programs/Horde/HordeDashboard/plugins/test-automation/testAutomationCommon.tsx
+++ b/Engine/Source/Programs/Horde/HordeDashboard/plugins/test-automation/testAutomationCommon.tsx
@@ -91,6 +91,23 @@ const getStyles = () => {
    return styles;
 }
 
+let _colors: Map<StatusColor, string> | undefined = undefined;
+
+export const getStatusColors = (): Map<StatusColor, string> => {
+
+    const dashboardStatusColors = dashboard.getStatusColors();
+
+    const colors = _colors ?? new Map<StatusColor, string>([
+        [StatusColor.Success, dashboardStatusColors.get(StatusColor.Success)!],
+        [StatusColor.Warnings, dashboardStatusColors.get(StatusColor.Warnings)!],
+        [StatusColor.Failure, dashboardStatusColors.get(StatusColor.Failure)!],
+        [StatusColor.Skipped,  dashboard.darktheme ? "#c5c4c3ff" : "#a1a1a1ff"],
+        [StatusColor.Unspecified, "#515151ff"]
+    ]);
+
+    return colors;
+}
+
 export const MultiOptionChooser: React.FC<{options: IComboBoxOption[], initialSelection: string[], updateKeys: (selectedKeys: string[]) => void, placeholder?: string, style?: any, disabled?: boolean }> = ({ options, initialSelection, updateKeys, placeholder, style, disabled }) => {
     const [selection, setSelection] = useState<string[]>(initialSelection);
     const optionsWithSelectAll = [...options];
@@ -171,7 +188,7 @@ export const SessionStatusBar = (session: TestSessionResult | TestSessionDetails
     const bDisplayBar = (metaFailedFactor + metaUnspecifiedFactor) > 0 && metaFailedFactor < 1 && metaUnspecifiedFactor < 1;
     if (!bDisplayBar) return undefined
 
-    const statusColors = dashboard.getStatusColors();
+    const statusColors = getStatusColors();
 
     const stack: StatusBarStack[] = [
         {
@@ -232,7 +249,7 @@ export const SessionStatusPie = (session: TestSessionResult, radius: number): JS
     const bDisplayBar = (metaFailedFactor + metaUnspecifiedFactor) > 0 && (metaFailedFactor + metaUnspecifiedFactor) < 1;
     if (!bDisplayBar) return;
 
-    const statusColors = dashboard.getStatusColors();
+    const statusColors = getStatusColors();
 
     const stack: StatusBarStack[] = [
         {
@@ -280,7 +297,7 @@ export const PhasesStatusPie = (phases: TestPhaseStatus[], radius: number): JSX.
     const unspecifiedFactor = Math.ceil(phasesUnspecifiedCount / totalCount * 100) / 100;
     const skippedFactor = Math.ceil(phasesSkippedCount / totalCount * 100) / 100;
 
-    const statusColors = dashboard.getStatusColors();
+    const statusColors = getStatusColors();
 
     const stack: StatusBarStack[] = [
         {
@@ -309,7 +326,7 @@ export const PhasesStatusPie = (phases: TestPhaseStatus[], radius: number): JSX.
 }
 
 export const SessionValues = (session: TestSessionResult, style?: any): JSX.Element => {
-    const statusColors = dashboard.getStatusColors();
+    const statusColors = getStatusColors();
     return <Text style={style}>
                 {!!session.phasesUnspecifiedCount && <span style={{color: statusColors.get(StatusColor.Unspecified)}}> Unspecified: {session.phasesUnspecifiedCount}</span>}
                 {!!session.phasesFailedCount && <span style={{color: statusColors.get(StatusColor.Failure)}}> Failed: {session.phasesFailedCount}</span>}
@@ -318,7 +335,7 @@ export const SessionValues = (session: TestSessionResult, style?: any): JSX.Elem
 }
 
 export const getPhaseSessionStatusColor = (phase: PhaseSessionResult | TestPhaseStatus) => {
-    const statusColors = dashboard.getStatusColors();
+    const statusColors = getStatusColors();
     let color = statusColors.get(StatusColor.Unspecified)!
     switch(phase.outcome) {
         case TestPhaseOutcome.Skipped:
@@ -336,7 +353,7 @@ export const getPhaseSessionStatusColor = (phase: PhaseSessionResult | TestPhase
 }
 
 export const getTestSessionStatusColor = (test: TestSessionResult | TestSessionDetails) => {
-    const statusColors = dashboard.getStatusColors();
+    const statusColors = getStatusColors();
     let color = statusColors.get(StatusColor.Success);
     if (test.outcome === TestOutcome.Failure) {
         color = statusColors.get(StatusColor.Failure);
@@ -446,7 +463,7 @@ export type ICheckListOption = IContextualMenuItem & {
 }
 
 export const getPhaseOutcomeOptions = (): ICheckListOption[] => {
-    const statusColors = dashboard.getStatusColors();
+    const statusColors = getStatusColors();
     return [
         {
             key: TestOutcome.Unspecified,
@@ -454,12 +471,6 @@ export const getPhaseOutcomeOptions = (): ICheckListOption[] => {
             checked: true,
             color: statusColors.get(StatusColor.Unspecified)
         },
-        {
-            key: TestOutcome.Skipped,
-            text: 'Skipped',
-            checked: true,
-            color: statusColors.get(StatusColor.Skipped)
-        },
         {
             key: TestOutcome.Failure,
             text: 'Failed',
@@ -472,6 +483,12 @@ export const getPhaseOutcomeOptions = (): ICheckListOption[] => {
             checked: true,
             color: statusColors.get(StatusColor.Success)
         },
+        {
+            key: TestOutcome.Skipped,
+            text: 'Skipped',
+            checked: true,
+            color: statusColors.get(StatusColor.Skipped)
+        },
     ]
 }
 
diff --git a/Engine/Source/Programs/Horde/HordeDashboard/plugins/test-automation/testAutomationEvents.tsx b/Engine/Source/Programs/Horde/HordeDashboard/plugins/test-automation/testAutomationEvents.tsx
--- a/Engine/Source/Programs/Horde/HordeDashboard/plugins/test-automation/testAutomationEvents.tsx
+++ b/Engine/Source/Programs/Horde/HordeDashboard/plugins/test-automation/testAutomationEvents.tsx
@@ -9,17 +9,16 @@ import dashboard, { StatusColor } from "horde/backend/Dashboard";
 import { getShortNiceTime, msecToElapsed } from "horde/base/utilities/timeUtils";
 import { DashboardPreference } from "horde/backend/Api";
 import { isBright } from "horde/base/utilities/colors";
-import { PhaseArtifactsModal, styles } from "./testAutomationCommon";
+import { PhaseArtifactsModal, styles, getStatusColors } from "./testAutomationCommon";
 
 let eventStyles: any = undefined;
 export const getEventStyling = () => {
     if (eventStyles) return eventStyles;
 
-    const colors = dashboard.getStatusColors();
+    const colors = getStatusColors();
     const errorColorPref = dashboard.getPreference(DashboardPreference.ColorError);
-    const errorHighlight = errorColorPref ? isBright(errorColorPref) ? "brightness(0.9)" : "brightness(1.2)" : "";
     const warningColorPref = dashboard.getPreference(DashboardPreference.ColorWarning);
-    const warningHighlight = warningColorPref ? isBright(warningColorPref) ? "brightness(0.9)" : "brightness(1.2)" : "";
+    const highlight = dashboard.darktheme ? "brightness(1.2)" : "brightness(0.9)";
 
     const eventClass = mergeStyles({
         fontSize: "12px",
@@ -49,7 +48,20 @@ export const getEventStyling = () => {
             selectors: {
                 ':active,:hover': {
                     backgroundColor: warningColorPref ? warningColorPref : "#E7C144",
-                    filter: warningHighlight
+                    filter: highlight
+                }
+            }
+        },
+        interruptedButton: {
+            minWidth: 16,
+            backgroundColor: colors.get(StatusColor.Unspecified)!,
+            borderStyle: "hidden",
+            color: "#FFFFFF",
+            selectors: {
+                ':active,:hover': {
+                    color: "#FFFFFF",
+                    backgroundColor: colors.get(StatusColor.Unspecified)!,
+                    filter: highlight
                 }
             }
         },
@@ -61,7 +73,7 @@ export const getEventStyling = () => {
                 ':active,:hover': {
                     color: "#FFFFFF",
                     backgroundColor: errorColorPref ? errorColorPref : "#DC3C37",
-                    filter: errorHighlight
+                    filter: highlight
                 }
             }
         },
@@ -77,7 +89,7 @@ export const getEventStyling = () => {
             borderColor: dashboard.darktheme ? "#4D4C4B" : "#6D6C6B",
             selectors: {
                 ':active,:hover': {
-                    filter: errorHighlight
+                    filter: highlight
                 }
             }
         },
diff --git a/Engine/Source/Programs/Horde/HordeDashboard/plugins/test-automation/testAutomationPhaseView.tsx b/Engine/Source/Programs/Horde/HordeDashboard/plugins/test-automation/testAutomationPhaseView.tsx
--- a/Engine/Source/Programs/Horde/HordeDashboard/plugins/test-automation/testAutomationPhaseView.tsx
+++ b/Engine/Source/Programs/Horde/HordeDashboard/plugins/test-automation/testAutomationPhaseView.tsx
@@ -123,7 +123,7 @@ const TestPhaseHeader: React.FC<{ phase: TestPhaseStatus, expanded: boolean, tog
                         <Stack style={{ paddingTop: 1, paddingRight: 4}}>
                             <FontIcon style={{ fontSize: 13 }} iconName={expanded ? "ChevronDown" : "ChevronRight"} />
                         </Stack>
-                        <Text style={{fontSize: 13, fontWeight: 600, whiteSpace: 'nowrap', overflow: 'hidden', textOverflow: 'ellipsis', cursor: "text"}} onClick={(ev) => ev.stopPropagation()}>{phase.name}</Text>
+                        <Text style={{width: 'calc(100vw - 150px)', fontSize: 13, fontWeight: 600, whiteSpace: 'nowrap', overflow: 'hidden', textOverflow: 'ellipsis', cursor: "text"}} onClick={(ev) => ev.stopPropagation()}>{phase.name}</Text>
                     </Stack>
                     <Stack horizontal verticalAlign="center">
                         <Stack style={{ paddingLeft: 5, paddingRight: 4 }}>
@@ -315,6 +315,7 @@ export const TestPhasesView: React.FC<{ testDataId: string, handler: TestDataHan
     const [jobArtifacts, setJobArtifacts] = useState<GetArtifactResponse>();
     const selectedPhase = useRef<TestPhaseStatus>();
     const failureCursor = useRef<CircularList<string>>();
+    const interruptedCursor = useRef<CircularList<string>>();
     const warningCursor = useRef<CircularList<string>>();
     const phaseCursor = useRef<CircularList<string>>();
     const listRef = useRef<List>();
@@ -329,6 +330,9 @@ export const TestPhasesView: React.FC<{ testDataId: string, handler: TestDataHan
         const failures = phases.filter((p) => p.outcome !== TestPhaseOutcome.Success && p.outcome !== TestPhaseOutcome.NotRun && p.outcome !== TestPhaseOutcome.Skipped).map((p) => p.key) ?? [];
         failureCursor.current = new CircularList(failures);
         failureCursor.current.setCursor(selectedPhase.current?.key);
+        const interrupted = phases.filter((p) => p.outcome === TestPhaseOutcome.Interrupted).map((p) => p.key) ?? [];
+        interruptedCursor.current = new CircularList(interrupted);
+        interruptedCursor.current.setCursor(selectedPhase.current?.key);
         phaseCursor.current = new CircularList(phases.map((p) => p.key));
         phaseCursor.current.setCursor(selectedPhase.current?.key);
         const warnings = phases.filter((p) => p.outcome !== TestPhaseOutcome.Success && p.hasWarning).map((p) => p.key) ?? [];
@@ -402,15 +406,15 @@ export const TestPhasesView: React.FC<{ testDataId: string, handler: TestDataHan
 
     const hordeTheme = getHordeTheme();
     const eventStyles = getEventStyling();
-    const buttonNoneTextColor = dashboard.darktheme ? "#949898" : "#616e85";
-    const errorColorPref = dashboard.getPreference(DashboardPreference.ColorError);
-    const failureButtonTextColor = errorColorPref && isBright(errorColorPref) ? "#2d3f5f" : "#F9F9FB";
-    const warningColorPref = dashboard.getPreference(DashboardPreference.ColorWarning);
-    const warningButtonTextColor = warningColorPref && isBright(warningColorPref) ? "#2d3f5f" : "#F9F9FB";
+    const buttonNoneTextColor = dashboard.darktheme ? "#616e85" : "#949898";
+    const failureButtonTextColor = "#F9F9FB";
+    const interruptedButtonTextColor = failureButtonTextColor;
+    const warningButtonTextColor = failureButtonTextColor;
 
     const failureEnabled = !!failureCursor.current?.length;
     const failureText = failureEnabled ? failureCursor.current!.length.toString() : "No";
 
+    const interruptedEnabled = !!interruptedCursor.current?.length;
     const warningEnabled = !!warningCursor.current?.length;
 
     const testId = sessionDetails?.test?.id ?? testDataId;
@@ -500,6 +504,30 @@ export const TestPhasesView: React.FC<{ testDataId: string, handler: TestDataHan
                                 }
                             </Stack>
 
+                            {interruptedEnabled &&
+                                <Stack horizontal verticalAlign='center' style={{ height: '100%' }}>
+                                    <Stack horizontal verticalAlign='center' style={{position: 'relative'}}>
+                                        <TooltipHost content='Navigate to next interrupted'>
+                                            <DefaultButton
+                                                className={eventStyles.interruptedButton}
+                                                style={{ color: interruptedButtonTextColor, height: 28, whiteSpace: 'nowrap', fontSize: 12, paddingLeft: 8, paddingRight: 8, borderRadius: '4px 0 0 4px' }}
+                                                text={`${interruptedCursor.current?.length}${(failureCursor.current?.length ?? 0) === 0 ? " Interrupted" : ""}`}
+                                                onClick={() => interruptedCursor.current && handler.setSearchParam('phase', interruptedCursor.current.next())} >
+                                                <Icon style={{ fontSize: 19, paddingLeft: 4, paddingRight: 0 }} iconName='ChevronDown' />
+                                            </DefaultButton>
+                                        </TooltipHost>
+                                        <div style={{backgroundColor: interruptedButtonTextColor, height: 18, width: 1, position: 'absolute', right: -0.5, zIndex: 1 }} />
+                                    </Stack>
+                                    <TooltipHost content={'Navigate to previous interrupted'}>
+                                        <IconButton
+                                            className={eventStyles.interruptedButton}
+                                            style={{ color: interruptedButtonTextColor, height: 28, fontSize: 12, padding: 2, borderRadius: '0 4px 4px 0' }}
+                                            iconProps={{ iconName: 'ChevronUp' }}
+                                            onClick={() => interruptedCursor.current && handler.setSearchParam('phase', interruptedCursor.current.back())} />
+                                    </TooltipHost>
+                                </Stack>
+                            }
+
                             {warningEnabled &&
                                 <Stack horizontal verticalAlign='center' style={{ height: '100%' }}>
                                     <Stack horizontal verticalAlign='center' style={{position: 'relative'}}>
@@ -507,7 +535,7 @@ export const TestPhasesView: React.FC<{ testDataId: string, handler: TestDataHan
                                             <DefaultButton
                                                 className={eventStyles.warningButton}
                                                 style={{ color: warningButtonTextColor, height: 28, whiteSpace: 'nowrap', fontSize: 12, paddingLeft: 8, paddingRight: 8, borderRadius: '4px 0 0 4px' }}
-                                                text={`${warningCursor.current?.length}${(failureCursor.current?.length ?? 0) === 0 ? (warningCursor.current?.length ?? 0) > 1 ? " Warnings" : " Warning" : ""}`}
+                                                text={`${warningCursor.current?.length}${(failureCursor.current?.length ?? interruptedCursor.current?.length ?? 0) === 0 ? (warningCursor.current?.length ?? 0) > 1 ? " Warnings" : " Warning" : ""}`}
                                                 onClick={() => warningCursor.current && handler.setSearchParam('phase', warningCursor.current.next())} >
                                                 <Icon style={{ fontSize: 19, paddingLeft: 4, paddingRight: 0 }} iconName='ChevronDown' />
                                             </DefaultButton>
diff --git a/Engine/Source/Programs/Horde/HordeDashboard/plugins/test-automation/testAutomationSummary.tsx b/Engine/Source/Programs/Horde/HordeDashboard/plugins/test-automation/testAutomationSummary.tsx
--- a/Engine/Source/Programs/Horde/HordeDashboard/plugins/test-automation/testAutomationSummary.tsx
+++ b/Engine/Source/Programs/Horde/HordeDashboard/plugins/test-automation/testAutomationSummary.tsx
@@ -9,7 +9,7 @@ import dashboard, { StatusColor } from "horde/backend/Dashboard";
 import { getShortNiceTime } from "horde/base/utilities/timeUtils";
 import { getHordeStyling } from "horde/styles/Styles";
 import { projectStore } from 'horde/backend/ProjectStore';
-import { SessionStatusBar, SessionValues, StatusBar, StatusBarStack, MultiOptionChooser, getTestSessionStatusColor, getStreamOptions } from "./testAutomationCommon";
+import { SessionStatusBar, SessionValues, StatusBar, StatusBarStack, MultiOptionChooser, getTestSessionStatusColor, getStreamOptions, getStatusColors } from "./testAutomationCommon";
 import { TestHistoryGraph } from "./testHistoryGraph";
 
 type SessionCallback = (session: TestSessionResult) => void;
@@ -97,7 +97,7 @@ export const TestHistoryView: React.FC<{ test: TestNameRef, handler: TestDataHan
         .map((meta) => ({key: meta.id, text: meta.getValuesExcept(commonMetaKeys).join(' / ')}))
         .toArray().sort((a, b) => a.text.localeCompare(b.text));
 
-    const scolors = dashboard.getStatusColors();
+    const scolors = getStatusColors();
     const streamViews = (!!availableStreams.length? [streams.current, streams.toCompare] : [streams.current]).map(stream => {
 
         const streamName = projectStore.streamById(stream)?.fullname ?? stream;
@@ -202,7 +202,7 @@ const TestSummaryCard: React.FC<{ test: TestNameRef, sessions: [MetadataRef, Tes
     const [expanded, setExpanded] = useState(false);
  
     const { hordeClasses, modeColors } = getHordeStyling();
-    const statusColors = dashboard.getStatusColors();
+    const statusColors = getStatusColors();
  
     const styles = mergeStyleSets({
        metaitem: {
diff --git a/Engine/Source/Programs/Horde/HordeDashboard/plugins/test-automation/testHistoryGraph.ts b/Engine/Source/Programs/Horde/HordeDashboard/plugins/test-automation/testHistoryGraph.ts
--- a/Engine/Source/Programs/Horde/HordeDashboard/plugins/test-automation/testHistoryGraph.ts
+++ b/Engine/Source/Programs/Horde/HordeDashboard/plugins/test-automation/testHistoryGraph.ts
@@ -5,7 +5,7 @@ import { TestDataHandler, TestNameRef, MetadataRef, TestMetaStatus, TestSessionR
 import dashboard, { StatusColor } from 'horde/backend/Dashboard';
 import { getHumanTime, getShortNiceTime, msecToElapsed } from 'horde/base/utilities/timeUtils';
 import { getHordeStyling } from 'horde/styles/Styles';
-import { SessionValues } from "./testAutomationCommon";
+import { getStatusColors, SessionValues } from "./testAutomationCommon";
 import { renderToString } from "react-dom/server"
 import { TestOutcome } from "./api";
 
@@ -69,7 +69,7 @@ export class TestHistoryGraph {
         const sessions = this.sessions;
         const width = 1000
 
-        const scolors = dashboard.getStatusColors();
+        const scolors = getStatusColors();
         const colors = new Map<TestOutcome, string>([
             [TestOutcome.Success, scolors.get(StatusColor.Success)!],
             [TestOutcome.Failure, scolors.get(StatusColor.Failure)!],
@@ -388,7 +388,7 @@ const SessionPieChart = (container: SVGGElement, session: TestSessionResult, rad
     const metaFailedFactor = Math.ceil(session.phasesFailedCount / (sessionTotalCount || 1) * 10) / 10;
     const metaUnspecifiedFactor = Math.ceil(session.phasesUnspecifiedCount / (sessionTotalCount || 1) * 10) / 10;
 
-    const statusColors = dashboard.getStatusColors();
+    const statusColors = getStatusColors();
 
     const stack: ItemStack[] = [
         {
@@ -437,7 +437,7 @@ const SessionConcentricCircle = (container: SVGGElement, session: TestSessionRes
     let metaSuccessFactor = 1 - (metaFailedFactor + metaUnspecifiedFactor);
     metaSuccessFactor = metaSuccessFactor < 0? 0 : metaSuccessFactor
 
-    const statusColors = dashboard.getStatusColors();
+    const statusColors = getStatusColors();
 
     const stack: ItemStack[] = [
         {
