Horde Dashboard: Reduce JobDetails V1

diff --git a/Engine/Source/Programs/Horde/HordeDashboard/src/App.tsx b/Engine/Source/Programs/Horde/HordeDashboard/src/App.tsx
--- a/Engine/Source/Programs/Horde/HordeDashboard/src/App.tsx
+++ b/Engine/Source/Programs/Horde/HordeDashboard/src/App.tsx
@@ -19,7 +19,7 @@ import { DeviceView } from './components/devices/DeviceView';
 import { ErrorDialog } from './components/error/ErrorDialog';
 import { errorDialogStore } from './components/error/ErrorStore';
 import { JobRedirector } from './components/JobRedirector';
-import { LogView } from './components/LogView';
+import { LogView } from './components/logs/LogView';
 import { NoticeView } from './components/NoticeView';
 import { PerforceServerView } from './components/PerforceView';
 import { PoolsView } from './components/PoolsView';
@@ -45,7 +45,6 @@ import hordePlugins from './legacyPlugins';
 import { preloadFonts } from './styles/Styles';
 import { darkTheme } from './styles/darkTheme';
 import { lightTheme } from './styles/lightTheme';
-import { ErrorBar } from './components/error/ErrorBar';
 import { BuildHealthView } from './components/buildhealth/BuildhealthView';
 
 let router: any;
diff --git a/Engine/Source/Programs/Horde/HordeDashboard/src/backend/Api.ts b/Engine/Source/Programs/Horde/HordeDashboard/src/backend/Api.ts
--- a/Engine/Source/Programs/Horde/HordeDashboard/src/backend/Api.ts
+++ b/Engine/Source/Programs/Horde/HordeDashboard/src/backend/Api.ts
@@ -223,6 +223,13 @@ export type JobsTabColumnData = GetJobsTabColumnResponse
 export type TimingInfo = GetTimingInfoResponse
 export type StepTimingInfo = GetStepTimingInfoResponse
 export type SubscriptionData = GetSubscriptionResponse
+export type ArtifactData = GetArtifactResponse;
+
+export type JobLabel = {
+    stateResponse: GetLabelStateResponse;
+    timing?: GetLabelTimingInfoResponse;
+    default: boolean;
+}
 
 export type IssueData = GetIssueResponse & {
 	events?: GetLogEventResponse[];
diff --git a/Engine/Source/Programs/Horde/HordeDashboard/src/backend/JobDetails.ts b/Engine/Source/Programs/Horde/HordeDashboard/src/backend/JobDetails.ts
deleted file mode 100644
diff --git a/Engine/Source/Programs/Horde/HordeDashboard/src/backend/JobEventHandler.ts b/Engine/Source/Programs/Horde/HordeDashboard/src/backend/JobEventHandler.ts
deleted file mode 100644
diff --git a/Engine/Source/Programs/Horde/HordeDashboard/src/base/utilities/logUtils.ts b/Engine/Source/Programs/Horde/HordeDashboard/src/base/utilities/logUtils.ts
new file mode 100644
--- /dev/null
+++ b/Engine/Source/Programs/Horde/HordeDashboard/src/base/utilities/logUtils.ts
@@ -0,0 +1,220 @@
+// Copyright Epic Games, Inc. All Rights Reserved.
+
+import { JobStepBatchState, JobStepState, JobStepError, JobStepOutcome } from "horde/backend/Api";
+import { getBatchText } from "horde/components/JobDetailCommon";
+import { JobDetailsV2 } from "horde/components/jobDetailsV2/JobDetailsViewCommon";
+import { getBatchInitElapsed, getNiceTime, getStepElapsed, getStepETA, getStepTimingDelta, getStepFinishTime } from "./timeUtils";
+import { NavigateFunction } from "react-router-dom";
+
+/**
+ * Returns batch summary markdown including information like start times, run times, current state, etc.
+ * @param jobDetails Job details related to the batch.
+ * @param batchId The ID of the batch to summarize.
+ */
+export const getBatchSummaryMarkdown = (jobDetails: JobDetailsV2, batchId: string): string => {
+
+    const batch = jobDetails.batches.find(b => b.id === batchId);
+
+    if (!batch) {
+        return "";
+    }
+
+    const agentType = batch?.agentType ?? "";
+    const agentPool = jobDetails.stream?.agentTypes[agentType!]?.pool ?? "";
+
+    let summaryText = "";
+    if (batch && batch.agentId) {
+
+        if (batch.state === JobStepBatchState.Running || batch.state === JobStepBatchState.Complete) {
+            let duration = "";
+            if (batch.startTime) {
+                duration = getBatchInitElapsed(batch);
+                summaryText = `Batch started ${getNiceTime(batch.startTime)} and ${batch.finishTime ? "ran" : "is running"} on [${batch.agentId}](?batch=${batchId}&agentId=${encodeURIComponent(batch.agentId)}) with a setup time of ${duration}.`;
+            }
+        }
+    }
+
+    if (!summaryText) {
+        summaryText = getBatchText({ batch: batch, agentId: batch?.agentId, agentType: agentType, agentPool: agentPool }) ?? "";
+
+        if (summaryText && batch.state === JobStepBatchState.Starting) {
+            summaryText += " batch is the starting state";
+        }
+
+        if (summaryText && batch.state === JobStepBatchState.Stopping) {
+            summaryText += " batch is the stopping state";
+        }
+
+        if (!summaryText) {
+            summaryText = batch?.agentId ?? (batch?.state ?? "Batch is unassigned");
+        }
+    }
+
+    return summaryText;
+
+}
+
+/**
+ * Returns step summary markdown including information about creation, retries, run times, etc.
+ * @param jobDetails Job details related to the step.
+ * @param stepId The id of the step to summarize.
+ * @todo Resolve redundancy with near-identical function in StepDetailSummary.tsx.
+ */
+export const getStepSummaryMarkdown = (jobDetails: JobDetailsV2, stepId: string): string => {
+
+    const step = jobDetails.stepById(stepId)!;
+    const batch = jobDetails.batchByStepId(stepId);
+
+    if (!step || !batch) {
+        return "";
+    }
+
+    const duration = getStepElapsed(step);
+    let eta = getStepETA(step, jobDetails.jobData!);
+
+    const text: string[] = [];
+
+    if (jobDetails.jobData) {
+        text.push(`Job created by ${jobDetails.jobData.startedByUserInfo ? jobDetails.jobData.startedByUserInfo.name : "scheduler"}`);
+    }
+
+    const batchText = () => {
+
+        if (!batch) {
+            return undefined;
+        }
+
+        const agentType = batch?.agentType;
+        const agentPool = jobDetails.stream?.agentTypes[agentType!]?.pool;
+        return getBatchText({ batch: batch, agentType: agentType, agentPool: agentPool });
+
+    };
+
+    if (step.retriedByUserInfo) {
+        const retryId = jobDetails.getRetryStepId(step.id);
+        if (retryId) {
+            text.push(`[Step was retried by ${step.retriedByUserInfo.name}](/job/${jobDetails.jobId!}?step=${retryId})`);
+        } else {
+            text.push(`Step was retried by ${step.retriedByUserInfo.name}`);
+        }
+    }
+
+    if (step.abortRequested || step.state === JobStepState.Aborted) {
+        eta.display = eta.server = "";
+        let aborted = "";
+        if (step.abortedByUserInfo) {
+            aborted = "This step was canceled";
+            aborted += ` by ${step.abortedByUserInfo.name}.`;
+        } else if (jobDetails.jobData?.abortedByUserInfo) {
+            aborted = "The job was canceled";
+            aborted += ` by ${jobDetails.jobData?.abortedByUserInfo.name}.`;
+        } else {
+            aborted = "This step was canceled by Horde";
+
+            if (step.error === JobStepError.TimedOut) {
+                aborted = "The step was canceled due to reaching the maximum run time limit";
+            }
+        }
+        text.push(aborted);
+    } else if (step.state === JobStepState.Skipped) {
+        eta.display = eta.server = "";
+        text.push("The step was skipped");
+    } else if (step.state === JobStepState.Ready || step.state === JobStepState.Waiting) {
+
+        text.push(batchText() ?? `The step is pending in ${step.state} state`);
+    }
+
+    if (batch?.agentId) {
+
+        if (step.startTime) {
+            const str = getNiceTime(step.startTime);
+            text.push(`Step started on ${str}`);
+        }
+
+        let runningText = `${step.finishTime ? "Ran" : "Running"} on [${batch.agentId}](?step=${stepId}&agentId=${encodeURIComponent(batch.agentId)})`;
+
+        if (duration) {
+            runningText += ` for ${duration.trim()}`;
+        }
+
+        if (step.finishTime && (step.outcome === JobStepOutcome.Success || step.outcome === JobStepOutcome.Warnings)) {
+            const delta = getStepTimingDelta(step);
+            if (delta) {
+                runningText += " " + delta;
+            }
+        }
+
+        if (runningText) {
+            text.push(runningText);
+        }
+
+    } else {
+        if (!step.abortRequested) {
+            text.push(batchText() ?? "Step does not have a batch.");
+        }
+    }
+
+    if (eta.display) {
+        text.push(`Expected to finish around ${eta.display}.`);
+    }
+
+    const finish = getStepFinishTime(step).display;
+
+    if (finish && step.state !== JobStepState.Aborted) {
+
+        let outcome = "";
+        if (step.outcome === JobStepOutcome.Success) {
+            outcome += `Completed at ${finish}`;
+        }
+        if (step.outcome === JobStepOutcome.Failure)
+            outcome += `Completed with errors at ${finish}.`;
+        if (step.outcome === JobStepOutcome.Warnings)
+            outcome += `Completed with warnings at ${finish}.`;
+
+        if (outcome) {
+            text.push(outcome);
+        }
+    }
+
+    if (!text.length) {
+        text.push(`Step is in ${step.state} state.`);
+    }
+
+    return text.join(".&nbsp;&nbsp;");
+
+}
+
+/**
+ * Returns the value of any url parameter matching lineIndex or lineindex. Returns undefined if none found.
+ */
+export function getQueryLine(): number | undefined {
+
+   const search = new URLSearchParams(window.location.search);
+
+   if (search.get("lineindex")) {
+      return parseInt(search.get("lineindex")!) + 1;
+   }
+
+   if (search.get("lineIndex")) {
+      return parseInt(search.get("lineIndex")!);
+   }
+
+   return undefined;
+
+}
+
+/**
+ * Updates the lineIndex query parameter in the url.
+ * @param line The line number to set the query parameter with.
+ * @param navigate The navigate function from react router should be passed to this function.
+ */
+export function updateLineQuery(line: number, navigate: NavigateFunction) {
+
+   const search = new URLSearchParams(window.location.search);
+   search.set("lineIndex", line.toString());
+   search.delete("lineindex");
+   const url = `${window.location.pathname}?` + search.toString();
+
+   navigate(url, { replace: true })
+
+}
\ No newline at end of file
diff --git a/Engine/Source/Programs/Horde/HordeDashboard/src/base/utilities/promiseUtils.ts b/Engine/Source/Programs/Horde/HordeDashboard/src/base/utilities/promiseUtils.ts
new file mode 100644
--- /dev/null
+++ b/Engine/Source/Programs/Horde/HordeDashboard/src/base/utilities/promiseUtils.ts
@@ -0,0 +1,15 @@
+// Copyright Epic Games, Inc. All Rights Reserved.
+
+/**
+ * Waits for a potentially asynchronous value to not be null or undefined.
+ * @param accessor A simple function that returns the asynchronous value to keep track of.
+ * @param interval How often to check the asynchronous value.
+ * @returns The value of the property when it is anything but undefined or null.
+ */
+export const waitFor = async <T>(accessor: () => T, interval = 100): Promise<T> => {
+    let value: T;
+    while((value = accessor()) == null) {
+        await new Promise(resolve => setTimeout(resolve, interval));
+    }
+    return value;
+}
\ No newline at end of file
diff --git a/Engine/Source/Programs/Horde/HordeDashboard/src/base/utilities/timeUtils.ts b/Engine/Source/Programs/Horde/HordeDashboard/src/base/utilities/timeUtils.ts
--- a/Engine/Source/Programs/Horde/HordeDashboard/src/base/utilities/timeUtils.ts
+++ b/Engine/Source/Programs/Horde/HordeDashboard/src/base/utilities/timeUtils.ts
@@ -1,8 +1,7 @@
 // Copyright Epic Games, Inc. All Rights Reserved.
 
 import moment, { Moment } from "moment-timezone";
-import { BatchData, GetAgentLeaseResponse, JobData, LabelState, StepData, TimingInfo } from "../../backend/Api";
-import { JobLabel } from "../../backend/JobDetails";
+import { BatchData, GetAgentLeaseResponse, JobData, JobLabel, LabelState, StepData, TimingInfo } from "../../backend/Api";
 import dashboard from "../../backend/Dashboard";
 
 export const serverTimeZone = "UTC";
diff --git a/Engine/Source/Programs/Horde/HordeDashboard/src/components/AuditLog.tsx b/Engine/Source/Programs/Horde/HordeDashboard/src/components/AuditLog.tsx
--- a/Engine/Source/Programs/Horde/HordeDashboard/src/components/AuditLog.tsx
+++ b/Engine/Source/Programs/Horde/HordeDashboard/src/components/AuditLog.tsx
@@ -16,7 +16,7 @@ import { BreadcrumbItem, Breadcrumbs } from "./Breadcrumbs";
 import { DateTimeRange } from "./DateTimeRange";
 import { HistoryModal } from "./agents/HistoryModal";
 import { IssueModalV2 } from "./IssueViewV2";
-import { getLogStyles, logMetricNormal } from "./LogStyle";
+import { getLogStyles, logMetricNormal } from "./logs/LogStyle";
 import { TopNav } from "./TopNav";
 import { getHordeStyling } from "../styles/Styles";
 import { getHordeTheme } from "../styles/theme";
diff --git a/Engine/Source/Programs/Horde/HordeDashboard/src/components/AutomationReport.tsx b/Engine/Source/Programs/Horde/HordeDashboard/src/components/AutomationReport.tsx
--- a/Engine/Source/Programs/Horde/HordeDashboard/src/components/AutomationReport.tsx
+++ b/Engine/Source/Programs/Horde/HordeDashboard/src/components/AutomationReport.tsx
@@ -14,7 +14,7 @@ import { projectStore } from "../backend/ProjectStore";
 import { getNiceTime, msecToElapsed } from "../base/utilities/timeUtils";
 import { getHordeStyling } from "../styles/Styles";
 import { getHordeTheme } from "../styles/theme";
-import { renderLine } from "./LogRender";
+import { renderLine } from "./logs/LogRender";
 
 
 let _styles: any;
diff --git a/Engine/Source/Programs/Horde/HordeDashboard/src/components/IssueViewV2.tsx b/Engine/Source/Programs/Horde/HordeDashboard/src/components/IssueViewV2.tsx
--- a/Engine/Source/Programs/Horde/HordeDashboard/src/components/IssueViewV2.tsx
+++ b/Engine/Source/Programs/Horde/HordeDashboard/src/components/IssueViewV2.tsx
@@ -18,7 +18,7 @@ import { getHordeStyling } from "../styles/Styles";
 import { getHordeTheme } from "../styles/theme";
 import { errorDialogStore } from "./error/ErrorStore";
 import { useQuery } from "horde/base/utilities/hooks";
-import { renderLine } from "./LogRender";
+import { renderLine } from "./logs/LogRender";
 import { UserSelect } from "./UserSelect";
 
 const smallThreshhold = 1100;
diff --git a/Engine/Source/Programs/Horde/HordeDashboard/src/components/JobLabelCallout.tsx b/Engine/Source/Programs/Horde/HordeDashboard/src/components/JobLabelCallout.tsx
deleted file mode 100644
diff --git a/Engine/Source/Programs/Horde/HordeDashboard/src/components/JobSearchSimple.tsx b/Engine/Source/Programs/Horde/HordeDashboard/src/components/JobSearchSimple.tsx
--- a/Engine/Source/Programs/Horde/HordeDashboard/src/components/JobSearchSimple.tsx
+++ b/Engine/Source/Programs/Horde/HordeDashboard/src/components/JobSearchSimple.tsx
@@ -18,7 +18,6 @@ import { getHordeTheme } from "../styles/theme";
 import { getHordeStyling } from "../styles/Styles";
 import { StreamChooser } from "./projects/StreamChooser";
 import { getLabelColor } from "horde/styles/colors";
-import { CalloutController, JobLabelCallout } from "./JobLabelCallout";
 import LZString from "lz-string"
 import { errorBarStore } from "./error/ErrorStore";
 import { useQuery } from "horde/base/utilities/hooks";
@@ -57,7 +56,6 @@ export const JobSearchSimpleModal: React.FC<{ streamIdIn: string, changeListIn?:
    const [searchResults, setSearchResults] = useState<{items: JobItem[], groups?: IGroup[], querying: boolean}>({items: [], querying: false});
    const query = useQuery();
    const [templateData, setTemplateData] = useState<{ streamId: string; templates: GetTemplateRefResponse[] } | undefined>(undefined);
-   const [controller] = useState(new CalloutController());
    const location = useLocation();
    const [isFirstRender, setIsFirstRender] = useState(true);
 
@@ -649,11 +647,7 @@ export const JobSearchSimpleModal: React.FC<{ streamIdIn: string, changeListIn?:
       return <Stack>
          <div id={target} className="horde-no-darktheme">
             <Link to={url}><Stack className={hordeClasses.badgeNoIcon}><DefaultButton key={label.dashboardName} style={{ backgroundColor: color.primaryColor }} text={label.dashboardName}
-               onMouseOver={(ev) => {
-                  ev.stopPropagation();
-                  controller.setState({ target: `#${target}`, label: label, jobId: item.job.id })
-               }}
-               onMouseLeave={() => controller.setState(undefined)}
+               onMouseOver={(ev) => ev.stopPropagation()}
                onMouseMove={(ev) => ev.stopPropagation()}>
                {!!color.secondaryColor && <div style={{
                   borderLeft: "10px solid transparent",
@@ -709,7 +703,7 @@ export const JobSearchSimpleModal: React.FC<{ streamIdIn: string, changeListIn?:
          return <JobLabel key={`label_${item.job.id}_${label.dashboardCategory ?? ""}_${label.dashboardName ?? ""}_${key++}`} item={item} label={label} />;
       });
 
-      return (<Stack horizontalAlign="center" onMouseMove={() => controller.setState(undefined)}>
+      return (<Stack horizontalAlign="center">
          <Stack wrap horizontal horizontalAlign="center" tokens={{ childrenGap: 4 }} styles={{ root: { paddingTop: 2, width: labelColumnWidth } }}>
             {buttons}
          </Stack>
@@ -838,7 +832,6 @@ export const JobSearchSimpleModal: React.FC<{ streamIdIn: string, changeListIn?:
       onDismiss={() => { query.clear(["tab"]); onClose(); }}
    >
       {searchResults.querying && <StatusModal text={"Finding Jobs"} />}
-      <JobLabelCallout controller={controller} />
       <Stack styles={{ root: { paddingTop: 8, paddingLeft: 24, paddingRight: 12, paddingBottom: 8 } }}>
          <Stack tokens={{ childrenGap: 12 }}>
             <Stack horizontal styles={{ root: { padding: 0 } }}>
@@ -986,7 +979,7 @@ export const JobSearchSimpleModal: React.FC<{ streamIdIn: string, changeListIn?:
                      </Stack>}
 
                      {!!searchResults.items.length && <Stack>
-                        <div style={{ overflowY: 'auto', overflowX: 'hidden', height: "700px" }} data-is-scrollable={true} onScroll={() => { controller.setState(undefined, true) }}>
+                        <div style={{ overflowY: 'auto', overflowX: 'hidden', height: "700px" }} data-is-scrollable={true} >
                            <Stack tokens={{ childrenGap: 12 }} style={{ paddingRight: 12 }}>
                               <DetailsList
                                  compact={true}
diff --git a/Engine/Source/Programs/Horde/HordeDashboard/src/components/JobView.tsx b/Engine/Source/Programs/Horde/HordeDashboard/src/components/JobView.tsx
--- a/Engine/Source/Programs/Horde/HordeDashboard/src/components/JobView.tsx
+++ b/Engine/Source/Programs/Horde/HordeDashboard/src/components/JobView.tsx
@@ -14,7 +14,6 @@ import { displayTimeZone } from '../base/utilities/timeUtils';
 import { getJobStateColor, getLabelColor } from '../styles/colors';
 import { ChangeButton } from './ChangeButton';
 import { useQuery } from "horde/base/utilities/hooks";
-import { CalloutController, JobLabelCallout } from './JobLabelCallout';
 import { JobOperationsContextMenu } from './JobOperationsContextMenu';
 import { getHordeStyling } from '../styles/Styles';
 
@@ -177,7 +176,7 @@ const sortJobGroups = (items: JobItem[]): [JobItem[], JobGroup[]] => {
 
 const jobHandler = new JobHandler();
 
-const JobList: React.FC<{ tab: string; filter: JobFilterSimple, controller: CalloutController }> = observer(({ tab, filter, controller }) => {
+const JobList: React.FC<{ tab: string; filter: JobFilterSimple }> = observer(({ tab, filter }) => {
 
    const { streamId } = useParams<{ streamId: string }>();
    const query = useQuery();
@@ -388,11 +387,7 @@ const JobList: React.FC<{ tab: string; filter: JobFilterSimple, controller: Call
          <div id={target}>
             <Link to={url}><Stack className={hordeClasses.badgeNoIcon}>
                <DefaultButton key={label.dashboardName} style={{ backgroundColor: color.primaryColor }} text={label.dashboardName}
-                  onMouseOver={(ev) => {
-                     ev.stopPropagation();
-                     controller.setState({ target: `#${target}`, label: label, jobId: item.job.id })
-                  }}
-                  onMouseLeave={() => controller.setState(undefined)}
+                  onMouseOver={(ev) => ev.stopPropagation()}
                   onMouseMove={(ev) => ev.stopPropagation()}>
                   {!!color.secondaryColor && <div style={{
                      borderLeft: "10px solid transparent",
@@ -433,7 +428,7 @@ const JobList: React.FC<{ tab: string; filter: JobFilterSimple, controller: Call
       });
 
       return (
-         <Stack horizontalAlign="center" onMouseMove={() => controller.setState(undefined)}>
+         <Stack horizontalAlign="center">
             <Stack wrap horizontal horizontalAlign="center" tokens={{ childrenGap: 4 }} styles={{ root: { paddingTop: 2, width: column?.maxWidth } }}>
                {buttons}
             </Stack>
@@ -582,7 +577,7 @@ const JobList: React.FC<{ tab: string; filter: JobFilterSimple, controller: Call
    return (
       <Stack tokens={{ childrenGap: 0 }} className={detailClasses.detailsRow}>
          <div className={detailClasses.container} style={{ width: "100%", height: 'calc(100vh - 240px)', position: 'relative', marginTop: 0 }}>
-            {<ScrollablePane scrollbarVisibility={ScrollbarVisibility.always} onScroll={() => { controller.setState(undefined, true) }} style={{ overflow: "visible" }}>
+            {<ScrollablePane scrollbarVisibility={ScrollbarVisibility.always} style={{ overflow: "visible" }}>
                {renderItems.length > 0 &&
                   <Stack style={{ width: width, marginLeft: 4, background: modeColors.content, boxShadow: "0 1.6px 3.6px 0 rgba(0,0,0,0.132), 0 0.3px 0.9px 0 rgba(0,0,0,0.108)" }}>
                      <DetailsList
@@ -613,7 +608,6 @@ const JobList: React.FC<{ tab: string; filter: JobFilterSimple, controller: Call
                   <Stack style={{ width: width }}><Spinner size={SpinnerSize.large} /></Stack>
                </Stack>
                }
-               <JobLabelCallout controller={controller} />
             </ScrollablePane>
             }
          </div>
@@ -624,8 +618,6 @@ const JobList: React.FC<{ tab: string; filter: JobFilterSimple, controller: Call
 
 export const JobView: React.FC<{ tab: string; filter: JobFilterSimple }> = ({ tab, filter }) => {
 
-
-   const [controller] = useState(new CalloutController());
-   return (<JobList tab={tab} filter={filter} controller={controller} />);
+   return (<JobList tab={tab} filter={filter} />);
 
 };
diff --git a/Engine/Source/Programs/Horde/HordeDashboard/src/components/JobViewIncremental.tsx b/Engine/Source/Programs/Horde/HordeDashboard/src/components/JobViewIncremental.tsx
--- a/Engine/Source/Programs/Horde/HordeDashboard/src/components/JobViewIncremental.tsx
+++ b/Engine/Source/Programs/Horde/HordeDashboard/src/components/JobViewIncremental.tsx
@@ -14,7 +14,6 @@ import { filterJob, JobFilterSimple } from '../base/utilities/filter';
 import { displayTimeZone } from '../base/utilities/timeUtils';
 import { getJobStateColor, getLabelColor } from '../styles/colors';
 import { ChangeContextMenu, ChangeContextMenuTarget } from './ChangeButton';
-import { CalloutController, JobLabelCallout } from './JobLabelCallout';
 import { JobOperationsContextMenu } from './JobOperationsContextMenu';
 import { getHordeStyling } from '../styles/Styles';
 
@@ -181,7 +180,7 @@ const sortJobGroups = (items: JobItem[]): [JobItem[], JobGroup[]] => {
 
 const jobHandler = new JobHandler();
 
-const JobList: React.FC<{ tab: string; filter: JobFilterSimple, controller: CalloutController }> = observer(({ tab, filter, controller }) => {
+const JobList: React.FC<{ tab: string; filter: JobFilterSimple }> = observer(({ tab, filter }) => {
 
    const { streamId } = useParams<{ streamId: string }>();
    const { projectStore } = useBackend();
@@ -412,11 +411,7 @@ const JobList: React.FC<{ tab: string; filter: JobFilterSimple, controller: Call
       return <Stack>
          <div id={target} className="horde-no-darktheme">
             <Link to={url}><Stack className={hordeClasses.badgeNoIcon}><DefaultButton key={label.dashboardName} style={{ backgroundColor: color.primaryColor }} text={label.dashboardName}
-               onMouseOver={(ev) => {
-                  ev.stopPropagation();
-                  controller.setState({ target: `#${target}`, label: label, jobId: item.job.id })
-               }}
-               onMouseLeave={() => controller.setState(undefined)}
+               onMouseOver={(ev) => ev.stopPropagation() }
                onMouseMove={(ev) => ev.stopPropagation()}>
 
                {!!color.secondaryColor && <div style={{
@@ -458,7 +453,7 @@ const JobList: React.FC<{ tab: string; filter: JobFilterSimple, controller: Call
       });
 
       return (
-         <Stack horizontalAlign="center" onMouseMove={() => controller.setState(undefined)}>
+         <Stack horizontalAlign="center" >
             <Stack horizontal horizontalAlign="center" tokens={{ childrenGap: 4 }} styles={{ root: { paddingTop: 2, width: column?.maxWidth } }}>
                {buttons}
             </Stack>
@@ -630,7 +625,7 @@ const JobList: React.FC<{ tab: string; filter: JobFilterSimple, controller: Call
             {commitState.target && commitState.job && <ChangeContextMenu target={commitState.target} job={commitState.job} commit={commitState.commit} rangeCL={commitState.rangeCL} onDismiss={() => setCommitState({})} />}
             <div className={detailClasses.container} style={{ width: "100%", height: 'calc(100vh - 240px)', position: 'relative', marginTop: 0 }} data-is-scrollable={true}>
 
-               {<ScrollablePane scrollbarVisibility={ScrollbarVisibility.always} onScroll={() => { controller.setState(undefined, true) }} style={{ overflow: "visible" }}>
+               {<ScrollablePane scrollbarVisibility={ScrollbarVisibility.always} style={{ overflow: "visible" }}>
                   {renderItems.length > 0 &&
                      <Stack style={{ width: width, marginLeft: 4, background: modeColors.content, boxShadow: "0 1.6px 3.6px 0 rgba(0,0,0,0.132), 0 0.3px 0.9px 0 rgba(0,0,0,0.108)",  }}>
                         <DetailsList
@@ -670,7 +665,6 @@ const JobList: React.FC<{ tab: string; filter: JobFilterSimple, controller: Call
                      <Stack style={{ width: width }}><Spinner size={SpinnerSize.large} /></Stack>
                   </Stack>
                   }
-                  <JobLabelCallout controller={controller} />
                </ScrollablePane>
                }
             </div>
@@ -682,8 +676,6 @@ const JobList: React.FC<{ tab: string; filter: JobFilterSimple, controller: Call
 
 export const JobViewIncremental: React.FC<{ tab: string; filter: JobFilterSimple }> = ({ tab, filter }) => {
 
-
-   const [controller] = useState(new CalloutController());
-   return (<JobList tab={tab} filter={filter} controller={controller} />);
+   return (<JobList tab={tab} filter={filter} />);
 
 };
diff --git a/Engine/Source/Programs/Horde/HordeDashboard/src/components/LogRender.tsx b/Engine/Source/Programs/Horde/HordeDashboard/src/components/LogRender.tsx
deleted file mode 100644
diff --git a/Engine/Source/Programs/Horde/HordeDashboard/src/components/LogSource.tsx b/Engine/Source/Programs/Horde/HordeDashboard/src/components/LogSource.tsx
deleted file mode 100644
diff --git a/Engine/Source/Programs/Horde/HordeDashboard/src/components/LogStyle.tsx b/Engine/Source/Programs/Horde/HordeDashboard/src/components/LogStyle.tsx
deleted file mode 100644
diff --git a/Engine/Source/Programs/Horde/HordeDashboard/src/components/LogView.tsx b/Engine/Source/Programs/Horde/HordeDashboard/src/components/LogView.tsx
deleted file mode 100644
diff --git a/Engine/Source/Programs/Horde/HordeDashboard/src/components/StatusIcon.tsx b/Engine/Source/Programs/Horde/HordeDashboard/src/components/StatusIcon.tsx
--- a/Engine/Source/Programs/Horde/HordeDashboard/src/components/StatusIcon.tsx
+++ b/Engine/Source/Programs/Horde/HordeDashboard/src/components/StatusIcon.tsx
@@ -6,7 +6,6 @@ import { FontIcon, Stack } from '@fluentui/react';
 import React, { CSSProperties } from 'react';
 import { DeviceStatus, FindIssueResponse, GetAgentLeaseResponse, GetBatchResponse, GetIssueResponse, GetJobStepRefResponse, GetLabelStateResponse, GetStepResponse, IssueSeverity, JobStepBatchError, JobStepOutcome, JobStepState, LabelOutcome, LabelState, LeaseOutcome, LeaseState } from '../backend/Api';
 import dashboard, { StatusColor } from "../backend/Dashboard";
-import { JobLabel } from '../backend/JobDetails';
 import { getStepStatusColor } from '../styles/colors';
 
 const StatusIcon: React.FC<{ iconName: string, style?: CSSProperties }> = ({ iconName, style }) => {
diff --git a/Engine/Source/Programs/Horde/HordeDashboard/src/components/error/ErrorPane.tsx b/Engine/Source/Programs/Horde/HordeDashboard/src/components/error/ErrorPane.tsx
deleted file mode 100644
diff --git a/Engine/Source/Programs/Horde/HordeDashboard/src/components/jobDetailsV2/JobDetailLabels.tsx b/Engine/Source/Programs/Horde/HordeDashboard/src/components/jobDetailsV2/JobDetailLabels.tsx
--- a/Engine/Source/Programs/Horde/HordeDashboard/src/components/jobDetailsV2/JobDetailLabels.tsx
+++ b/Engine/Source/Programs/Horde/HordeDashboard/src/components/jobDetailsV2/JobDetailLabels.tsx
@@ -3,8 +3,7 @@
 import { ConstrainMode, DefaultButton, DetailsList, DetailsListLayoutMode, DirectionalHint, IColumn, SelectionMode, Stack, Text, TooltipHost } from "@fluentui/react";
 import { observer } from "mobx-react-lite";
 import { useNavigate, useLocation } from "react-router-dom";
-import { GetStepResponse, JobStepOutcome, LabelOutcome, LabelState, StepData } from "../../backend/Api";
-import { JobLabel } from "../../backend/JobDetails";
+import { GetStepResponse, JobLabel, JobStepOutcome, LabelOutcome, LabelState, StepData } from "../../backend/Api";
 import { getLabelColor } from "../../styles/colors";
 import { useQuery } from "horde/base/utilities/hooks";
 import { JobDataView, JobDetailsV2 } from "./JobDetailsViewCommon";
diff --git a/Engine/Source/Programs/Horde/HordeDashboard/src/components/jobDetailsV2/JobDetailViewSteps.tsx b/Engine/Source/Programs/Horde/HordeDashboard/src/components/jobDetailsV2/JobDetailViewSteps.tsx
--- a/Engine/Source/Programs/Horde/HordeDashboard/src/components/jobDetailsV2/JobDetailViewSteps.tsx
+++ b/Engine/Source/Programs/Horde/HordeDashboard/src/components/jobDetailsV2/JobDetailViewSteps.tsx
@@ -659,130 +659,4 @@ export const StepsPanelV2: React.FC<{ jobDetails: JobDetailsV2, depStepId?: stri
 
    </Stack>);
 
-});
-
-export const getStepSummaryMarkdown = (jobDetails: JobDetailsV2, stepId: string): string => {
-
-   const jobData = jobDetails.jobData;
-
-   if (!jobData) {
-      return "";
-   }
-
-   const step = jobDetails.stepById(stepId)!;
-   const batch = jobDetails.batchByStepId(stepId);
-
-   if (!step || !batch) {
-      return "";
-   }
-
-   const duration = getStepElapsed(step);
-   let eta = getStepETA(step, jobData);
-
-   const text: string[] = [];
-
-   if (jobData) {
-      text.push(`Job created by ${jobData.startedByUserInfo ? jobData.startedByUserInfo.name : "scheduler"}`);
-   }
-
-   const batchText = () => {
-
-      if (!batch) {
-         return undefined;
-      }
-
-      const agentType = batch.agentType;
-      const agentPool = jobDetails.stream?.agentTypes[agentType!]?.pool;
-      return getBatchText({ batch: batch, agentType: agentType, agentPool: agentPool });
-
-   };
-
-   if (step.retriedByUserInfo) {
-      const retryId = jobDetails.getRetryStepId(step.id);
-      if (retryId) {
-         text.push(`[Step was retried by ${step.retriedByUserInfo.name}](/job/${jobDetails.jobId}?step=${retryId})`);
-      } else {
-         text.push(`Step was retried by ${step.retriedByUserInfo.name}`);
-      }
-
-   }
-   if (step.abortRequested || step.state === JobStepState.Aborted) {
-      eta.display = eta.server = "";
-      let aborted = "";
-      if (step.abortedByUserInfo) {
-         aborted = "This step was canceled";
-         aborted += ` by ${step.abortedByUserInfo.name}.`;
-      } else if (jobData.abortedByUserInfo) {
-         aborted = "The job was canceled";
-         aborted += ` by ${jobData.abortedByUserInfo.name}.`;
-      } else {
-         aborted = "This step was canceled by Horde";
-      }
-      text.push(aborted);
-   } else if (step.state === JobStepState.Skipped) {
-      eta.display = eta.server = "";
-      text.push("The step was skipped");
-   } else if (step.state === JobStepState.Ready || step.state === JobStepState.Waiting) {
-
-      text.push(batchText() ?? `The step is pending in ${step.state} state`);
-   }
-
-   if (batch?.agentId) {
-
-      if (step.startTime) {
-         const str = getNiceTime(step.startTime);
-         text.push(`Step started on ${str}`);
-      }
-
-      let runningText = `${step.finishTime ? "Ran" : "Running"} on [${batch.agentId}](?step=${stepId}&agentId=${encodeURIComponent(batch.agentId)})`;
-
-      if (duration) {
-         runningText += ` for ${duration.trim()}`;
-      }
-
-      if (step.finishTime && (step.outcome === JobStepOutcome.Success || step.outcome === JobStepOutcome.Warnings)) {
-         const delta = getStepTimingDelta(step);
-         if (delta) {
-            runningText += " " + delta;
-         }
-      }
-
-      if (runningText) {
-         text.push(runningText);
-      }
-
-   } else {
-      if (!step.abortRequested) {
-         text.push(batchText() ?? "Step does not have a batch.");
-      }
-   }
-
-   if (eta.display) {
-      text.push(`Expected to finish around ${eta.display}.`);
-   }
-
-   const finish = getStepFinishTime(step).display;
-
-   if (finish && step.state !== JobStepState.Aborted) {
-
-      let outcome = "";
-      if (step.outcome === JobStepOutcome.Success) {
-         outcome += `Completed at ${finish}`;
-      }
-      if (step.outcome === JobStepOutcome.Failure)
-         outcome += `Completed with errors at ${finish}.`;
-      if (step.outcome === JobStepOutcome.Warnings)
-         outcome += `Completed with warnings at ${finish}.`;
-
-      if (outcome) {
-         text.push(outcome);
-      }
-   }
-
-   if (!text.length) {
-      text.push(`Step is in ${step.state} state.`);
-   }
-
-   return text.join(".&nbsp;&nbsp;");
-
-}
\ No newline at end of file
+});
\ No newline at end of file
diff --git a/Engine/Source/Programs/Horde/HordeDashboard/src/components/jobDetailsV2/JobDetailsViewCommon.tsx b/Engine/Source/Programs/Horde/HordeDashboard/src/components/jobDetailsV2/JobDetailsViewCommon.tsx
--- a/Engine/Source/Programs/Horde/HordeDashboard/src/components/jobDetailsV2/JobDetailsViewCommon.tsx
+++ b/Engine/Source/Programs/Horde/HordeDashboard/src/components/jobDetailsV2/JobDetailsViewCommon.tsx
@@ -7,8 +7,7 @@ import moment from "moment";
 import React, { useState } from "react";
 import { useLocation, useNavigate } from "react-router-dom";
 import backend from "../../backend";
-import { BatchData, GetArtifactResponse, GetJobTimingResponse, GetLabelStateResponse, GetReportResponse, GetTemplateRefResponse, JobData, JobState, LabelState, ReportPlacement, StepData, StreamData } from "../../backend/Api";
-import { JobLabel } from "../../backend/JobDetails";
+import { BatchData, GetArtifactResponse, GetJobTimingResponse, GetLabelStateResponse, GetReportResponse, GetTemplateRefResponse, JobData, JobLabel, JobState, LabelState, ReportPlacement, StepData, StreamData } from "../../backend/Api";
 import { PollBase } from "../../backend/PollBase";
 import { projectStore } from "../../backend/ProjectStore";
 import TemplateCache from "../../backend/TemplateCache";
@@ -16,6 +15,7 @@ import { ISideRailLink } from "../../base/components/SideRail";
 import { useQuery } from "horde/base/utilities/hooks";
 import { LabelStatusIcon, StepStatusIcon } from "../StatusIcon";
 import { getHordeStyling } from "../../styles/Styles";
+import { errorBarStore } from "../error/ErrorStore";
 
 export abstract class JobDataView {
 
@@ -248,26 +248,38 @@ export class JobDetailsV2 extends PollBase {
       return this.batches.map(b => b.steps).flat();
    }
 
+   /**
+    * Returns the first step that satisfies the predicate.
+    * @param predicate The predicate is called for every step until it finds one where the predicate returns true.
+    * @returns The first step that satisfies the predicate. If none do, findStep returns undefined.
+    */
+   private findStep(predicate: (s: StepData) => boolean): StepData | undefined {
+      if(!this.jobId || !this.batches.length) return undefined;
 
-   stepById(id?: string): StepData | undefined {
-
-      if (!id || !this.jobId || !this.batches.length) {
-         return undefined;
+      for (const batch of this.batches) {
+         const step = batch.steps.find(predicate);
+         if(step) return step;
       }
+   }
 
-      let step: StepData | undefined;
-
-      this.batches.forEach(b => {
-
-         if (step) {
-            return;
-         }
-
-         step = b.steps.find(s => s.id === id);
+   /**
+    * Returns the first step with matching step id.
+    * @param id The id of the step to get.
+    * @returns The first step with matching step id. If none found, returns undefined.
+    */
+   stepById(id?: string): StepData | undefined {
+      if(!id) return undefined;
 
-      });
+      return this.findStep(s => s.id === id);
+   }
 
-      return step;
+   /**
+    * Returns the furst step with matching log id.
+    * @param logId The log id of the step to get.
+    * @returns The first step with matching log id. If none found, returns undefined.
+    */
+   stepByLogId(logId: string): StepData | undefined {
+      return this.findStep(s => s.logId === logId);
    }
 
    getStepName(stepId: string | undefined, includeRetry: boolean = true): string {
@@ -584,6 +596,7 @@ export class JobDetailsV2 extends PollBase {
          }
 
          if (this.jobError) {
+            errorBarStore.set({title: "Error Getting Job Data", message: this.jobError}, true);
             this.setRootUpdated();
             return;
          }
diff --git a/Engine/Source/Programs/Horde/HordeDashboard/src/components/jobDetailsV2/StepErrorPanel.tsx b/Engine/Source/Programs/Horde/HordeDashboard/src/components/jobDetailsV2/StepErrorPanel.tsx
--- a/Engine/Source/Programs/Horde/HordeDashboard/src/components/jobDetailsV2/StepErrorPanel.tsx
+++ b/Engine/Source/Programs/Horde/HordeDashboard/src/components/jobDetailsV2/StepErrorPanel.tsx
@@ -10,7 +10,7 @@ import { DashboardPreference, EventData, EventSeverity, JobStepState } from '../
 import { ISideRailLink } from "../../base/components/SideRail";
 import { getHordeStyling } from "../../styles/Styles";
 import { getHordeTheme } from "../../styles/theme";
-import { renderLine } from '../LogRender';
+import { renderLine } from '../logs/LogRender';
 import { JobDataView, JobDetailsV2 } from "./JobDetailsViewCommon";
 import dashboard from "horde/backend/Dashboard";
 
diff --git a/Engine/Source/Programs/Horde/HordeDashboard/src/components/logs/LogHandler.ts b/Engine/Source/Programs/Horde/HordeDashboard/src/components/logs/LogHandler.ts
new file mode 100644
--- /dev/null
+++ b/Engine/Source/Programs/Horde/HordeDashboard/src/components/logs/LogHandler.ts
@@ -0,0 +1,250 @@
+// Copyright Epic Games, Inc. All Rights Reserved.
+
+import { action, makeObservable, observable } from "mobx";
+import { getLogStyles, logMetricNormal, logMetricSmall } from "./LogStyle";
+import { EventSeverity, EventData } from "horde/backend/Api";
+import { LogSource } from "./LogSource";
+
+/**
+ * Handler class for Log View data like the log source, navigation state, and helper functions.
+ */
+export class LogHandler {
+    private static instance?: LogHandler;
+    
+    logSource?: LogSource;
+    currentLine?: number;
+    trailing?: boolean;
+    scroll?: number;
+    initialRender = true;
+
+    @observable
+    updated = 0
+
+    constructor() {
+        makeObservable(this);
+        LogHandler.instance = this;
+    }
+
+    async init(logId: string, startLine?: number) {
+        this.logSource = await LogSource.create(logId, startLine);
+        await this.logSource.init();
+
+        if(!startLine && this.logSource.active) {
+            this.trailing = true;
+        }
+
+        this.currentLine = startLine;
+    }
+
+    static clear() {
+        this.instance?.logSource?.stop();
+        this.instance = undefined;
+    }
+
+    @action
+    externalUpdate() {
+        this.updated++;
+    }
+
+    stopTrailing() {
+        this.scroll = undefined;
+        this.trailing = false;
+    }
+
+    getCurrentEvent() {
+        return this.getLogEvent(this.currentLine)
+    }
+
+    getEventsOfSeverity(warnings?: boolean){
+        return this.events.filter(e => {
+            if (warnings && e.severity !== EventSeverity.Warning) {
+                return false;
+            }
+            if (!warnings && e.severity !== EventSeverity.Error) {
+                return false;
+            }
+            return true;
+        });
+    }
+
+    getEventBlocks(events: EventData[]){
+        return events.filter((e, i, events) => {
+            // first event must be a block
+            if (i == 0) {
+                return true
+            }
+            // consider different issue a new block, regardless of spacing
+            if (e.issueId != events[i-1].issueId) {
+                return true
+            }
+            
+            if (e.lineIndex == events[i-1].lineIndex + events[i-1].lineCount) {
+                return false
+            }
+            return true
+        });
+    }
+
+    getNextLogEvent(warnings?: boolean) {
+
+        let currentLine = -1;
+        if (this.currentLine !== undefined) {
+            currentLine = this.currentLine - 1;
+        }
+
+        const startLine = currentLine + 1;
+
+        const events = this.getEventsOfSeverity(warnings);
+
+        let event = events.find(e => e.lineIndex >= startLine);
+
+        if (!event && events.length) {
+            event = events[0];
+        }
+
+        return event;
+    }
+
+    getNextLogEventBlock(warnings?: boolean) {
+        
+        let currentLine = -1;
+        if (this.currentLine !== undefined) {
+            currentLine = this.currentLine - 1;
+        }
+
+        const startLine = currentLine + 1;
+
+        const events = this.getEventsOfSeverity(warnings);
+        const eventBlocks = this.getEventBlocks(events)
+
+        let event = eventBlocks.find(eb => eb.lineIndex >= startLine);
+
+        if (!event && eventBlocks.length) {
+            event = eventBlocks[0];
+        }
+
+        return event;
+    }
+
+    getPrevLogEvent(warnings?: boolean) {
+
+        let currentLine = 0;
+        if (this.currentLine !== undefined) {
+            currentLine = this.currentLine - 1;
+        }
+
+        const startLine = currentLine - 1;
+
+        const events = this.getEventsOfSeverity(warnings).reverse();
+
+        if (startLine < 0) {
+            if (events.length) {
+                return events[0];
+            }
+            return undefined;
+        }
+
+        let event = events.find(e => e.lineIndex <= startLine)
+
+        if (!event && events.length) {
+            event = events[0];
+        }
+
+        return event;
+    }
+
+    getPrevLogEventBlock(warnings?: boolean) {
+        
+        let currentLine = -1;
+        if (this.currentLine !== undefined) {
+            currentLine = this.currentLine - 1;
+        }
+
+        const startLine = currentLine - 1;
+
+        const events = this.getEventsOfSeverity(warnings);
+        const eventBlocks = this.getEventBlocks(events).reverse();
+
+        if (startLine < 0) {
+            if (eventBlocks.length) {
+                return eventBlocks[0];
+            }
+            return undefined;
+        }
+
+        let event = eventBlocks.find(eb => eb.lineIndex <= startLine);
+
+        if (!event && eventBlocks.length) {
+            event = eventBlocks[0];
+        }
+
+        return event;
+    }
+
+    getLogEvent(line: number | undefined) {
+
+        if (line === undefined) {
+            return undefined;
+        }
+
+        line--;
+
+        return this.events.find(e => line >= e.lineIndex && line < (e.lineIndex + e.lineCount));
+    }
+
+    get events(): EventData[] {
+        const events = this.logSource?.errors.map(e => e) ?? [];
+        events.push(...(this.logSource?.warnings.map(e => e) ?? []))
+        return events;
+    }
+
+    infoLine?: number;
+
+    // could be a preference
+    compact = false;
+
+    get lineHeight(): number {
+
+        if (!this.logSource?.logItems) {
+            return logMetricNormal.lineHeight;
+        }
+
+        return !this.compact ? logMetricNormal.lineHeight : logMetricSmall.lineHeight;
+
+    }
+
+    get fontSize(): number {
+
+        if (!this.logSource?.logItems) {
+            return logMetricNormal.fontSize;
+        }
+
+        return !this.compact ? logMetricNormal.fontSize : logMetricSmall.fontSize;
+
+    }
+
+    get style(): any {
+
+        const { logStyleNormal, logStyleSmall } = getLogStyles();
+
+        if (!this.logSource?.logItems) {
+            return logStyleNormal;
+        }
+
+        return !this.compact ? logStyleNormal : logStyleSmall;
+
+    }
+
+    get lineRenderStyle(): any {
+
+        const { lineRenderStyleNormal, lineRenderStyleSmall } = getLogStyles();
+
+        if (!this.logSource?.logItems) {
+            return lineRenderStyleNormal;
+        }
+
+        return !this.compact ? lineRenderStyleNormal : lineRenderStyleSmall;
+
+    }
+}
+
diff --git a/Engine/Source/Programs/Horde/HordeDashboard/src/components/logs/LogRender.tsx b/Engine/Source/Programs/Horde/HordeDashboard/src/components/logs/LogRender.tsx
new file mode 100644
--- /dev/null
+++ b/Engine/Source/Programs/Horde/HordeDashboard/src/components/logs/LogRender.tsx
@@ -0,0 +1,413 @@
+// Copyright Epic Games, Inc. All Rights Reserved.
+
+import { Stack, Text } from '@fluentui/react';
+import Highlight from 'react-highlighter';
+import { IssueData, LogLine } from '../../backend/Api';
+import backend from '../../backend';
+import { NavigateFunction } from 'react-router-dom';
+
+enum TagType {
+   None,
+   SourceFile,
+   MSDNCode,
+   Link,
+   JobId,
+   LogId,
+   ArtifactId,
+   AgentId,
+   PoolId,
+   LeaseId,
+   ProjectId,
+   StreamId,
+   DeviceId,
+}
+
+export type LogItem = {
+   line?: LogLine;
+   lineNumber: number;
+   requested: boolean;
+   issueId?: number;
+   issue?: IssueData;
+}
+
+// Provides the start of hrefs for each tag type with simple link construction behavior
+const simpleTagLinks: Partial<Record<TagType, string>> = {
+   [TagType.JobId]: "/job/",
+   [TagType.LogId]: "/log/",
+   [TagType.ProjectId]: "/project/",
+   [TagType.StreamId]: "/stream/",
+   [TagType.PoolId]: "/pools?poolid=",
+   [TagType.DeviceId]: "/audit/device/"
+}
+
+const getLineNumber = (line: LogLine): number | undefined => {
+
+   const properties = line.properties;
+   if (!properties) {
+      return undefined;
+   }
+
+   const lineProp = properties["line"];
+
+   if (!lineProp || typeof (lineProp) !== "object") {
+      return undefined;
+   }
+
+   return (lineProp as Record<string, number | undefined>).line;
+
+}
+
+const renderMessage = (line: LogLine, lineNumber: number | undefined, logStyle: any, search?: string) => {
+
+   const message = line?.message ?? "";
+
+   if (lineNumber === undefined || lineNumber === null) {
+      lineNumber = getLineNumber(line);
+   }
+
+   const key = `log_line_${lineNumber}_message`;
+
+   return <Highlight key={key} search={search ? search : ""} className={logStyle.logLine}>{message ?? ""}</Highlight >;
+
+};
+
+
+const renderTags = (navigate: NavigateFunction, line: LogLine, lineNumber: number | undefined, logStyle: any, tags: string[], search?: string) => {
+
+   if (!line || !line.format || !line.properties) {
+      return <Stack styles={{ root: { color: "#000000", paddingLeft: 8, whiteSpace: "pre", tabSize: "3" } }}>Internal log line format error</Stack>;
+   }
+
+   const properties = line.properties;
+   const lineProp = properties["line"];
+
+   if (lineNumber === undefined || lineNumber === null) {
+      lineNumber = (lineProp as Record<string, number | undefined>)?.line ?? 0;
+   }
+
+   // render tags
+   let renderedTags = tags.map((tag, index) => {
+
+      let tagType = TagType.None;
+      let property: Record<string, string | number> | string | undefined;
+      let record: Record<string, string> | undefined;
+      const key = `log_line_${lineNumber}_${index}`;
+
+      let formatter = "";
+      tag = tag.replace("{", "").replace("}", "");
+      if (tag.indexOf(":") !== -1) {
+         [tag, formatter] = tag.split(":");
+      }
+
+      property = properties[tag];
+
+      // C# format specifiers, yes this could be more elegant
+      if (formatter) {
+         if (typeof (property) === "number") {
+
+            let precision: number | undefined;
+
+            if (formatter.startsWith("n") && formatter.length > 1) {
+               precision = parseInt(formatter.slice(1));
+            }
+
+            if (formatter.startsWith("0.") && formatter.length > 2) {
+               precision = 0;
+               for (let i = 2; i < formatter.length; i++) {
+                  if (formatter[i] === "0") {
+                     precision++;
+                  } else {
+                     break;
+                  }
+               }
+            }
+
+            if (precision !== undefined && !isNaN(precision)) {
+               property = (property as number).toFixed(precision);
+            }
+
+         }
+      }
+
+      let type = "";
+      let text = "";
+
+      if (property && typeof (property) !== "string" && typeof (property) !== "number" && typeof (property) !== "boolean" && property !== null) {
+         record = property as Record<string, string>;
+      } else {
+         if (property === undefined) {
+            return null;
+         }
+         if (property === null) {
+            text = "null";
+         } else {
+            text = property.toString();
+         }
+      }
+
+      if (record) {
+
+         type = record.type ?? record["$type"] ?? "";
+         text = record.text ?? record["$text"] ?? "";
+         
+         if (type === "ErrorCode" && text.startsWith("C")) {
+            tagType = TagType.MSDNCode;
+         }
+
+         const enumCast = type as keyof typeof TagType;
+         if(enumCast in TagType) {
+            tagType = TagType[enumCast];
+         }
+      }
+
+      let highlight = (
+         <Highlight
+            search={search ? search : ""} 
+            className={logStyle.logLine} 
+         >
+            {text}
+         </Highlight>
+      )
+
+      if(tagType === TagType.None || !record) {
+         return (
+            // Defined separately from highlight because of needed key prop
+            <Highlight
+               key={key}
+               search={search ? search : ""} 
+               className={logStyle.logLine} 
+            >
+               {text}
+            </Highlight>
+         )
+      }
+
+      let linkProps: React.AnchorHTMLAttributes<HTMLAnchorElement> | null = null;
+
+      // Update link props object based on tag type, then return a link with those props
+      // or a span
+      switch(tagType) {
+         // Cases with similar simple tag behaviors
+         case TagType.JobId:
+         case TagType.LogId:
+         case TagType.ProjectId:
+         case TagType.StreamId:
+         case TagType.PoolId:
+         case TagType.DeviceId:
+            const url = simpleTagLinks[tagType] + text; // simpleTagLinks created for concision here
+            linkProps = {
+               href: url,
+               onClick: (ev) => { ev.stopPropagation(); ev.preventDefault(); navigate(url) }
+            }
+            break;
+
+         case TagType.MSDNCode:
+            linkProps = {
+               href: `https://msdn.microsoft.com/query/dev16.query?appId=Dev16IDEF1&l=EN-US&k=k(${text.toLowerCase()})&rd=true`,
+               onClick: (ev) => ev.stopPropagation(),
+               target: "_blank",
+               rel: "noopener noreferrer"
+            };
+            break;
+
+         case TagType.Link:
+            linkProps = {
+               href: record.target,
+               onClick: (ev) => ev.stopPropagation(),
+               rel: "noreferrer"
+            }
+            break;
+
+         case TagType.AgentId:
+            const search = new URLSearchParams(window.location.search);
+            search.set("agentId", encodeURIComponent(text));
+            // Update highlight because of new search in block scope
+            highlight = (
+               <Highlight 
+                  search={search ? search : ""} 
+                  className={logStyle.logLine}
+               >
+                  {text}
+               </Highlight>
+            )
+            const url2 = `${window.location.pathname}?` + search.toString();
+            linkProps = {
+               href: url2,
+               onClick: async (ev) => { ev.stopPropagation(); ev.preventDefault(); navigate(url2, { replace: true }) }
+            }
+            break;
+
+         case TagType.ArtifactId:
+            const artifactType = properties.ArtifactType ?? "_none";
+            const search2 = new URLSearchParams(window.location.search);
+            search2.set("artifactContext", encodeURIComponent(artifactType as string));
+            if (properties.ArtifactId) {
+               if (properties.ArtifactId["$text"]) {
+                  search2.set("artifactId", encodeURIComponent(properties.ArtifactId["$text"] as string));
+               }
+            }
+            // Update highlight because of new search in block scope
+            highlight = (
+               <Highlight 
+                  search={search2 ? search2 : ""} 
+                  className={logStyle.logLine}
+               >
+                  {text}
+               </Highlight>
+            )
+            const url3 = `${window.location.pathname}?${search2.toString()}`;
+            linkProps = {
+               href: url3,
+               onClick: async (ev) => { ev.stopPropagation(); ev.preventDefault(); navigate(url3, { replace: true }) }
+            };
+            break;
+
+         case TagType.SourceFile:
+            let depotPath = record.depotPath;
+            if (!depotPath) break;
+            if (depotPath.indexOf("@") !== -1) {
+               depotPath = depotPath.slice(0, depotPath.indexOf("@"));
+            }
+            depotPath = encodeURIComponent(depotPath);
+            // @todo: handle line, it isn't in the source file meta data, and matching isn't obvious (file vs note, etc)
+            let tagLine = "";
+            if (tagLine) {
+               depotPath += `&line=${tagLine}`;
+            }
+            linkProps = {
+               href: `ugs://timelapse?depotPath=${(depotPath)}`,
+               onClick: (ev) => ev.stopPropagation()
+            }
+            break;
+
+         case TagType.LeaseId:
+            const navigateToLeaseLog = async (toplevel: boolean) => {
+               const logData = await backend.getLease(text);
+               const url = `/log/${logData?.logId}`;
+               if (!toplevel) {
+                  navigate(url)
+               } else {
+                  window.open(url, "_blank");
+               }
+            }
+            linkProps = {
+               href: "/",
+               onClick: (ev) => {
+                  ev.stopPropagation();
+                  ev.preventDefault();
+                  navigateToLeaseLog(!!ev?.ctrlKey || !!ev.metaKey)
+               },
+               onAuxClick: (ev) => {
+                  ev.preventDefault();
+                  ev.stopPropagation()
+                  navigateToLeaseLog(true)
+               }
+            }
+            break;
+      }
+
+      return linkProps ? <a key={key} {...linkProps}>{highlight}</a> : <span key={key}/>
+   }).filter(t => !!t);
+
+   let remaining = line.format;
+
+   renderedTags = renderedTags.map((t, idx) => {
+
+      let current = remaining;
+
+      const tag = tags[idx];
+      const index = remaining.indexOf(tag);
+
+      remaining = remaining.slice(tag.length + (index > 0 ? index : 0));
+
+      if (index < 0) {
+         console.error("not able to find tag in format");
+         return <Text>Error, unable to find tag</Text>;
+      }
+
+      if (index === 0) {
+         return t;
+      }
+
+      const rtags: any = [];
+
+      const key = `log_line_${lineNumber}_${idx}_${index}_fragment`;
+
+      rtags.push(<Highlight key={key} search={search ? search : ""} className={logStyle.logLine}>{current.slice(0, index)}</Highlight>);
+      rtags.push(t);
+
+      return rtags;
+
+   }).flat();
+
+   if (remaining) {
+      const key = `log_line_${lineNumber}_remaining_fragment`;
+      renderedTags.push(<Highlight key={key} search={search ? search : ""} className={logStyle.logLine}>{remaining}</Highlight>)
+   }
+
+   return <div>
+      {renderedTags}
+   </div>;
+
+};
+
+export const renderLine = (navigate: NavigateFunction, line: LogLine | undefined, lineNumber: number | undefined, logStyle: any, search?: string) => {
+
+   if (!line) {
+      return null;
+   }
+
+   const tagRegex = /{[^{}]+}/g
+   let format = line?.format;
+
+   let tags: string[] = [];
+   if (format) {
+
+      // some automation logs output their own structured logging, which is marked by {{}}
+      if (format.indexOf("{{") !== -1) {
+         format = format.replaceAll("{", "");
+         format = format.replaceAll("}", "");
+      }
+
+      const match = format.match(tagRegex);
+      if (match?.length)
+         tags = match;
+
+      const properties = line.properties;
+
+      // fix issue with tag span, we need to replace react-highlighter as it does not highlight across child nodes
+      // should just be using a selector
+      if (properties) {
+         tags = tags.filter(t => {
+            const pname = t.slice(1, -1);
+
+            if (pname !== "WarningCode" && pname !== "WarningMessage") {
+               return true;
+            }
+
+            const ptext = (properties[pname] as any);
+            if (!ptext) {
+               return true;
+            }
+
+            line.format = line.format?.replaceAll(t, ptext);
+            return false;
+         });
+      }
+   }
+
+   // we don't support c# alignment, as this requireds read behinds, etc
+   // so strip tags in this case and just output the line
+   if (tags.find(t => t.indexOf(",") !== -1)) {
+      tags = [];
+   }
+
+   if (tags.length && format && line.properties) {
+      return renderTags(navigate, line, lineNumber, logStyle, tags, search);
+   }
+
+   return renderMessage(line, lineNumber, logStyle, search);
+
+};
+
+
diff --git a/Engine/Source/Programs/Horde/HordeDashboard/src/components/logs/LogSource.ts b/Engine/Source/Programs/Horde/HordeDashboard/src/components/logs/LogSource.ts
new file mode 100644
--- /dev/null
+++ b/Engine/Source/Programs/Horde/HordeDashboard/src/components/logs/LogSource.ts
@@ -0,0 +1,588 @@
+// Copyright Epic Games, Inc. All Rights Reserved.
+
+import backend from "horde/backend";
+import { LogData, EventData, IssueData, StepData, BatchData, ArtifactData, StreamData, EventSeverity, GetJobsTabResponse, IssueQuery, AgentData, LeaseData } from "horde/backend/Api";
+import { PollBase } from "horde/backend/PollBase";
+import { waitFor } from "horde/base/utilities/promiseUtils";
+import { getStepPercent, getLeaseElapsed } from "horde/base/utilities/timeUtils";
+import { observable, makeObservable, action, runInAction, computed } from "mobx";
+import moment from "moment";
+import { Moment } from "moment-timezone";
+import { AgentTelemetryHandler } from "../agents/AgentTelemetrySparkline";
+import { BreadcrumbItem } from "../Breadcrumbs";
+import { errorBarStore } from "../error/ErrorStore";
+import { JobDetailsV2 } from "../jobDetailsV2/JobDetailsViewCommon";
+import { LogItem } from "./LogRender";
+import { getStepSummaryMarkdown, getBatchSummaryMarkdown } from "horde/base/utilities/logUtils";
+
+/**
+ * Data handler parent class for job and lease logs.
+ */
+export abstract class LogSource extends PollBase {
+    startTime: Moment;
+    startLine?: number;
+    agentTelemetry: AgentTelemetryHandler;
+
+    logData: LogData;
+    logId: string;
+    @observable logItems: LogItem[];
+    @observable private _itemsVersion = 0;
+    @observable active: boolean;
+
+    @observable fatalError?: string;
+
+    /**
+     * Factory method for creating Lease or Log Job Sources
+     * @param logId The ID of the source's log
+     * @returns A LogSource: JobLogSource or LeaseLogSource
+    */
+    static async create(logId: string, startLine?: number): Promise<LogSource> {
+        let logData: LogData;
+       
+        try {
+           logData = await backend.getLogData(logId);
+        } catch (err) {
+            throw new Error(`Failed to load log data for ${logId}: ${err}`);
+        }
+
+        // JobId may come in as a string of 0s. This converts to boolean.
+        const hasJobId = !!Number(`0x${logData.jobId}`);
+        if(hasJobId) return new JobLogSource(logData, startLine);
+        if(!logData.leaseId) throw new Error("Bad Lease Log");
+        return new LeaseLogSource(logData, startLine);
+    }
+
+    //#region Abstract members
+    
+    /**
+     * Downloads the log in either JSON or rest format.
+     * @param json If true, downloads the log in .jsonl (JSON Lines). If false, downloads as text in .log format.
+     */
+    abstract download(json: boolean);
+
+    abstract get summary(): string;
+
+    abstract get agentId(): string | undefined;
+
+    abstract get crumbs(): BreadcrumbItem[];
+
+    //#endregion Abstract members
+
+    //#region LeaseLogSource No-ops
+    
+    get events(): EventData[] { return [] };
+
+    get warnings(): EventData[] { return [] };
+
+    get errors(): EventData[] { return [] };
+
+    get issues(): IssueData[] { return [] };
+
+    get percentComplete(): number | undefined { return undefined };
+
+    get crumbTitle(): string | undefined { return undefined };
+
+    //#endregion LeaseLogSource No-ops
+
+    @computed get itemsVersion(): number {
+        return this._itemsVersion;
+    }
+
+    @action.bound
+    setActive(active: boolean) {
+        this.active = active;
+        this.agentTelemetry.setActive(active);
+    }
+
+    @action.bound
+    setFatalError(error: string) {
+        this.fatalError = error;
+    }
+    
+    constructor(logData: LogData, startLine?: number) {
+        super();
+        
+        makeObservable(this);
+        
+        this.logData = logData;
+        this.logId = logData.id;
+        this.startLine = startLine;
+        this.logItems = [];
+        this.agentTelemetry = new AgentTelemetryHandler();
+    }
+
+    async init() {
+        try {
+            await this.updateLogData();
+        } catch(reason) {
+            errorBarStore.set({title: "Error Loading Log Data", message: reason}, true)
+        }
+    }
+
+    async updateLogData() {
+        try {
+            const newLogData = await backend.getLogData(this.logId);
+            runInAction(() => {
+                this.logData = newLogData;
+                this.resize(this.logData.lineCount);
+            })
+        } catch(reason) {
+            errorBarStore.set({title: "Error Updating Log Data", message: reason}, true);
+        }
+    }
+
+    /**
+     * Adds a number of placeholder line items to the logItems array.
+     * @param lineCount The desired total line count of the updated logItems array.
+     */
+    @action.bound
+    resize(lineCount) {
+        // LogData is updated before resizing, so use logItems.
+        const count = lineCount - this.logItems.length;
+
+        if(count <= 0) return;
+
+        this.logItems.push(...this.getPlaceholderLogItems(count, this.logItems.length + 1));
+        this._itemsVersion++;
+    }
+    
+    /**
+     * Creates an array of simple (placeholder/unrequested) LogItems.
+     * @param count Number of placeholder items to generate.
+     * @param start The starting index of the new placeholder LogItems.
+     */
+    getPlaceholderLogItems(count: number, start = 1): LogItem[] {
+        return Array.from({length: count}, (_, i) => ({
+            lineNumber: start + i,
+            requested: false
+        }))
+    }
+
+    /**
+     * Loads log data to unrequested lines from the logItems array.
+     * @param startIndex The index from which to load logItem data.
+     * @param count The number of lines after startIndex to load/check.
+     */
+    async loadLines(startIndex: number, count: number) {
+        try {
+            if(!this.logData) throw new Error("Invalid/Missing Log Data");
+
+            let anyLinesRequested = false;
+
+            // Detect if any lines within requested range are actually unrequested.
+            for (let i = 0; i < count; i++) {
+                const offset = i + startIndex;
+                
+                if (offset >= this.logItems.length) break;
+
+                const item = this.logItems![i + startIndex];
+                if (!item.requested) {
+                    item.requested = true;
+                    anyLinesRequested = true;
+                }
+            }
+
+            // If all specified lines are already requested, return early.
+            if(!anyLinesRequested) return;
+
+            const data = await backend.getLogLines(this.logId, startIndex, count);
+            
+            if(!this.startTime) throw new Error("Invalid log start time.");
+
+            runInAction(() => {
+                for(let i = 0; i < data.count; i++) {
+                    const line = data.lines![i];
+    
+                    const offset = i + data.index;
+                    if(offset >= this.logItems.length) break;
+                    
+                    // Whitespace adjustments
+                    const stripNewlinesRegex = /(\r\n|\n|\r)/gm;
+                    const fixWhitespace = (text: string | undefined): string | undefined => {
+                      return text?.replace(stripNewlinesRegex, "").trimEnd();
+                    };
+                    line.message = fixWhitespace(line.message)!;
+                    line.format = fixWhitespace(line.format);
+                    
+                    this.logItems[offset].line = line;
+                }
+
+                this._itemsVersion++;
+            })
+        } catch(reason) {
+            errorBarStore.set({title: "Error Loading Line Data", message: reason}, true)
+        }
+    }
+
+    async poll() {
+        if(!this.logData) return;
+
+        if(!this.active) {
+            this.stop();
+        }
+
+        try {
+            await this.updateLogData();
+        } catch(reason) {
+            errorBarStore.set({title: "Error Polling Log Data", message: reason}, true)
+        }
+    }
+}
+
+/**
+ * Log source for job logs. Contains extra information like job details, events, etc.
+ */
+export class JobLogSource extends LogSource {
+    jobDetails: JobDetailsV2;
+    @observable _events: EventData[];
+    _issues: IssueData[];
+    jobName: string = "Unknown Job";
+    projectName: string = "Unknown Project";
+    change: string | number;
+    step?: StepData;
+    batch?: BatchData;
+    artifactsV2?: ArtifactData[];
+    stream?: StreamData;
+    
+    get events(): EventData[] {
+        return this._events;
+    };
+
+    get issues(): IssueData[] {
+        return this._issues;
+    }
+
+    @computed
+    get warnings(): EventData[] {
+        return this.events.filter(e => e.severity === EventSeverity.Warning);
+    };
+
+    @computed
+    get errors(): EventData[] {
+        return this.events.filter(e => e.severity === EventSeverity.Error);
+    }
+
+    get percentComplete(): number | undefined {
+        if(!this.step) return undefined;
+
+        return getStepPercent(this.step);
+    }
+
+    get summary(): string {
+        if(this.step) return getStepSummaryMarkdown(this.jobDetails, this.step.id);
+        if(this.batch) return getBatchSummaryMarkdown(this.jobDetails, this.batch.id);
+        return "";
+    }
+    
+    get agentId(): string | undefined {
+        return this.batch?.agentId ?? this.jobDetails.batchByStepId(this.step?.id)?.agentId;
+    }
+    
+    get crumbs(): BreadcrumbItem[] {
+        if(!this.stream) return [];
+
+        const data = this.jobDetails.jobData!;
+
+        const tab = this.stream.tabs?.find(tab => {
+            const jobTab = tab as GetJobsTabResponse;
+            return !!jobTab.templates?.find(template => template === data.templateId);
+        });
+
+        let streamLink = `/stream/${this.stream.id}`;
+        if(tab) streamLink += `?tab=${tab.title}`;
+
+        return [
+            {
+                text: this.projectName ?? "Unknown Project",
+                link: `/project/${this.stream?.project?.id}`
+            },
+            {
+                text: this.stream.name,
+                link: streamLink
+            },
+            {
+                text: `${data?.name ?? ""} - ${this.clText}`,
+                link: `/job/${data?.id}`
+            },
+            {
+                text: this.jobName,
+                link: this.step ? `/job/${data?.id}?step=${this.step?.id}` : `/job/${data?.id}?batch=${this.batch?.id}`
+            },
+            {
+                text: this.jobName + " (Log)"
+            },
+        ];
+    }
+
+    get crumbTitle(): string | undefined {
+        if(!this.stream) return undefined;
+        
+        return `Horde - ${this.clText}: ${this.jobDetails.jobData?.name ?? ""} - ${this.jobName}`;
+    }
+
+    get clText(): string {
+        const data = this.jobDetails.jobData;
+
+        if(!data) return "";
+
+        let clText = "";
+        if(data.preflightChange) {
+            clText = `Preflight ${data.preflightChange} - Base ${data.change ? data.change : "Latest"}`;
+        } else {
+            clText = `${data.change ? "CL " + data.change : "Latest CL"}`;
+        }
+
+        return clText;
+    }
+
+    constructor(logData: LogData, startLine?: number) {
+        super(logData, startLine);
+
+        makeObservable(this);
+
+        this._events = [];
+        this._issues = [];
+    }
+
+    async init() {
+        await super.init();
+        this.jobDetails = new JobDetailsV2(this.logData.jobId);
+
+        // JobDetailsV2 extends PollBase, so jobData is asynchronously fetched. Wait for that to finish.
+        await waitFor(() => this.jobDetails.jobData);
+
+        if(this.jobDetails.jobError) this.setFatalError(this.jobDetails.jobError);
+
+        this.updateStepAndBatch();
+
+        await this.updateArtifacts();
+
+        this.startTime = moment.utc(this.step?.startTime ?? this.batch?.startTime);
+        
+        // Determines if this log source should continue polling
+        this.setActive(this.getLogActive());
+
+        this.stream = this.jobDetails.stream;
+
+        // Determine project name
+        this.projectName = this.stream?.project?.name ?? this.projectName;
+        if(this.projectName === "Engine") this.projectName = "UE4";
+
+        this.change = this.jobDetails.jobData?.change ?? "Latest";
+        
+        // Determine job name
+        if(this.step) this.jobName = this.jobDetails.getStepName(this.step.id) ?? "Unknown Step Node";
+        if(this.batch) this.jobName = `Batch-${this.batch.id}`;
+        
+        this.agentTelemetry.set(
+            this.agentId ?? "", 
+            this.startTime.toDate(), 
+            this.step?.finishTime ? new Date(this.step.finishTime) : undefined
+        );
+
+        await this.updateEvents();
+        await this.updateIssues();
+
+        if(this.active) this.start();
+    }
+
+    async poll() {
+        super.poll();
+
+        await this.updateEvents();
+        await this.updateIssues();
+        
+        this.setActive(this.getLogActive());
+    }
+
+    private async updateEvents() {
+        try {
+            const events = await backend.getLogEvents(this.logId);
+            runInAction(() => {
+                this._events = events;
+            })
+        } catch (reason) {
+            errorBarStore.set({title: "Error Getting Events", message: reason}, true);
+        }
+    }
+
+    private async updateIssues() {
+        try {
+            // Update Issues
+            const issueQuery: IssueQuery = { jobId: this.jobDetails.jobId, stepId: this.step?.id, label: undefined, count: 50};
+            // Get both resolved and unresolved issues
+            const unresolvedIssues = await backend.getIssues(issueQuery);
+            const resolvedIssues = await backend.getIssues({...issueQuery, resolved: true});
+            
+            runInAction(() => {
+                this._issues = [...unresolvedIssues, ...resolvedIssues];
+            })
+        } catch (reason) {
+            errorBarStore.set({title: "Error Getting Issues", message: reason}, true);
+        }
+    }
+
+    private async updateArtifacts() {
+        if (
+            !! this.artifactsV2 
+            || !this.logId
+            || !this.step
+            || this.getLogActive()
+        ) {
+            return;
+        }
+
+        const key = `job:${this.jobDetails.jobId!}/step:${this.step.id}`;
+        try {
+            const artifactResponse = await backend.getJobArtifactsV2(undefined, [key]);
+            runInAction(() => {
+                this.artifactsV2 = artifactResponse.artifacts;
+            })
+        } catch (err) {
+            errorBarStore.set({title: "Error Getting Artifacts", message: err}, true);
+        }
+    }
+
+    download(json: boolean) {
+        const name = `++${this.projectName}+${this.stream?.name}-CL ${this.change}-${this.jobName}.${json ? "jsonl" : "log"}`;
+        backend.downloadLog(this.logId, name, json);
+    }
+
+    /**
+     * Returns true if related batch or step do not have a finish time. False otherwise.
+     */
+    private getLogActive(): boolean {
+        this.updateStepAndBatch();
+
+        // If related batch has a finish time, log isn't active
+        if(this.batch) return !this.batch.finishTime;
+
+        // If step has a finish time, log isn't active
+        if(this.step) return !this.step.finishTime;
+
+        return false;
+    }
+
+    private updateStepAndBatch() {
+        this.step = this.jobDetails.stepByLogId(this.logId);
+        this.batch = this.jobDetails.batches.find(b => b.logId === this.logId);
+    }
+}
+
+/**
+ * Log source for lease logs. Contains extra information about leases.
+ */
+export class LeaseLogSource extends LogSource {
+    leaseId: string;
+    agent?: AgentData;
+    lease?: LeaseData;
+
+    get agentId(): string | undefined {
+        return this.lease?.agentId;
+    }
+
+    get crumbs(): BreadcrumbItem[] {
+        return [
+            {
+                text: `Agents`,
+                link: '/agents'
+            },
+            {
+                text: `${this.agent?.name}`,
+                link: `/agents?agentId=${this.agent?.id}`
+            },
+            {
+                text: `${this.lease?.name ?? this.leaseId}`,
+                link: `/agents?agentId=${this.agent?.id}`
+            }
+        ];
+    }
+    
+    get summary(): string {
+        const duration = getLeaseElapsed(this.lease) 
+        if(duration){
+            return `${!!this.lease?.finishTime ? "Ran" : "Running"} on [${this.agent?.id}](?agentId=${encodeURIComponent(this.agent?.id ?? "")}) for ${duration}`;
+        }
+
+        return "";
+    }
+
+    constructor(logData: LogData, startLine?: number) {
+        super(logData, startLine);
+
+        if(!logData.leaseId) {
+            errorBarStore.set({title: "Failed to Create LeaseLogSource", message: "Bad Lease ID"});
+        } else {
+            this.leaseId = logData.leaseId!;
+        }
+    }
+
+    async init() {
+        await super.init();
+
+        this.lease = await backend.getLease(this.leaseId); 
+
+        try {
+            this.agent = await backend.getAgent(this.lease.agentId!);
+        } catch (reason) {
+            this.agent = {
+                id: "missing-agent",
+                name: "Missing Agent"
+            } as AgentData;
+        }
+
+        this.leaseUpdated();
+
+        if(this.active) this.start();
+    }
+
+    async poll() {
+        if(!this.lease) {
+            super.poll();
+            return;
+        }
+
+        this.updateAgent();
+        this.updateLease();
+        
+        this.leaseUpdated();
+
+        super.poll();
+    }
+
+    async updateAgent() {
+        try {
+            const agent = await backend.getAgent(this.agentId!);
+            runInAction(() => {
+                this.agent = agent;
+            })
+        } catch (reason) {
+            errorBarStore.set({title: "Error Updating Agent", message: reason}, true);
+        }
+    }
+
+    async updateLease() {
+        try {
+            const lease = await backend.getLease(this.leaseId);
+            runInAction(() => {
+                this.lease = lease;
+            })
+        } catch (reason) {
+            errorBarStore.set({title: "Error Updating Lease", message: reason}, true);
+        }
+    }
+
+    /**
+     * Secondary updates to run when the lease log source updates
+     */
+    leaseUpdated() {
+        this.startTime = moment.utc(this.lease!.startTime);
+        this.setActive(!this.lease!.finishTime);
+        this.agentTelemetry.set(this.lease?.agentId ?? "", new Date(this.startTime as any), this.lease?.finishTime ? new Date(this.lease.finishTime) : undefined);
+    }
+
+    download(json: boolean) {
+        const name = `${this.agent!.name}-Lease-${this.leaseId}.${json ? "jsonl" : "log"}`;
+        backend.downloadLog(this.logId, name, json);
+    }
+}
\ No newline at end of file
diff --git a/Engine/Source/Programs/Horde/HordeDashboard/src/components/logs/LogStyle.tsx b/Engine/Source/Programs/Horde/HordeDashboard/src/components/logs/LogStyle.tsx
new file mode 100644
--- /dev/null
+++ b/Engine/Source/Programs/Horde/HordeDashboard/src/components/logs/LogStyle.tsx
@@ -0,0 +1,256 @@
+// Copyright Epic Games, Inc. All Rights Reserved.
+
+import { mergeStyleSets } from "@fluentui/react";
+import dashboard from "../../backend/Dashboard";
+import { DashboardPreference } from "horde/backend/Api";
+import { hexToRgb, isBright } from "horde/base/utilities/colors";
+
+export type LogMetricType = {
+   lineHeight: number;
+   fontSize: number;
+}
+
+export const logMetricNormal = {
+   lineHeight: 24,
+   fontSize: 11
+}
+
+const adjustForDisplayScale = /*(window as any).safari === undefined &&*/ window.devicePixelRatio > 1.25;
+
+export const logMetricSmall = {
+   // note, making this any larger limits log line range for overflow
+   // make height even to avoid rounding issues
+   lineHeight: adjustForDisplayScale ? 14 : 18,
+   // don't go below 10pt for legibility
+   fontSize: 10
+}
+
+let _lineRenderStyleNormal: any;
+let _lineRenderStyleSmall: any;
+let _logStyleBase: any;
+let _logStyleNormal: any;
+let _logStyleSmall: any;
+
+export const getLogStyles = () => {
+
+   const lineRenderStyleNormal = _lineRenderStyleNormal ?? mergeStyleSets({
+      logLine: {
+         padding: 0,
+         height: logMetricNormal.lineHeight,
+         tabSize: "3",
+         fontFamily: "Horde Cousine Regular, monospace, monospace",
+         fontSize: logMetricNormal.fontSize,
+         whiteSpace: "pre-wrap"
+
+      }
+   });
+
+   const lineRenderStyleSmall = _lineRenderStyleSmall ?? mergeStyleSets({
+      logLine: {
+         padding: 0,
+         height: logMetricSmall.lineHeight,
+         tabSize: "3",
+         fontFamily: "Horde Cousine Regular, monospace, monospace",
+         fontSize: logMetricSmall.fontSize,
+         whiteSpace: "pre-wrap"
+
+      }
+   });
+
+   const warningColorPref = dashboard.getPreference(DashboardPreference.ColorWarning);
+   const errorColorPref = dashboard.getPreference(DashboardPreference.ColorError);
+   const warningHighlight = warningColorPref ? isBright(warningColorPref) ? "brightness(0.9)" : "brightness(1.2)" : "";
+   const errorHighlight = errorColorPref ? isBright(errorColorPref) ? "brightness(0.9)" : "brightness(1.2)" : "";
+
+   const logStyleBase = _logStyleBase ?? mergeStyleSets({
+      container: {
+         overflow: 'auto',         
+         marginTop: 8,
+      },
+      logLine: [
+         {
+            fontFamily: "Horde Cousine Regular, monospace, monospace"
+         }
+      ],
+      logLineOuter: {
+      },
+      itemWarning: [
+         {
+            background: warningColorPref ? `rgb(${hexToRgb(warningColorPref)}, 0.2)` : dashboard.darktheme ? "#302402": "#FEF8E7"
+         }
+      ],
+      errorButton: {
+         backgroundColor: errorColorPref ? errorColorPref : dashboard.darktheme ? "#9D1410" : "#EC4C47",
+         borderStyle: "hidden",
+         color: "#FFFFFF",
+         selectors: {
+            ':active,:hover': {
+               color: "#FFFFFF",
+               backgroundColor: errorColorPref ? errorColorPref : "#DC3C37",
+               filter: errorHighlight
+            }
+         }
+      },
+      errorButtonDisabled: {         
+         color: dashboard.darktheme ? "#909398" : "#616E85 !important",
+         backgroundColor: dashboard.darktheme ? "#1F2223" : "#f3f2f1"
+      },
+      warningButton: {
+         backgroundColor: warningColorPref ? warningColorPref : dashboard.darktheme ? "#9D840E" : "#F7D154",
+         borderStyle: "hidden",
+         selectors: {
+            ':active,:hover': {
+               backgroundColor: warningColorPref ? warningColorPref : "#E7C144",
+               filter: warningHighlight
+            }
+         }
+      },
+      warningButtonDisabled: {
+         color: dashboard.darktheme ? "#909398" : undefined,
+         backgroundColor: dashboard.darktheme ? "#1F2223" : "#f3f2f1"
+      },
+      gutter: [
+         {
+            padding: 0,
+            margin: 0,
+            paddingTop: 0,
+            paddingBottom: 0,
+            paddingRight: 14,
+            marginTop: 0,
+            marginBottom: 0,
+         }
+      ],
+      gutterError: [
+         {
+            borderLeftStyle: 'solid',
+            borderLeftColor: errorColorPref ? errorColorPref : "#EC4C47",
+            borderLeftWidth: 6,
+            padding: 0,
+            margin: 0,
+            paddingTop: 0,
+            paddingBottom: 0,
+            paddingRight: 8,
+            marginTop: 0,
+            marginBottom: 0,
+         }
+      ],
+      gutterWarning: [
+         {
+            borderLeftStyle: 'solid',
+            borderLeftColor: warningColorPref ? warningColorPref : "#F7D154",
+            borderLeftWidth: 6,
+            padding: 0,
+            margin: 0,
+            paddingTop: 0,
+            paddingBottom: 0,
+            paddingRight: 8,
+            marginTop: 0,
+            marginBottom: 0,
+         }
+      ],
+      itemError: [
+         {
+           background: errorColorPref ? `rgb(${hexToRgb(errorColorPref)}, 0.2)` : dashboard.darktheme ? "#330606" : "#FEF6F6",
+         }
+      ]
+   });
+
+
+   const logStyleNormal = _logStyleNormal ?? mergeStyleSets(logStyleBase, {
+
+      container: {
+         selectors: {
+            '.ms-List-cell': {
+               height: logMetricNormal.lineHeight,
+               lineHeight: logMetricNormal.lineHeight
+            }
+         }
+      },
+      logLine: [
+         {
+            fontSize: logMetricNormal.fontSize,
+            selectors: {
+               "#infoview": {
+                  opacity: 0
+               },
+               ":hover #infoview": {
+                  opacity: 1
+               },
+            },
+   
+         }
+      ],
+      gutter: [
+         {
+            height: logMetricNormal.lineHeight
+         }
+      ],
+      gutterError: [
+         {
+            height: logMetricNormal.lineHeight
+         }
+      ],
+      gutterWarning: [
+         {
+            height: logMetricNormal.lineHeight
+         }
+      ]
+   });
+
+   const logStyleSmall = _logStyleSmall ?? mergeStyleSets(logStyleBase, {
+
+      container: {
+         selectors: {
+            '.ms-List-cell': {
+               height: logMetricSmall.lineHeight,
+               lineHeight: logMetricSmall.lineHeight
+            }
+         }
+      },
+      logLine: [
+         {
+            fontSize: logMetricSmall.fontSize,
+            selectors: {
+               "#infoview": {
+                  opacity: 0
+               },
+               ":hover #infoview": {
+                  opacity: 1
+               },
+            },
+   
+         }
+      ],
+      gutter: [
+         {
+            height: logMetricSmall.lineHeight
+         }
+      ],
+      gutterError: [
+         {
+            height: logMetricSmall.lineHeight
+         }
+      ],
+      gutterWarning: [
+         {
+            height: logMetricSmall.lineHeight
+         }
+      ]
+   });
+
+   _lineRenderStyleNormal = lineRenderStyleNormal;
+   _lineRenderStyleSmall = lineRenderStyleSmall;
+   _logStyleBase = logStyleBase;
+   _logStyleNormal = logStyleNormal;
+   _logStyleSmall = logStyleSmall;
+
+   return {
+      lineRenderStyleNormal:lineRenderStyleNormal,
+      lineRenderStyleSmall:lineRenderStyleSmall,
+      logStyleBase:logStyleBase,
+      logStyleNormal:logStyleNormal,
+      logStyleSmall:logStyleSmall
+   }
+}
+
+
diff --git a/Engine/Source/Programs/Horde/HordeDashboard/src/components/logs/LogView.tsx b/Engine/Source/Programs/Horde/HordeDashboard/src/components/logs/LogView.tsx
new file mode 100644
--- /dev/null
+++ b/Engine/Source/Programs/Horde/HordeDashboard/src/components/logs/LogView.tsx
@@ -0,0 +1,1499 @@
+// Copyright Epic Games, Inc. All Rights Reserved.
+
+import { Callout, DefaultButton, DetailsList, DetailsListLayoutMode, DetailsRow, DirectionalHint, Dropdown, FocusZone, FocusZoneDirection, FontIcon, IColumn, Icon, IconButton, IContextualMenuItem, IContextualMenuProps, IDetailsListProps, ITextField, List, mergeStyles, MessageBar, MessageBarType, Modal, PrimaryButton, ProgressIndicator, ScrollToMode, Selection, SelectionMode, SelectionZone, Separator, Spinner, SpinnerSize, Stack, Text, TextField, TooltipHost } from '@fluentui/react';
+import { observer } from 'mobx-react-lite';
+import moment from 'moment-timezone';
+import React, { useEffect, useId, useState } from 'react';
+import { Link, useLocation, useNavigate, useParams } from 'react-router-dom';
+import backend from '../../backend';
+import { ArtifactContextType, DashboardPreference, EventSeverity, GetChangeSummaryResponse, GetJobStepRefResponse, GetLogEventResponse, LogLevel } from '../../backend/Api';
+import { CommitCache } from '../../backend/CommitCache';
+import dashboard from '../../backend/Dashboard';
+import { Markdown } from '../../base/components/Markdown';
+import { useWindowSize } from '../../base/utilities/hooks';
+import { displayTimeZone, getElapsedString } from '../../base/utilities/timeUtils';
+import { getHordeStyling } from '../../styles/Styles';
+import { getHordeTheme } from '../../styles/theme';
+import { AgentTelemetrySparkline } from "../agents/AgentTelemetrySparkline";
+import { JobArtifactsModal } from '../artifacts/ArtifactsModal';
+import { Breadcrumbs } from '../Breadcrumbs';
+import { ChangeContextMenu, ChangeContextMenuTarget } from '../ChangeButton';
+import { HistoryModal } from '../agents/HistoryModal';
+import { IssueModalV2 } from '../IssueViewV2';
+import { useQuery } from "horde/base/utilities/hooks";
+import { LogItem, renderLine } from './LogRender';
+import { StepRefStatusIcon } from '../StatusIcon';
+import { TopNav } from '../TopNav';
+import { isBright } from 'horde/base/utilities/colors';
+import { JobLogSource, LogSource } from './LogSource';
+import { JobDetailsV2 } from '../jobDetailsV2/JobDetailsViewCommon';
+import { LogHandler } from './LogHandler';
+import { updateLineQuery, getQueryLine } from 'horde/base/utilities/logUtils';
+
+let listRef: List | undefined;
+
+const selection = new Selection({ selectionMode: SelectionMode.multiple });
+
+const searchBox = React.createRef<ITextField>();
+
+let curSearchIdx = 0;
+
+let logListKey = 0;
+
+// have to be available to a global keyboard handler
+let globalHandler: LogHandler | undefined;
+let globalSearchState: { search?: string, results?: number[], curRequest?: any, text?: string } | undefined;
+
+const searchUp = (navigate: any) => {
+
+   if (globalSearchState?.results?.length) {
+
+      curSearchIdx--;
+      if (curSearchIdx < 0) {
+         curSearchIdx = globalSearchState.results.length - 1;
+      }
+
+      globalHandler?.stopTrailing();
+      let lineIdx = globalSearchState.results[curSearchIdx];
+
+      if (globalHandler) {
+         globalHandler.currentLine = lineIdx + 1;
+      }
+
+      updateLineQuery(lineIdx + 1, navigate);
+
+      lineIdx -= 10;
+      if (lineIdx < 0) {
+         lineIdx = 0;
+      }
+
+      if (globalHandler) {
+         listRef?.scrollToIndex(lineIdx, (index) => { return globalHandler!.lineHeight; }, ScrollToMode.top);
+      }
+
+
+      globalHandler?.externalUpdate();
+
+   }
+
+}
+
+const searchDown = (navigate: any) => {
+
+   if (globalSearchState?.results?.length) {
+
+      curSearchIdx++;
+      if (curSearchIdx >= globalSearchState.results.length) {
+         curSearchIdx = 0;
+      }
+
+      globalHandler?.stopTrailing();
+
+      let lineIdx = globalSearchState.results[curSearchIdx];
+
+      if (globalHandler) {
+         globalHandler.currentLine = lineIdx + 1;
+      }
+
+      updateLineQuery(lineIdx + 1, navigate);
+
+      lineIdx -= 10;
+      if (lineIdx < 0) {
+         lineIdx = 0;
+      }
+
+      if (globalHandler) {
+         listRef?.scrollToIndex(lineIdx, (index) => { return globalHandler!.lineHeight; }, ScrollToMode.top);
+      }
+
+      globalHandler?.externalUpdate();
+
+   }
+
+}
+
+const commitCache = new CommitCache();
+
+const StepHistoryModal: React.FC<{ jobDetails: JobDetailsV2, stepId: string | undefined, onClose: () => void }> = observer(({ jobDetails, stepId, onClose }) => {
+
+   const navigate = useNavigate();
+   const location = useLocation();
+   const [commitState, setCommitState] = useState<{ target?: ChangeContextMenuTarget, commit?: GetChangeSummaryResponse, rangeCL?: number }>({});
+   const [stepHistory, setStepHistory] = useState<GetJobStepRefResponse[] | undefined>(undefined);
+
+   const { hordeClasses } = getHordeStyling();
+   const hordeTheme = getHordeTheme();
+
+   const jobData = jobDetails.jobData;
+
+   if (!jobData || !jobData.streamId || !jobData.templateId) {
+      return null;
+   }
+
+   if (stepHistory === undefined) {
+      backend.getJobStepHistory(jobData.streamId, jobDetails.getStepName(stepId, false), 1024, jobData!.templateId!).then(r => {
+         setStepHistory(r);
+      })
+
+      return <Modal isOpen={true} styles={{ main: { padding: 8, width: 1084, height: '624px', backgroundColor: hordeTheme.horde.contentBackground } }} className={hordeClasses.modal} onDismiss={() => { onClose() }}>
+         <Stack style={{ paddingTop: 24 }} horizontalAlign='center' tokens={{ childrenGap: 18 }}>
+            <Stack>
+               <Text variant='mediumPlus'>Loading Step History</Text>
+            </Stack>
+            <Spinner size={SpinnerSize.large} />
+         </Stack>
+      </Modal>
+   }
+
+   // subscribe
+   if (commitCache.updated) { }
+   if (jobDetails.updated) { }
+
+   if (!stepId || !stepHistory?.length) {
+      return null;
+   }
+
+   const step = jobDetails.stepById(stepId);
+
+   type HistoryItem = {
+      ref: GetJobStepRefResponse;
+   };
+
+   const items: HistoryItem[] = stepHistory.map(h => {
+      return { ref: h }
+   });
+
+   const columns = [
+      { key: 'column1', name: 'Name', minWidth: 400, maxWidth: 400, isResizable: false },
+      { key: 'column2', name: 'Change', minWidth: 100, maxWidth: 100, isResizable: false },
+      { key: 'column3', name: 'Started', minWidth: 180, maxWidth: 180, isResizable: false },
+      { key: 'column4', name: 'Agent', minWidth: 100, maxWidth: 100, isResizable: false },
+      { key: 'column5', name: 'Duration', minWidth: 110, maxWidth: 110, isResizable: false },
+   ];
+
+   const renderItem = (item: HistoryItem, index?: number, column?: IColumn) => {
+
+      if (!column) {
+         return <div />;
+      }
+
+      const ref = item.ref;
+
+
+      if (column.name === "Name") {
+         return <Stack horizontal verticalAlign="center" tokens={{ childrenGap: 0, padding: 0 }} style={{ width: "100%", height: "100%" }} >{<StepRefStatusIcon stepRef={ref} />}<Text>{step?.name}</Text></Stack>;
+      }
+
+      if (column.name === "Change") {
+
+         if (item.ref.change && jobDetails.stream?.id) {
+
+            return <Stack verticalAlign="center" horizontalAlign="center" tokens={{ childrenGap: 0, padding: 0 }} style={{ width: "100%", height: "100%" }} onClick={async (ev) => {
+               ev?.stopPropagation();
+               ev?.preventDefault();
+
+               let commit = commitCache.getCommit(jobDetails.stream!.id, item.ref.change);
+
+               if (!commit) {
+                  await commitCache.set(jobDetails.stream!.id, [item.ref.change]);
+               }
+
+               commit = commitCache.getCommit(jobDetails.stream!.id, item.ref.change);
+
+               if (commit) {
+                  const index = stepHistory.indexOf(item.ref);
+                  let rangeCL: number | undefined;
+                  if (index < stepHistory.length - 1) {
+                     rangeCL = stepHistory[index + 1].change;
+                  }
+                  setCommitState({ commit: commit, rangeCL: rangeCL, target: { point: { x: ev.clientX, y: ev.clientY } } })
+               }
+            }}>
+               <Text style={{ color: "rgb(0, 120, 212)" }}>{item.ref.change}</Text>
+
+            </Stack>
+         }
+
+         return null;
+      }
+
+      if (column.name === "Agent") {
+
+         const agentId = item.ref.agentId;
+
+         if (!agentId) {
+            return null;
+         }
+
+         const url = `${location.pathname}?agentId=${agentId}`;
+
+         return <Stack verticalAlign="center" horizontalAlign="center" tokens={{ childrenGap: 0, padding: 0 }} style={{ width: "100%", height: "100%" }}>
+            <a href={url} onClick={(ev) => { ev.preventDefault(); ev.stopPropagation(); navigate(url, { replace: true }); }}><Stack horizontal horizontalAlign={"end"} verticalFill={true} tokens={{ childrenGap: 0, padding: 0 }}><Text>{agentId}</Text></Stack></a>
+         </Stack>
+      }
+
+      if (column.name === "Started") {
+
+         if (ref.startTime) {
+
+            const displayTime = moment(ref.startTime).tz(displayTimeZone());
+            const format = dashboard.display24HourClock ? "HH:mm:ss z" : "LT z";
+
+            let displayTimeStr = displayTime.format('MMM Do') + ` at ${displayTime.format(format)}`;
+
+
+            return <Stack verticalAlign="center" horizontalAlign="start" tokens={{ childrenGap: 0, padding: 0 }} style={{ width: "100%", height: "100%" }}>
+               <Stack >{displayTimeStr}</Stack>
+            </Stack>;
+
+         } else {
+            return "???";
+         }
+      }
+
+      if (column.name === "Duration") {
+
+         const start = moment(ref.startTime);
+         let end = moment(Date.now());
+
+         if (ref.finishTime) {
+            end = moment(ref.finishTime);
+         }
+         if (item.ref.startTime) {
+            const time = getElapsedString(start, end);
+            return <Stack verticalAlign="center" horizontalAlign="end" tokens={{ childrenGap: 0, padding: 0 }} style={{ width: "100%", height: "100%" }}><Stack style={{ paddingRight: 8 }}><Text>{time}</Text></Stack></Stack>;
+         } else {
+            return "???";
+         }
+      }
+
+
+      return <Stack />;
+   }
+
+   const renderRow: IDetailsListProps['onRenderRow'] = (props) => {
+
+      if (props) {
+
+         const item = props!.item as HistoryItem;
+         const ref = item.ref;
+
+         const url = `/log/${ref.logId}`;
+
+         const commonSelectors = { ".ms-DetailsRow-cell": { "overflow": "visible", padding: 0 } };
+
+         if (ref.stepId === stepId && ref.jobId === jobDetails.jobId) {
+            props.styles = { ...props.styles, root: { background: `${hordeTheme.palette.neutralLight} !important`, selectors: { ...commonSelectors as any } } };
+         } else {
+            props.styles = { ...props.styles, root: { selectors: { ...commonSelectors as any } } };
+         }
+
+         return <Link to={url} onClick={(ev) => { if (!ev.ctrlKey) onClose() }}><div className="job-item"><DetailsRow {...props} /> </div></Link>;
+
+      }
+      return null;
+   };
+
+   return (<Modal isOpen={true} styles={{ main: { padding: 8, width: 1084, height: '624px', backgroundColor: hordeTheme.horde.contentBackground } }} className={hordeClasses.modal} onDismiss={() => { onClose() }}>
+      {commitState.target && <ChangeContextMenu target={commitState.target} job={jobDetails.jobData} commit={commitState.commit} rangeCL={commitState.rangeCL} onDismiss={() => setCommitState({})} />}
+      <Stack styles={{ root: { paddingTop: 8, paddingLeft: 24, paddingRight: 12, paddingBottom: 8 } }}>
+         <Stack tokens={{ childrenGap: 12 }}>
+            <Stack horizontal styles={{ root: { padding: 8 } }}>
+               <Stack style={{ paddingLeft: 8, paddingTop: 4 }} grow>
+                  <Text variant="mediumPlus" styles={{ root: { fontFamily: "Horde Open Sans SemiBold" } }}>{"Step History"}</Text>
+               </Stack>
+               <Stack grow horizontalAlign="end">
+                  <IconButton
+                     iconProps={{ iconName: 'Cancel' }}
+                     onClick={() => { onClose(); }}
+                  />
+               </Stack>
+            </Stack>
+
+            <Stack styles={{ root: { paddingLeft: 4, paddingRight: 0, paddingTop: 8, paddingBottom: 4 } }}>
+               <div style={{ overflowY: 'auto', overflowX: 'hidden', height: "504px" }} data-is-scrollable={true}>
+                  <DetailsList
+                     compact={true}
+                     isHeaderVisible={false}
+                     indentWidth={0}
+                     items={items}
+                     columns={columns}
+                     setKey="set"
+                     selectionMode={SelectionMode.none}
+                     layoutMode={DetailsListLayoutMode.justified}
+                     onRenderItemColumn={renderItem}
+                     onRenderRow={renderRow}
+                  />
+               </div>
+            </Stack>
+         </Stack>
+      </Stack>
+   </Modal>);
+});
+
+const LogProgressIndicator: React.FC<{ logSource: LogSource }> = observer(({ logSource }) => {
+
+   // subscribe
+   const details = (logSource as JobLogSource).jobDetails;
+   if (details?.updated) { }
+
+   const percentComplete = logSource.percentComplete;
+
+   return <Stack>
+      <ProgressIndicator percentComplete={percentComplete} barHeight={2} styles={{ root: { paddingLeft: 12, width: 300 } }} />
+   </Stack>
+});
+
+const LogLineIndicator: React.FC<{ lineIndex: number }> = ({ lineIndex }) => {
+
+   useQuery();
+
+   if (!globalHandler) {
+      return null;
+   }
+
+   let prefix = "";
+   const event = globalHandler.getLogEvent(globalHandler.currentLine);
+
+   const queryLine = getQueryLine();
+
+   if (queryLine) {
+      if (lineIndex === queryLine) {
+         prefix = ">>> ";
+      }
+   } else if (event && lineIndex >= event.lineIndex && (lineIndex < event.lineIndex + event.lineCount)) {
+      prefix = ">>> ";
+   } else if (globalSearchState?.results?.length && curSearchIdx < globalSearchState.results.length) {
+      if (lineIndex === globalSearchState.results[curSearchIdx]) {
+         prefix = ">>> ";
+      }
+   }
+
+   if (!prefix) {
+      return null;
+   }
+
+   return <span>{">>>"}</span>
+}
+
+const LogList: React.FC<{ logId: string }> = observer(({ logId }) => {
+
+   const windowSize = useWindowSize();
+   const query = useQuery();
+   const location = useLocation();
+   const navigate = useNavigate();
+   const [tsFormat, setTSFormat] = useState('relative');
+   const [handler, setHandler] = useState<LogHandler | undefined>(undefined);
+   const [searchState, setSearchState] = useState<{ search?: string, results?: number[], curRequest?: any, text?: string }>({});
+   const [issueHistory, setIssueHistory] = useState(false);
+   const [logHistory, setLogHistory] = useState(false);
+   const [logArtifacts, setLogArtifacts] = useState("");
+   const [logTelemetry, setLogTelemetry] = useState(false);
+   const [logError, setLogError] = useState("")
+
+   let [historyAgentId, setHistoryAgentId] = useState<string | undefined>(undefined);
+
+   if (!historyAgentId && query.get("agentId")) {
+      historyAgentId = query.get("agentId")!;
+   }
+
+   const artifactContext = !!query.get("artifactContext") ? query.get("artifactContext")! as ArtifactContextType : undefined;
+   const artifactPath = !!query.get("artifactPath") ? query.get("artifactPath")! : undefined;
+   const artifactId = !!query.get("artifactId") ? query.get("artifactId")! : undefined;
+
+   globalHandler = handler;
+   globalSearchState = searchState;
+   const vw = Math.max(document.documentElement.clientWidth, window.innerWidth || 0);
+
+   let inTimeout = false;
+
+   function searchSelectedText() {
+      const selection = window.getSelection();
+      const selected_text = selection ? selection.toString() : '';
+      
+      if (selected_text) {
+         setSearchState({...searchState, text: selected_text})
+      }
+   }
+
+   useEffect(() => {
+
+      const handler = (e: KeyboardEvent) => {
+
+
+         if (e.keyCode === 114) {
+            e.preventDefault();
+            if (e.shiftKey) {
+               searchUp(navigate);
+            } else {
+               searchDown(navigate);
+            }
+
+         }
+
+         if ((e.ctrlKey || e.metaKey) && e.keyCode === 70) {
+            e.preventDefault();
+
+            searchSelectedText();
+            
+            if (searchState.text !== '') {
+               searchBox.current?.setSelectionRange(0, searchBox.current.value!.length)
+            }
+
+            searchBox.current?.focus();
+         }
+      }
+
+      window.addEventListener("keydown", handler);
+
+      return () => {
+         globalHandler = globalSearchState = undefined;
+         window.removeEventListener("keydown", handler);
+      };
+
+   }, []);
+
+   const handleSearchBoxClick = (e: React.MouseEvent<HTMLElement, MouseEvent>) => {
+      if (e.button !== 0) return; // only capture left clicks
+
+      searchSelectedText();
+   };
+   
+
+   if (logError) {
+      return <Stack horizontalAlign='center' style={{ paddingTop: 24 }} >
+         <Text variant='mediumPlus'>{`Unable to load log data - ${logError}`}</Text>
+      </Stack>
+   }
+
+   if (handler && handler.logSource && handler.logSource.logId !== logId) {
+      selection.setItems([], true);
+      LogHandler.clear();
+
+      setHandler(undefined);
+   }
+
+   const hordeTheme = getHordeTheme();
+   const { hordeClasses, modeColors } = getHordeStyling();
+
+   const warningColorPref = dashboard.getPreference(DashboardPreference.ColorWarning);
+   const errorColorPref = dashboard.getPreference(DashboardPreference.ColorError);
+   const buttonNoneTextColor = dashboard.darktheme ? "#949898" : "#616e85";
+   const warningButtonTextColor = warningColorPref ? isBright(warningColorPref) ? "#2d3f5f" : "#F9F9FB" : dashboard.darktheme ? modeColors.text : buttonNoneTextColor;
+   const errorButtonTextColor = errorColorPref && isBright(errorColorPref) ? "#2d3f5f" : "#F9F9FB";
+
+   if (!logId) {
+      console.error("Bad log id settings up LogList");
+      return <div>Error</div>;
+   }
+
+   if (handler?.logSource?.fatalError) {
+
+      const error = `Error getting job data, please check that you are logged in and that the link is valid.\n\n${handler.logSource.fatalError}`;
+      return (
+         <Stack horizontal style={{ paddingTop: 48 }}>
+            <div key={`windowsize_logview1_${windowSize.width}_${windowSize.height}`} style={{ width: vw / 2 - 720, flexShrink: 0, backgroundColor: hordeTheme.horde.contentBackground }} />
+            <Stack horizontalAlign="center" style={{ width: 1440 }}>
+               <div style={{whiteSpace:"pre"}}> {error} </div>
+            </Stack>
+         </Stack>
+      )
+
+   }
+
+   if(!handler){
+      let lineIndexQuery: number | undefined;
+      if (query.get("lineIndex")) {
+         lineIndexQuery = parseInt(query.get("lineIndex")!);
+      } else if (query.get("lineindex")) {
+         lineIndexQuery = parseInt(query.get("lineindex")!) + 1;
+      }
+
+      const handler = new LogHandler();
+      handler.init(logId, lineIndexQuery ?? undefined).then(() => {
+         setHandler(handler);
+      });
+
+      return <Spinner size={SpinnerSize.large} />;
+   }
+
+   // subscribe
+   if (handler.updated) { }
+   
+   const logSource = handler.logSource!;
+   const logSourceIsJobLogSource = logSource instanceof JobLogSource;
+   const warnings = logSource.warnings.sort((a, b) => a.lineIndex - b.lineIndex);
+   const errors = logSource.errors.sort((a, b) => a.lineIndex - b.lineIndex);
+   const issues = logSource.issues;
+
+   handler.compact = !!(logSource.logItems?.length > 500000);
+
+   const onRenderCell = (item?: LogItem, index?: number, isScrolling?: boolean): JSX.Element => {
+
+      const LogCell: React.FC = observer(() => {
+
+         if (handler.updated) { }
+
+         if (!item) {
+            return (<div style={{ height: handler.lineHeight }} />);
+         }
+
+         const styles = handler.style;
+
+         if (!item.requested) {
+            let index = item.lineNumber - 50;
+            if (index < 0) {
+               index = 0;
+            }
+
+            logSource.loadLines(index, 100);
+
+            return <Stack key={`key_log_line_${item.lineNumber}`} style={{ width: "max-content", height: handler.lineHeight }}>
+               <div style={{ position: "relative" }}>
+                  <Stack className={styles.logLine} tokens={{ childrenGap: 8 }} horizontal disableShrink={true}>
+                     <Stack styles={{ root: { color: "#c0c0c0", width: 48, textAlign: "right", userSelect: "none", fontSize: handler.fontSize } }}>{item.lineNumber}</Stack>
+                     <Stack className={styles.logLine} horizontal disableShrink={true} >
+                        <Stack className={styles.gutter}></Stack>
+                        <Stack styles={{ root: { color: "#8a8a8a", width: 82, whiteSpace: "nowrap", fontSize: handler.fontSize, userSelect: "none" } }}> </Stack>
+                        <div className={styles.logLineOuter}> <Stack styles={{ root: { paddingLeft: 8, paddingRight: 8 } }}> </Stack></div>
+                     </Stack>
+                  </Stack>
+               </div>
+            </Stack>
+
+         }
+
+         const line = item.line;
+         const warning = line?.level === LogLevel.Warning;
+         const error = (line?.level === LogLevel.Error) || (line?.level === LogLevel.Critical);
+
+         const ev = logSource.events.find(event => event.lineIndex === item.lineNumber - 1);
+
+         if (ev && ev.issueId) {
+            item.issueId = ev.issueId;
+            item.issue = issues.find(issue => issue.id === ev.issueId);
+         }
+
+         const style = warning ? styles.itemWarning : error ? styles.itemError : styles.logLine;
+         const gutterStyle = warning ? styles.gutterWarning : error ? styles.gutterError : styles.gutter;
+
+         let timestamp = "";
+         let tsWidth = 82;
+         if (line && line.time && logSource.startTime) {
+
+            if (tsFormat === 'relative') {
+               const start = logSource.startTime;
+               const end = moment.utc(line.time);
+               if (end >= start) {
+                  const duration = moment.duration(end.diff(start));
+                  const hours = Math.floor(duration.asHours());
+                  const minutes = duration.minutes();
+                  const seconds = duration.seconds();
+                  timestamp = `[${hours <= 9 ? "0" + hours : hours}:${minutes <= 9 ? "0" + minutes : minutes}:${seconds <= 9 ? "0" + seconds : seconds}]`;
+               } else {
+                  timestamp = "[00:00:00]";
+               }
+            }
+
+            if (tsFormat === 'utc' || tsFormat === 'local') {
+               tsWidth = 112;
+               let tm = moment.utc(line.time);
+               if (tsFormat === 'local') {
+                  tm = tm.local();
+               }
+               timestamp = `[${tm.format("MMM DD HH:mm:ss")}]`;
+            }
+         }
+
+         const IssueButton: React.FC<{ item: LogItem, event: GetLogEventResponse }> = ({ item, event }) => {
+
+            const tooltipId = useId();
+
+            const error = event.severity === EventSeverity.Error;
+
+            const issueId = item!.issueId!.toString();
+
+            const newQuery = query.getCopy();
+            newQuery.set("issue", issueId);
+
+            // this href is for when log is copy/pasted, note that we need the onClick handler when on site, so don't reload page
+            const href = `${location.pathname}?${newQuery.toString()}`;
+
+            const fontStyle = {
+               fontFamily: "Horde Open Sans SemiBold, sans-serif, sans-serif", color: error ? errorButtonTextColor: warningButtonTextColor, fontSize: handler.fontSize, textDecoration: !item!.issue?.resolvedAt ? undefined : "line-through"
+            }
+
+            return <Stack style={{width: 80}}>
+               <TooltipHost
+                  content={timestamp}
+                  id={tooltipId}
+                  calloutProps={{ gapSpace: 0 }}
+                  styles={{ root: { display: 'inline-block', width: 80 } }}>
+                  <DefaultButton className={error ? styles.errorButton : styles.warningButton}
+                     href={href}
+                     style={{ padding: 0, margin: 0, minWidth: 65, width: 65, paddingLeft: 4, paddingRight: 6, height: "100%", fontWeight: "unset" }}
+                     onClick={(ev) => {
+                        ev.preventDefault();
+                        ev.stopPropagation();
+                        location.search = `?issue=${issueId}`
+                        setIssueHistory(true);
+                        navigate(location);
+                     }}>
+                     <Text variant="small" style={{ ...fontStyle }}>{`${issueId}`}</Text>
+                  </DefaultButton>
+               </TooltipHost>
+            </Stack>
+         }
+
+         const eyeColor = modeColors.text + "44";
+
+         return (
+            <Stack key={`key_log_line_${item.lineNumber}`} style={{ width: "max-content", height: handler.lineHeight }}
+               onMouseEnter={(ev) => {
+
+                  if (item.line?.time) {
+                     let time = item.line?.time;
+                     if (!time.endsWith("Z")) {
+                        time += "Z";
+                     }
+                     logSource?.agentTelemetry?.setCurrentTime(new Date(time));
+                  }
+               }}
+               onClick={() => {
+                  updateLineQuery(item.lineNumber, navigate);
+                  handler.currentLine = item.lineNumber;
+               }}>
+               <div style={{ position: "relative" }}>
+                  <Stack className={styles.logLine} style={{ position: "relative" }} tokens={{ childrenGap: 8 }} horizontal disableShrink={true}>
+                     <Stack horizontal styles={{ root: { color: "#c0c0c0", width: 80, textAlign: "right", userSelect: "none", fontSize: handler.fontSize } }}>
+                        <LogLineIndicator lineIndex={item.lineNumber} />
+                        <Stack styles={{ root: { color: "#c0c0c0", width: 80, textAlign: "right", userSelect: "none", fontSize: handler.fontSize } }}>{item.lineNumber}</Stack>
+                     </Stack>
+                     <Stack className={style} horizontal disableShrink={true}>
+                        <Stack className={gutterStyle}></Stack>
+                        {(!item.issueId || !ev) && <Stack styles={{ root: { color: "#8a8a8a", width: tsWidth, whiteSpace: "nowrap", fontSize: handler.fontSize, userSelect: "none" } }}> {timestamp}</Stack>}
+                        {!!item.issueId && !!ev && <IssueButton item={item} event={ev!} />}
+                        <div className={styles.logLineOuter}> <Stack styles={{ root: { paddingLeft: 8, paddingRight: 8, position: "relative", verticalAlign: "center" } }}> {renderLine(navigate, item.line, item.lineNumber, handler.lineRenderStyle, searchState.search)}
+                           <Stack id={`callout_target_${item?.lineNumber}`} style={{ position: "absolute", cursor: "pointer", userSelect: "none", left: "-12px", top: "0px", zIndex: 100 }} onClick={() => {
+                              handler.infoLine = item.lineNumber;
+                              handler.externalUpdate();
+                           }}><FontIcon id="infoview" style={{ fontSize: 14, color: eyeColor }} iconName="Eye" /></Stack>
+                           {handler.infoLine === item.lineNumber && <Callout
+                              styles={{ root: { padding: "32px 24px", maxWidth: 1300 } }}
+                              role="dialog"
+                              gapSpace={4}
+                              target={`#callout_target_${item?.lineNumber}`}
+                              isBeakVisible={true}
+                              beakWidth={12}
+                              onDismiss={() => {
+                                 handler.infoLine = undefined;
+                                 handler.externalUpdate();
+                              }}
+                              directionalHint={DirectionalHint.bottomCenter}
+                              setInitialFocus>
+                              <Stack style={{ maxWidth: 1140 }}>
+                                 <Stack style={{ paddingBottom: 24 }}>
+                                    <Text style={{ fontSize: 14, fontFamily: "Horde Open Sans SemiBold" }}>Structured Log Line</Text>
+                                 </Stack>
+                                 <Stack style={{ paddingLeft: 12 }}>
+                                    <Text style={{ fontSize: 11, whiteSpace: "pre-wrap", fontFamily: "Horde Cousine Regular" }}>{JSON.stringify(item.line, undefined, 2).replaceAll("\\r", "").replaceAll("\\n", "\n")}</Text></Stack>
+                              </Stack>
+                           </Callout>}
+
+                        </Stack>
+                        </div>
+                     </Stack>
+                  </Stack>
+               </div>
+            </Stack>
+
+         );
+      });
+
+      return <LogCell />
+   };
+
+   if (handler.trailing) {
+      listRef?.scrollToIndex(logSource.logData!.lineCount - 1, undefined, ScrollToMode.bottom);
+   }
+
+   const downloadProps: IContextualMenuProps = {
+      items: [
+         {
+            key: 'downloadLog',
+            text: 'Download Text',
+            onClick: () => {
+               logSource.download(false);
+            }
+         },
+         {
+            key: 'downloadJson',
+            text: 'Download JSON',
+            onClick: () => {
+               logSource.download(true);
+            }
+         },
+      ],
+      directionalHint: DirectionalHint.bottomRightEdge
+   };
+
+   let summaryText = logSource.summary;
+
+   let largeLogWarning = logSource.logItems?.length > 1500000;
+   let largeLogWarningText= `Warning: large log file (${logSource.logItems?.length} lines) may not display entirely.  Please download to view full log and consider reducing verbosity.`;
+
+   let baseUrl = location.pathname;
+
+   /**
+    * Executes a log search query and updates UI state accordingly.
+    * 
+    * @todo Improve logic, flow, and commenting for this function.
+    */
+   const doQuery = (newValue: string) => {
+
+      let currentLine = 0
+      if (handler.currentLine !== undefined) {
+         currentLine = handler.currentLine
+      }
+
+      if (!newValue.trim()) {
+
+         logListKey++;
+
+         setSearchState({
+               search: '',
+               curRequest: undefined,
+               results: undefined,
+               text: ''
+         });
+
+         // Reset current line and scroll to value before search was cleared
+         handler.stopTrailing();
+         let lineIdx = currentLine - 10;
+         if (lineIdx < 0) {
+            lineIdx = 0;
+         }
+
+         updateLineQuery(currentLine, navigate);
+
+         inTimeout = true;
+         setTimeout(() => {inTimeout = false; listRef?.scrollToIndex(lineIdx, () => { return handler.lineHeight; }, ScrollToMode.top);
+         }, 10)
+
+      } else if (newValue !== searchState.search) {
+
+         const request = backend.searchLog(logId, newValue, 0, 65535);
+         let newState = {
+            search: newValue,
+            curRequest: request,
+            text: newValue
+         }
+
+         request.then(logResponse => {
+
+            if (newState.curRequest === request) {
+
+               const lines = logResponse.lines;
+
+               if (lines.length) {
+
+                  logListKey++;
+
+                  // Get next occurrence after current line
+                  curSearchIdx = 0;
+                  for (const [index, value] of lines.entries()) {
+                     if (value + 1 >= currentLine) {
+                        curSearchIdx = index
+                        break
+                     }
+                  }
+
+                  handler.currentLine = lines[curSearchIdx] + 1
+
+                  setSearchState({
+                     search: newValue,
+                     curRequest: undefined,
+                     results: lines,
+                     text: newValue
+                  });
+
+
+                  handler.stopTrailing();
+                  let lineIdx = lines[curSearchIdx] - 10;
+                  if (lineIdx < 0) {
+                     lineIdx = 0;
+                  }
+
+                  updateLineQuery(lines[curSearchIdx] + 1, navigate);
+
+                  // oof
+                  inTimeout = true;
+                  setTimeout(() => {
+                     inTimeout = false;
+                     listRef?.scrollToIndex(lineIdx, () => { return handler.lineHeight; }, ScrollToMode.top);
+                  }, 250)
+
+               } else {
+                  setSearchState({
+                     search: newValue,
+                     curRequest: undefined,
+                     results: lines,
+                     text: newValue
+                  });
+
+               }
+
+            };
+         });
+
+         setSearchState(newState);
+
+      } else {
+
+         curSearchIdx = 0;
+
+         if (searchState.results?.length) {
+
+            updateLineQuery(searchState.results[curSearchIdx] + 1, navigate);
+
+            handler.stopTrailing();
+            let lineIdx = searchState.results[curSearchIdx] - 10;
+            if (lineIdx < 0) {
+               lineIdx = 0;
+            }
+            listRef?.scrollToIndex(lineIdx, (index) => { return handler.lineHeight; }, ScrollToMode.top);
+         }
+
+      }
+
+   }
+
+   //#region View Btn Menu Props
+
+   const menuProps: IContextualMenuProps = { items: [] };
+
+   if (logSourceIsJobLogSource) {
+      const stepArtifacts = (logSource as JobLogSource).artifactsV2;
+
+      const atypes = new Map<ArtifactContextType, number>();
+      const knownTypes = new Set<string>(["step-saved", "step-output", "step-trace"]);
+
+      stepArtifacts?.forEach(a => {
+         let c = atypes.get(a.type) ?? 0;
+         c++;
+         atypes.set(a.type, c);
+      });
+
+      const opsList: IContextualMenuItem[] = [];
+
+      const navigateToArtifacts = (context: string) => {
+         const search = new URLSearchParams(window.location.search);
+         search.set("artifactContext", encodeURIComponent(context));
+         const url = `${window.location.pathname}?` + search.toString();
+         navigate(url, { replace: true })
+      }
+
+      opsList.push({
+         key: 'stepops_artifacts_step',
+         text: "Logs",
+         iconProps: { iconName: "Folder" },
+         disabled: !atypes.get("step-saved"),
+         onClick: () => {
+            navigateToArtifacts("step-saved");
+         }
+      });
+
+      opsList.push({
+         key: 'stepops_artifacts_output',
+         text: "Temp Storage",
+         iconProps: { iconName: "MenuOpen" },
+         disabled: !atypes.get("step-output"),
+         onClick: () => {
+            navigateToArtifacts("step-output");
+         }
+      });
+
+      opsList.push({
+         key: 'stepops_artifacts_trace',
+         text: "Traces",
+         iconProps: { iconName: "SearchTemplate" },
+         disabled: !atypes.get("step-trace"),
+         onClick: () => {
+            navigateToArtifacts("step-trace");
+         }
+      });
+
+      const custom = stepArtifacts?.filter(a => !knownTypes.has(a.type)).sort((a, b) => a.type.localeCompare(b.type));
+      custom?.forEach(c => {
+         opsList.push({
+            key: `stepops_artifacts_${c.type}`,
+            text: c.description ?? c.name,
+            iconProps: { iconName: "Clean" },
+            onClick: () => { navigateToArtifacts(c.type) }
+         });
+      })
+
+
+      menuProps.items.push({
+         key: 'jobstep_artifacts',
+         text: `Artifacts`,
+         subMenuProps: {
+            items: opsList
+         }
+      })
+   }
+
+   menuProps.items.push({
+      key: 'jobstep_history',
+      text: 'Step History',
+      onClick: () => setLogHistory(true)
+   })
+
+   if (logSource.agentTelemetry) {
+      menuProps.items.push({
+         key: 'jobstep_agent_telemetry',
+         text: logTelemetry ? 'Hide Telemetry' : 'Show Telemetry',
+         onClick: () => {
+            logSource.agentTelemetry?.show(!logTelemetry);
+            setLogTelemetry(!logTelemetry)
+         }
+      })
+   }
+
+   //#endregion View Btn Menu Props
+
+   //#region Log Navigation Functions
+
+   function updateEvent() {
+
+      const event = handler?.getCurrentEvent();
+      if (!handler || !event) {
+         return;
+      }
+
+      handler.stopTrailing();
+
+      updateLineQuery(event.lineIndex + 1, navigate);
+      handler.externalUpdate();
+      listRef?.scrollToIndex(event.lineIndex - 10 < 0 ? 0 : event.lineIndex - 10, () => handler.lineHeight, ScrollToMode.top);
+   }
+
+   function prevError() {
+
+      if (!handler) {
+         return;
+      }
+
+      const event = handler.getPrevLogEvent();
+
+      if (event) {
+         handler.currentLine = event.lineIndex + 1;
+      }
+
+      updateEvent();
+
+   }
+
+   function prevErrorBlock() {
+
+      if (!handler) {
+         return;
+      }
+
+      const eventBlock = handler.getPrevLogEventBlock();
+
+      if (eventBlock) {
+         handler.currentLine = eventBlock.lineIndex + 1;
+      }
+
+      updateEvent();
+   }
+
+   function nextError() {
+
+      if (!handler) {
+         return;
+      }
+
+      const event = handler.getNextLogEvent();
+
+      if (event) {
+         handler.currentLine = event.lineIndex + 1;
+      }
+
+      updateEvent();
+   }
+
+   function nextErrorBlock() {
+
+      if (!handler) {
+         return;
+      }
+
+      const eventBlock = handler.getNextLogEventBlock();
+
+      if (eventBlock) {
+         handler.currentLine = eventBlock.lineIndex + 1;
+      }
+
+      updateEvent();
+   }
+
+   function prevWarning() {
+
+      if (!handler) {
+         return;
+      }
+
+      const event = handler.getPrevLogEvent(true);
+
+      if (event) {
+         handler.currentLine = event.lineIndex + 1;
+      }
+
+      updateEvent();
+
+   }
+
+   function prevWarningBlock() {
+
+      if (!handler) {
+         return;
+      }
+
+      const eventBlock = handler.getPrevLogEventBlock(true);
+
+      if (eventBlock) {
+         handler.currentLine = eventBlock.lineIndex + 1;
+      }
+
+      updateEvent();
+
+   }
+
+   function nextWarning() {
+      if (!handler) {
+         return;
+      }
+
+      const event = handler.getNextLogEvent(true);
+
+      if (event) {
+         handler.currentLine = event.lineIndex + 1;
+      }
+
+      updateEvent();
+
+   }
+
+   function nextWarningBlock() {
+      if (!handler) {
+         return;
+      }
+
+      const eventBlock = handler.getNextLogEventBlock(true);
+
+      if (eventBlock) {
+         handler.currentLine = eventBlock.lineIndex + 1;
+      }
+
+      updateEvent();
+
+   }
+
+   //#endregion Log Navigation Functions
+
+   let warningsText = "";
+   let errorText = "";
+
+   if (warnings?.length) {
+      warningsText = warnings.length.toString();
+      if (warnings.length >= 49) {
+         warningsText += "+";
+      }
+   } else {
+      warningsText = "No";
+   }
+
+   if (errors?.length) {
+      errorText = errors.length.toString();
+      if (errors.length >= 49) {
+         errorText += "+";
+      }
+   } else {
+      errorText = "No";
+   }
+
+
+   let heightAdjust = 285;
+   if (logTelemetry) {
+      // sparkline height
+      heightAdjust += 238;
+   }
+   if (largeLogWarning) {
+      heightAdjust += 55
+   }
+
+   const nextErrorEnabled: boolean = !!handler?.getNextLogEvent();
+   const prevErrorEnabled: boolean = !!handler?.getPrevLogEvent();
+   const nextErrorBlockEnabled: boolean = !!handler?.getNextLogEventBlock();
+   const prevErrorBlockEnabled: boolean = !!handler?.getPrevLogEventBlock();
+   const nextWarningEnabled: boolean = !!handler?.getNextLogEvent(true);
+   const prevWarningEnabled: boolean = !!handler?.getPrevLogEvent(true);
+   const nextWarningBlockEnabled: boolean = !!handler?.getNextLogEventBlock(true);
+   const prevWarningBlockEnabled: boolean = !!handler?.getPrevLogEventBlock(true);
+
+   let jobDetails: JobDetailsV2 | undefined;
+   if(logSourceIsJobLogSource) {
+      jobDetails = (logSource as JobLogSource).jobDetails;
+   }
+
+   return <Stack>
+      {logSourceIsJobLogSource && <IssueModalV2 issueId={query.get("issue")} popHistoryOnClose={issueHistory} />}
+      
+      {
+         logSourceIsJobLogSource 
+         && logHistory 
+         && 
+            <StepHistoryModal 
+               jobDetails={jobDetails!} 
+               stepId={jobDetails!.stepByLogId(logId)?.id} 
+               onClose={() => setLogHistory(false)} 
+            />
+      }
+      
+      {
+         logSourceIsJobLogSource 
+         && !!artifactContext 
+         && logArtifacts !== "legacy" 
+         && 
+            <JobArtifactsModal 
+               jobId={jobDetails!.jobData!.id} 
+               stepId={jobDetails!.stepByLogId(logId)?.id!} 
+               artifacts={(logSource as JobLogSource).artifactsV2} 
+               contextType={artifactContext} artifactPath={artifactPath} 
+               artifactId={artifactId} 
+               onClose={() => { navigate(window.location.pathname, { replace: true }) }} 
+            />
+      }
+
+      {!!historyAgentId && <HistoryModal agentId={historyAgentId} onDismiss={() => { navigate(baseUrl, { replace: true }); setHistoryAgentId(undefined) }} />}
+      
+      <Breadcrumbs items={logSource?.crumbs ?? []} title={logSource?.crumbTitle} />
+      
+      <Stack tokens={{ childrenGap: 0 }} style={{ backgroundColor: hordeTheme.horde.contentBackground, paddingTop: 12 }}>
+         <Stack horizontal >
+            <div key={`windowsize_logview1_${windowSize.width}_${windowSize.height}`} style={{ width: vw / 2 - (1440 / 2), flexShrink: 0 }} />
+            <Stack tokens={{ childrenGap: 0, maxWidth: 1440 }} disableShrink={true} styles={{ root: { width: "100%", backgroundColor: hordeTheme.horde.contentBackground, paddingLeft: 4, paddingRight: 24, paddingTop: 12 } }}>
+               <Stack horizontal style={{ paddingBottom: 4 }}>
+                  <Stack className={hordeClasses.button} horizontal horizontalAlign={"start"} verticalAlign="center" tokens={{ childrenGap: 8 }}>
+                     <Stack horizontal tokens={{ childrenGap: 2 }}>
+
+                        <Stack horizontal verticalAlign='center'>
+                           <Stack horizontal verticalAlign='center' style={{position: 'relative'}}>
+                              <TooltipHost content={nextErrorEnabled ? 'Navigate to next error': ''}>
+                                 <DefaultButton disabled={!nextErrorEnabled} className={nextErrorEnabled ? handler.style.errorButton : handler.style.errorButtonDisabled}
+                                    text={`${errorText} ${errors.length === 1 ? "Error" : "Errors"}`}
+                                    onClick={() => {
+                                       nextError();
+                                    }}
+                                    style={{ color: errors.length ? errorButtonTextColor : buttonNoneTextColor, height: 30, fontSize: 19, paddingLeft: errors.length ? 12 : 8, borderRadius: errors.length ? '4px 0 0 4px' : '4px' }}>
+                                    {!!errors.length && <Icon style={{ fontSize: 19, paddingLeft: 4, paddingRight: 0 }} iconName='ChevronDown' />}
+                                 </DefaultButton>
+                              </TooltipHost>
+                              {!!errors.length && <div style={{backgroundColor: errorButtonTextColor, height: 18, width: 1, position: 'absolute', right: -0.5, zIndex: 1 }} />}
+                           </Stack>
+                           {!!errors.length && <TooltipHost content={'Navigate to previous error'}>
+                              <IconButton disabled={!prevErrorEnabled} className={prevErrorEnabled ? handler.style.errorButton : handler.style.errorButtonDisabled} style={{ color: errorButtonTextColor, height: 30, fontSize: 19, padding: 2, borderRadius: '0 4px 4px 0' }} iconProps={{ iconName: 'ChevronUp' }} onClick={(event: any) => {
+                                 event?.stopPropagation();
+                                 prevError();
+                              }} />
+                           </TooltipHost>}
+                        </Stack>
+
+                        {!!errors.length && <Stack horizontal verticalAlign='center' style={{paddingLeft: 1}} >
+                           <Stack horizontal verticalAlign='center' style={{position: 'relative'}}>
+                              <TooltipHost content={'Skip to next block of errors or issue'}>
+                                 <IconButton disabled={!nextErrorBlockEnabled} className={nextErrorBlockEnabled ? handler.style.errorButton : handler.style.errorButtonDisabled}
+                                    style={{ color: errorButtonTextColor, height: 30, fontSize: 19, padding: 9, borderRadius: '4px 0 0 4px' }}
+                                    styles={{rootPressed: {border: 'none', boxShadow: 'none'}}} 
+                                    iconProps={{ iconName: 'DoubleChevronDown' }}
+                                    onClick={() => {
+                                       nextErrorBlock();
+                                    }} />
+                              </TooltipHost>
+                                 <div style={{backgroundColor: errorButtonTextColor, height: 18, width: 1, position: 'absolute', right: -0.5, zIndex: 1 }}></div>
+                           </Stack>
+                           <TooltipHost content={'Skip to previous block of errors or issue'}>
+                              <IconButton disabled={!prevErrorBlockEnabled} className={prevErrorBlockEnabled ? handler.style.errorButton : handler.style.errorButtonDisabled}
+                                 style={{ color: errorButtonTextColor, height: 30, fontSize: 19, padding: 9,  borderRadius: '0 4px 4px 0'}} 
+                                 iconProps={{ iconName: 'DoubleChevronUp' }}
+                                 onClick={(event: any) => {
+                                    event?.stopPropagation();
+                                    prevErrorBlock();
+                                 }} />
+                           </TooltipHost>
+                        </Stack>}
+
+
+                        <Stack horizontal verticalAlign='center' style={{ paddingLeft: 12 }} >
+                           <Stack horizontal verticalAlign='center' style={{position: 'relative'}}>
+                              <TooltipHost content={nextWarningEnabled ? 'Navigate to next warning': ''}>
+                                 <DefaultButton disabled={!nextWarningEnabled} className={nextWarningEnabled ? handler.style.warningButton : handler.style.warningButtonDisabled}
+                                    text={`${warningsText} ${warnings.length === 1 ? "Warning" : "Warnings"}`}
+                                    onClick={() => {
+                                       nextWarning();
+                                    }}
+                                    style={{ color: warnings.length ? warningButtonTextColor : buttonNoneTextColor, height: 30, fontSize: 19, paddingLeft:  warnings.length ? 12 : 8, borderRadius: warnings.length ? '4px 0 0 4px': '4px' }} >
+                                    {!!warnings.length && <Icon style={{ fontSize: 19, paddingLeft: 4, paddingRight: 0 }} iconName='ChevronDown' />}
+                                 </DefaultButton>
+                              </TooltipHost>     
+                              {!!warnings.length && <div style={{backgroundColor: warningButtonTextColor, height: 18, width: 1, position: 'absolute', right: -0.5, zIndex: 1 }} />}
+                           </Stack>
+                           {!!warnings.length && <Stack>
+                              <TooltipHost content={'Navigate to previous warning'}>
+                                 <IconButton disabled={!prevWarningEnabled} className={prevWarningEnabled ? handler.style.warningButton : handler.style.warningButtonDisabled} style={{ height: 30, fontSize: 19, padding: 2, color: warningButtonTextColor, borderRadius: '0 4px 4px 0'}} iconProps={{ iconName: 'ChevronUp' }} onClick={(event: any) => {
+                                    event?.stopPropagation();
+                                    prevWarning();
+                                 }} />
+                              </TooltipHost>
+                           </Stack>}
+                        </Stack>
+
+                        {!!warnings.length && <Stack horizontal verticalAlign='center' style={{paddingLeft: 1}} >
+                           <Stack horizontal verticalAlign='center' style={{position: 'relative'}}>
+                              <TooltipHost content={'Skip to next block of warnings or issue'}>
+                                 <IconButton disabled={!nextWarningBlockEnabled} className={nextWarningBlockEnabled ? handler.style.warningButton : handler.style.warningButtonDisabled}
+                                    style={{ color: warningButtonTextColor, height: 30, fontSize: 19, padding: 9, borderRadius: '4px 0 0 4px' }}
+                                    styles={{rootPressed: {border: 'none', boxShadow: 'none'}}} 
+                                    iconProps={{ iconName: 'DoubleChevronDown' }}
+                                    onClick={() => {
+                                       nextWarningBlock();
+                                    }} />
+                              </TooltipHost>
+                                 <div style={{backgroundColor: warningButtonTextColor, height: 18, width: 1, position: 'absolute', right: -0.5, zIndex: 1 }}></div>
+                           </Stack>
+                           <TooltipHost content={'Skip to previous block of warnings or issue'}>
+                              <IconButton disabled={!prevWarningBlockEnabled} className={prevWarningBlockEnabled ? handler.style.warningButton : handler.style.warningButtonDisabled}
+                                 style={{ color: warningButtonTextColor, height: 30, fontSize: 19, padding: 9,  borderRadius: '0 4px 4px 0'}} 
+                                 iconProps={{ iconName: 'DoubleChevronUp' }}
+                                 onClick={(event: any) => {
+                                    event?.stopPropagation();
+                                    prevWarningBlock();
+                                 }} />
+                           </TooltipHost>
+                        </Stack>}
+
+                        {!!menuProps.items.find(i => !i.disabled) && <Stack style={{ paddingLeft: 18 }}><PrimaryButton disabled={!menuProps.items.find(i => !i.disabled)} text="View" menuProps={menuProps} style={{ fontFamily: "Horde Open Sans SemiBold", borderStyle: "hidden", padding: 15 }} /></Stack>}
+                        {!menuProps.items.find(i => !i.disabled) && <Stack style={{ paddingLeft: 18 }}><DefaultButton disabled={true} className={handler.style.warningButtonDisabled} text="View" style={{ color: "rgb(97, 110, 133)", padding: 15 }} /> </Stack>}
+
+                     </Stack>
+                     <Stack>
+                        {logSource.active && <LogProgressIndicator logSource={logSource} />}
+                     </Stack>
+                  </Stack>
+
+                  <Stack grow />
+                  <Stack horizontalAlign={"end"}>
+                     <Stack>
+                        <Stack verticalAlign="center" horizontal tokens={{ childrenGap: 24 }} styles={{ root: { paddingRight: 4, paddingTop: 4 } }}>
+                           <Stack horizontal tokens={{ childrenGap: 24 }}>
+                              <Stack horizontal tokens={{ childrenGap: 0 }}>
+                                 {!!searchState.curRequest && <Spinner style={{ paddingRight: 8 }} />}
+
+                                 <TextField
+                                    spellCheck={false}
+                                    autoComplete="off"
+                                    deferredValidationTime={1500}
+                                    validateOnLoad={false}
+                                    componentRef={searchBox}
+                                    value={searchState.text ? searchState.text : ''}
+                                    onChange={(ev, value) => { setSearchState({...searchState, text: value}); }}
+                                    onMouseDown={handleSearchBoxClick}
+                                    
+                                    onRenderSuffix={
+                                       !!searchState.text ? () => (
+                                          <Stack verticalAlign="center" horizontalAlign="center" horizontal tokens={{ childrenGap: 2 }} styles={{ root: { padding: 0 } }}>
+                                             <IconButton iconProps={{ iconName: 'Cancel'}} styles={{ root: { padding: 12, height: 20, width: 20 }, icon: { fontSize: 15 } }}
+                                             onClick={() => { doQuery('') }}
+                                             />
+                                             <span className={mergeStyles({ paddingLeft: 1, paddingRight: 7, paddingBottom: 1})}>
+                                                {searchState?.results ? searchState?.results?.length ? `${curSearchIdx + 1}/${searchState?.results?.length}` : '0/0' : undefined}
+                                             </span>               
+                                          </Stack>
+                                       ) : undefined
+                                    }
+
+                                    styles={{
+                                       root: { width: 280, fontSize: 12 }, 
+                                       fieldGroup: { borderWidth: 1 },
+                                       suffix: { padding: '0 2px', height: 30 }
+                                    }}
+                                    placeholder="Search"
+
+                                    onKeyDown={(ev) => {
+
+                                       if (ev.key === "Enter" && !searchState.curRequest && !inTimeout) {
+                                          if (searchState.text === searchState.search) {
+                                             searchDown(navigate);
+                                          } else {
+                                             doQuery(searchState.text ?? "");
+                                          }
+                                       }
+
+                                    }}
+
+                                    onGetErrorMessage={(newValue) => {
+                                       if (searchState.text !== searchState.search) {
+                                          doQuery(newValue);
+                                       }
+                                       return undefined;
+                                    }}
+
+                                 />
+                                 <Stack horizontal style={{ borderWidth: 1, borderStyle: "solid", borderColor: dashboard.darktheme ? "#3F3F3F" : "rgb(96, 94, 92)", height: 32, borderLeft: 0 }}>
+                                    <IconButton style={{ height: 30 }} iconProps={{ iconName: 'ChevronUp' }} onClick={(event: any) => {
+                                       searchUp(navigate);
+                                    }} />
+                                    <IconButton style={{ height: 30 }} iconProps={{ iconName: 'ChevronDown' }} onClick={(event: any) => {
+                                       searchDown(navigate);
+                                    }} />
+
+                                 </Stack>
+                              </Stack>
+
+                              <Dropdown
+                                 styles={{ root: { width: 92 } }}
+                                 options={[{ key: 'relative', text: 'Relative' }, { key: 'local', text: 'Local' }, { key: 'utc', text: 'UTC' }]}
+                                 defaultSelectedKey={tsFormat}
+                                 onChanged={(value) => {
+                                    setTSFormat(value.key as string);
+                                    listRef?.forceUpdate();
+                                 }}
+                              />
+
+                              <PrimaryButton
+                                 text="Download"
+                                 split
+                                 onClick={() => logSource.download(false)}
+                                 menuProps={downloadProps}
+                                 style={{ fontFamily: "Horde Open Sans SemiBold" }}
+                              />
+                           </Stack>
+                        </Stack>
+                     </Stack>
+                  </Stack>
+
+               </Stack>
+               <Stack>
+                  <Stack style={{ paddingTop: 12 }}>
+                     <Separator styles={{ root: { fontSize: 0, width: "100%", padding: 0, selectors: { '::before': { background: dashboard.darktheme ? '#313638' : '#D3D2D1' } } } }} />
+                  </Stack>
+
+                  <Stack style={{ paddingBottom: 12, paddingTop: 12, fontSize: 12, fontFamily: "Horde Open Sans Regular" }}>
+                     <Stack style={{ paddingLeft: 4 }} horizontal verticalAlign="center" tokens={{ childrenGap: 12 }}>
+                        <Text style={{ fontSize: 12, fontFamily: "Horde Open Sans Semibold" }}>Summary:</Text>
+                        <Markdown>{summaryText}</Markdown>
+                     </Stack>
+                  </Stack>
+                  {logTelemetry && !!logSource.agentTelemetry && <AgentTelemetrySparkline handler={logSource.agentTelemetry} />}
+                  {largeLogWarning && <Stack style={{ paddingBottom: 12, paddingTop: 12, margin: 'auto'}}>
+                     <MessageBar
+                        messageBarType={MessageBarType.warning} isMultiline={false} >
+                        <Markdown styles={{ root: { fontSize: 12, fontFamily: "Horde Open Sans Bold" } }}>{largeLogWarningText}</Markdown>
+                     </MessageBar>
+                  </Stack>}
+                  <Stack horizontalAlign="center" style={{ paddingBottom: 12 }}>
+                     <Separator styles={{ root: { fontSize: 0, width: "100%", padding: 0, selectors: { '::before': { background: dashboard.darktheme ? '#313638' : '#D3D2D1' } } } }} />
+                  </Stack>
+               </Stack>
+
+            </Stack>
+         </Stack>
+
+
+         <Stack style={{ backgroundColor: hordeTheme.horde.contentBackground, paddingLeft: "24px", paddingRight: "24px" }}>
+            <Stack tokens={{ childrenGap: 0 }}>
+               <FocusZone direction={FocusZoneDirection.vertical} isInnerZoneKeystroke={() => { return true; }} defaultActiveElement="#LogList" style={{ padding: 0, margin: 0 }} >
+                  <div className={handler.style.container} data-is-scrollable={true} style={{ height: `calc(100vh - ${heightAdjust}px)` }}
+                     onScroll={(ev) => {
+
+                        const element: any = ev.target;
+
+                        const scroll = Math.ceil(element.scrollHeight - element.scrollTop);
+
+                        if (scroll < (element.clientHeight + 1)) {
+                           handler.scroll = undefined;
+                           handler.trailing = true;
+                        } else {
+
+                           if (handler.scroll === undefined) {
+                              handler.scroll = scroll;
+                           } else if (handler.scroll !== undefined && (handler.scroll < scroll || scroll > (element.clientHeight + 1))) {
+                              handler.stopTrailing();
+                           }
+                        }
+
+
+                     }}>
+                     <Stack horizontal>
+                        {!dashboard.leftAlignLog && <div key={`windowsize_logview2_${windowSize.width}_${windowSize.height}`} style={{ width: vw / 2 - (1440 / 2) - 48, flexShrink: 0 }} />}
+                        <Stack styles={{ root: { "backgroundColor": hordeTheme.horde.contentBackground, paddingLeft: "0px", paddingRight: "0px" } }}>
+                           <SelectionZone selection={selection} selectionMode={SelectionMode.multiple}>
+                              <List key={`log_list_key_${logListKey}`} version={logSource.itemsVersion} id="LogList" ref={(list: List) => { listRef = list; }}
+                                 items={logSource.logItems}
+                                 // NOTE: getPageSpecification breaks initial scrollToIndex when query contains lineIndex!
+                                 getPageHeight={() => 10 * (handler.lineHeight)}
+                                 onShouldVirtualize={() => { return true; }}
+                                 onRenderCell={onRenderCell}
+                                 onPagesUpdated={() => {
+                                    if (handler.initialRender && listRef) {
+                                       handler.initialRender = false;
+                                       if (logSource?.startLine !== undefined) {
+
+                                          listRef?.scrollToIndex(logSource?.startLine - 10 < 0 ? 0 : logSource?.startLine - 10, () => handler.lineHeight, ScrollToMode.top);
+                                       } else if (handler.trailing) {
+                                          listRef.scrollToIndex(logSource.logData!.lineCount - 1, undefined, ScrollToMode.bottom);
+                                       }
+                                    }
+                                 }}
+
+                                 data-is-focusable={true} />
+                           </SelectionZone>
+                        </Stack>
+                     </Stack>
+                  </div>
+               </FocusZone>
+            </Stack>
+         </Stack>
+
+      </Stack>
+   </Stack>;
+
+});
+
+export const LogView: React.FC = () => {
+
+   const { logId } = useParams<{ logId: string }>();
+
+   useEffect(() => {
+      return () => {
+         selection.setItems([], true);
+         LogHandler.clear();
+      };
+   }, []);
+
+   const { hordeClasses } = getHordeStyling();
+
+
+   if (!logId) {
+      console.error("Bad log id to LogView");
+   }
+
+   return (
+      <Stack className={hordeClasses.horde}>
+         <TopNav />
+         <LogList logId={logId!} />
+      </Stack>
+   );
+};
diff --git a/Engine/Source/Programs/Horde/HordeDashboard/src/legacyPlugins/example/components/JobDetailPluginExample.tsx b/Engine/Source/Programs/Horde/HordeDashboard/src/legacyPlugins/example/components/JobDetailPluginExample.tsx
deleted file mode 100644
diff --git a/Engine/Source/Programs/Horde/HordeDashboard/src/legacyPlugins/example/index.tsx b/Engine/Source/Programs/Horde/HordeDashboard/src/legacyPlugins/example/index.tsx
deleted file mode 100644
diff --git a/Engine/Source/Programs/Horde/HordeDashboard/src/legacyPlugins/example/routes/ExampleBar.tsx b/Engine/Source/Programs/Horde/HordeDashboard/src/legacyPlugins/example/routes/ExampleBar.tsx
deleted file mode 100644
