[Horde][BuildHleath][Dashboard] D3 StepOutcomeTable variant for improved render performance.

diff --git a/Engine/Source/Programs/Horde/HordeDashboard/src/components/buildhealth/BuildhealthView.tsx b/Engine/Source/Programs/Horde/HordeDashboard/src/components/buildhealth/BuildhealthView.tsx
--- a/Engine/Source/Programs/Horde/HordeDashboard/src/components/buildhealth/BuildhealthView.tsx
+++ b/Engine/Source/Programs/Horde/HordeDashboard/src/components/buildhealth/BuildhealthView.tsx
@@ -1,6 +1,6 @@
 // Copyright Epic Games, Inc. All Rights Reserved.
 
-import { ComboBox, DirectionalHint, Dropdown, DropdownMenuItemType, IDropdownOption, Label, Slider, Stack, TooltipHost } from "@fluentui/react";
+import { ComboBox, DirectionalHint, Dropdown, DropdownMenuItemType, IDropdownOption, Label, Stack, TooltipHost } from "@fluentui/react";
 import { GetProjectResponse, GetStreamResponse } from "horde/backend/Api";
 import { observer } from "mobx-react-lite";
 import { useEffect, useMemo, useState } from "react";
@@ -58,6 +58,7 @@ const BasicListParameter: React.FC<{ handler: BuildHealthDataHandler, buildHealt
         options={doptions}
         selectedKey={selectedKey}
         dropdownWidth={200}
+        styles={{ dropdown: { width: 300 } }}
         calloutProps={{
             directionalHint: DirectionalHint.rightCenter
         }}
@@ -517,7 +518,7 @@ export const BuildHealthView: React.FC = observer(function ConstructBuildHealthV
                     <Stack styles={{ root: { width: "100%" } }}>
                         <Stack horizontal styles={{ root: { minHeight: '85vh' } }} >
                             <BuildHealthSidebar handler={handler} buildHealthOptions={buildHealthController} />
-                            <Stack id="parent-buildhealth-stepoutcome" grow style={{ overflowX: "auto", overflowY: "visible", minWidth: "1128px", position: "relative" }}>
+                            <Stack id="parent-buildhealth-stepoutcome" grow style={{ overflowX: "auto", overflowY: "visible", minWidth: "1128px", position: "relative", border: `${modeColors.content} solid 2px` }}>
                                 <StepOutcomeView filter={filter} />
                             </Stack>
                         </Stack>
diff --git a/Engine/Source/Programs/Horde/HordeDashboard/src/components/buildhealth/stepoutcome/StepOutcome.tsx b/Engine/Source/Programs/Horde/HordeDashboard/src/components/buildhealth/stepoutcome/StepOutcome.tsx
--- a/Engine/Source/Programs/Horde/HordeDashboard/src/components/buildhealth/stepoutcome/StepOutcome.tsx
+++ b/Engine/Source/Programs/Horde/HordeDashboard/src/components/buildhealth/stepoutcome/StepOutcome.tsx
@@ -1,229 +1,17 @@
 // Copyright Epic Games, Inc. All Rights Reserved.
 
-import { ConstrainMode, DetailsList, DetailsListLayoutMode, Dialog, DialogType, IColumn, IGroup, Label, mergeStyleSets, ScrollablePane, SelectionMode, Spinner, SpinnerSize, Stack, TooltipHost } from "@fluentui/react";
+import { Dialog, DialogType, Label, mergeStyleSets, PrimaryButton, ScrollablePane, Spinner, SpinnerSize, Stack } from "@fluentui/react";
 import { Text } from '@fluentui/react/lib/Text';
-import dashboard, { StatusColor } from "horde/backend/Dashboard";
-import { getHordeStyling } from 'horde/styles/Styles';
 import { observer } from "mobx-react-lite";
-import React, { useEffect, useRef, useState } from 'react';
+import React, { useEffect, useState } from 'react';
+import { StepOutcomeD3Table } from "./StepOutcomeD3Component";
 import { StepOutcomeDataHandler, StepOutcomeFilters } from "./StepOutcomeDataHandler";
-import { encodeStepNameFromStepNameHeader, StepNameHeader, StepOutcomeTable, StepOutcomeTableEntry } from "./StepOutcomeDataTypes";
+import { StepOutcomeTableEntry } from "./StepOutcomeDataTypes";
+import { sharedBorderStyle } from "./StepOutcomeSharedUIComponents";
 import { toDate } from "./StepOutcomeUtilities";
 
-// #region -- UI Helper Methods --
-
-const ROW_HEADER_NAME: string = "stepName" as const;
-
-/**
- * Produces Fluent UI compatible groups from the provided Step Outcome Table. Groups are currently derived based on Streams.
- * @param stepOutcomeTable The step outcome table to use in producing Fluent UI groups.
- * @returns The Fluent UI compatible groups.
- */
-function getFluentGroups(stepOutcomeTable: StepOutcomeTable): any[] {
-    let returnGroups: IGroup[] = [];
-    let startIdx = 0;
-
-    for (const [key, value] of stepOutcomeTable.tableStreamMetadata) {
-        let streamGroup: IGroup = {
-            key: key,
-            name: key,
-            count: value,
-            startIndex: startIdx,
-            isCollapsed: false
-        };
-        startIdx += value;
-        returnGroups.push(streamGroup);
-    }
-    return returnGroups;
-}
-
-
-/**
- * Produces Fluent UI compatible columns from the provided Step Outcome Table.
- * @param stepOutcomeTable The step outcome table to use in producing Fluent UI columns.
- * @returns The Fluent UI compatible columns.
- */
-function getFluentColumns(stepOutcomeTable: StepOutcomeTable): IColumn[] {
-    const columns: IColumn[] = [
-        {
-            key: ROW_HEADER_NAME,
-            name: '',
-            fieldName: ROW_HEADER_NAME,
-            minWidth: 300,
-            onRender: (item) => {
-                return (
-                    <TooltipHost content={
-                        <div>
-                            <Stack>
-                                <Label>Step:</Label>
-                                <Text style={{ paddingLeft: 8 }}>
-                                    {item.stepName.stepName}
-                                </Text>
-                            </Stack>
-                            <Stack>
-                                <Label>Stream:</Label>
-                                <Text style={{ paddingLeft: 8 }}>
-                                    {item.stepName.streamId}
-                                </Text>
-                            </Stack>
-                        </div>
-                    }>
-                        <div style={{ whiteSpace: 'normal', wordWrap: 'break-word', minHeight: 20, maxHeight: 20 }}>
-                            {` ${item.stepName.stepName}`}
-                        </div>
-                    </TooltipHost>
-                );
-            },
-        }
-    ];
-
-    // We want to inject dates into the DetailsList column headers in order to assist with data consumption.
-    // Keep accumulators for the corresponding date strings.
-    let prevDay: string | null = null;
-    let prevChangelist: string | null = null;
-
-    for (let i: number = 0; i < stepOutcomeTable.changeColHeaders.length; ++i) {
-
-        // If the date associated with the changelist header has changed since the last, we are on a new day. Set that as the colimn.
-        let changeHeader = stepOutcomeTable.changeColHeaders[i];
-        if (changeHeader.date !== undefined) {
-            const headerDate = new Date(changeHeader.date!);
-            const headerChangelist = changeHeader.change.toString();
-
-            // Extract the day portion as YYYY-MM-DD
-            const dayStr = `${headerDate.getUTCFullYear()}-${headerDate.getUTCMonth() + 1}-${headerDate.getUTCDate()}`;
-
-            if (prevDay !== dayStr) {
-                prevDay = dayStr;
-                prevChangelist = headerChangelist;
-                columns.push({
-                    key: `${prevDay}_${changeHeader.change}`,
-                    name: changeHeader.change.toString(),
-                    minWidth: 10,
-
-                    onRenderHeader: (props, defaultRender) => {
-                        if (!props || !defaultRender) return null;
-
-                        return (
-                            <div className={headerStylesHighlight.headerCell}>
-                                <TooltipHost>
-                                    {dayStr}
-                                </TooltipHost>
-                            </div>
-                        );
-                    }
-                });
-            }
-        }
-
-        // Push all columns that relate to changelists.
-        columns.push({
-            key: changeHeader.change.toString(),
-            name: changeHeader.change.toString(),
-            fieldName: changeHeader.change.toString(),
-            minWidth: 10,
-            onRenderHeader: (props, defaultRender) => {
-                if (!props) {
-                    return null;
-                }
-                let innerChange = stepOutcomeTable.changeColHeaders[props?.columnIndex];
-                return props && defaultRender ? (
-                    <TooltipHost content={
-                        <Stack>
-                            <Label>Change:</Label>
-                            <Text style={{ paddingLeft: 8 }}>{changeHeader.change}</Text>
-                            <Label>Date (UTC):</Label>
-                            <Text style={{ paddingLeft: 8 }} >{new Date(changeHeader.date!).toISOString()}</Text>
-                            <Label>Date (Local):</Label>
-                            <Text style={{ paddingLeft: 8 }} >{new Date(changeHeader.date!).toString()}</Text>
-                        </Stack>
-                    }
-                    >
-                        {defaultRender(props)}
-                    </TooltipHost>
-                ) : null
-            },
-            onRender: (item: any) => {
-                const entry = item[changeHeader.change.toString()] as StepOutcomeTableEntry | null;
-                let bgColor = 'transparent';
-
-                if (entry) {
-                    let retrievedCellStype = getCellStyle(entry);
-                    bgColor = retrievedCellStype.bgColor;
-                }
-
-                return (
-                    <TooltipHost content={
-                        <Stack>
-                            <Label>Change:</Label>
-                            <Text className={stepOutcomeTooltipClasses.bottomBorder} style={{ paddingLeft: 8 }}>{changeHeader.change}</Text>
-                            {entry && (
-                                <>
-                                    <Label>Duration:</Label>
-                                    <Text style={{ paddingLeft: 8 }}>{entry.getDurationString()}</Text>
-                                    <Label>State:</Label>
-                                    <Text className={stepOutcomeTooltipClasses.bottomBorder} style={{ paddingLeft: 8 }}>{entry.stepResponse.state}</Text>
-                                    <Label>Stream:</Label>
-                                    <Text style={{ paddingLeft: 8 }}>{entry.streamId}</Text>
-                                    <Label>Job:</Label>
-                                    <Text style={{ paddingLeft: 8 }}>{entry.jobName}</Text>
-                                    <Label>Step:</Label>
-                                    <Text style={{ paddingLeft: 8 }}>{entry.stepResponse.name}</Text>
-                                </>)}
-                        </Stack>
-                    }>
-                        <div
-                            style={{
-                                backgroundColor: bgColor,
-                                minHeight: 20, maxHeight: 20,
-                                minWidth: 20, maxWidth: 20,
-                            }}
-                        >
-                            &nbsp;
-                        </div>
-                    </ TooltipHost>
-
-                );
-            }
-        });
-    }
-
-    return columns;
-}
-
-/**
- * Produces Fluent UI compatible rows from the provided Step Outcome Table.
- * @param stepOutcomeTable The step outcome table to use in producing Fluent UI rows.
- * @returns The Fluent UI compatible rows.
- */
-function getFluentRows(stepOutcomeTable: StepOutcomeTable): StepOutcomeFluentRow[] {
-    let items: StepOutcomeFluentRow[] = [];
-
-    for (let i: number = 0; i < stepOutcomeTable.stepNameRowHeaders.length; ++i) {
-        let rows: StepOutcomeTableEntry[] = [...handler.StepOutcomeTableData.getStepOutputRow(encodeStepNameFromStepNameHeader(stepOutcomeTable.stepNameRowHeaders[i]))];
-        let fluentRowEntry: StepOutcomeFluentRow = { [ROW_HEADER_NAME]: stepOutcomeTable.stepNameRowHeaders[i] };
-        for (let y: number = 0; y < rows.length; ++y) {
-            let tableEntry: StepOutcomeTableEntry = rows[y];
-            if (tableEntry.change === undefined) {
-                continue;
-            }
-            fluentRowEntry[tableEntry.change!.toString()] = tableEntry;
-        }
-        items.push(fluentRowEntry);
-    }
-
-    return items;
-}
-
-// #endregion -- UI Helper Methods --
-
 // #region -- View Components --
 
-type StepOutcomeFluentRow = {
-    [ROW_HEADER_NAME]: StepNameHeader;
-    [change: string]: StepOutcomeTableEntry | StepNameHeader | string | undefined;
-};
-
 type StepOutcomeTableProps = {
     filters: StepOutcomeFilters;
     handler: StepOutcomeDataHandler;
@@ -237,135 +25,8 @@ interface StepOutcomeViewProps {
 
 // #region -- Styles  --
 
-/**
- * Defines the style for column headers.
- */
-
-const sharedBorderStyle = {
-    bottomBorder: {
-        borderBottom: `1px solid rgba(80, 80, 80, 1)`
-    }
-};
-
-const stepOutcomeTooltipClasses = mergeStyleSets(sharedBorderStyle);
-
 const stepOutcomeModalClasses = mergeStyleSets(sharedBorderStyle);
 
-const headerStylesHighlight = mergeStyleSets({
-    headerCell: {
-        fontWeight: '700 !important',
-        fontSize: 'larger !important'
-    },
-});
-
-const stepOutcomeTableClasses = mergeStyleSets({
-    stepOutcomeTableHeaders: {
-        selectors: {
-            '.ms-DetailsHeader-cellName': {
-                fontWeight: 200
-            },
-            '.ms-DetailsList-headerWrapper': {
-                position: 'sticky'
-            },
-            '.ms-DetailsHeader-cellTitle': {
-                writingMode: 'vertical-lr',
-                transform: 'rotate(180deg)',
-                whiteSpace: 'nowrap',
-                textAlign: 'center',
-                paddingBottom: '10px',
-                width: '36px !important',
-                userSelect: 'text',
-            },
-            '.ms-DetailsHeader': {
-                minHeight: '100px',
-                lineHeight: 'normal',
-                paddingTop: '0 !important',
-            },
-            '.ms-DetailsHeader-cell': {
-                verticalAlign: 'bottom',
-                minHeight: '80px',
-                borderLeft: '1px solid rgba(41, 41, 41, 0.35)', borderRight: '1px solid  rgba(41, 41, 41, 0.35)'
-            },
-            '.ms-DetailsHeader-cell span': {
-                overflow: 'visible',
-            },
-            '.ms-List-cell': {
-                minHeight: 40,
-                height: 40,
-            },
-            '.ms-SelectionZone': {
-                display: 'none',
-            },
-        },
-    },
-    removeCellContentPadding: {
-        selectors: {
-            '.ms-DetailsRow-cell': {
-                paddingLeft: '0 !important',
-            }
-        }
-    },
-    hideGroupCollapseButton: {
-        selectors: {
-            'button.ms-GroupHeader-expand': {
-                display: 'none',
-            },
-            '.ms-GroupHeader': {
-                height: '35px'
-            }
-        },
-    },
-    hideGroupHeaderTitle: {
-        selectors: {
-            '.ms-GroupHeader': {
-                height: '35px'
-            },
-            '.ms-GroupHeader-title': {
-                display: 'none',
-            },
-            'button.ms-GroupHeader-expand': {
-                display: 'none',
-            },
-            '.ms-GroupSpacer': {
-                width: '0 !important',
-            },
-            '.ms-DetailsRow-cell': {
-                padding: '10px 4px',
-                // Cleanup work for proper squares
-                // maxHeight: '35px',
-                // width: '35px !important',
-                // maxWdith: '35px',
-                borderLeft: '1px solid  rgba(41, 41, 41, 0.35)', borderRight: '1px solid  rgba(41, 41, 41, 0.35)'
-            },
-
-        },
-    },
-    disableGroupCollapseButton: {
-        selectors: {
-            '.ms-GroupHeader-dropIcon': {
-                display: 'none',
-                pointerEvents: 'none',
-                opacity: 0.5,
-            },
-        },
-    },
-    stickyHeaderOverride: {
-        selectors: {
-            '.ms-DetailsHeader-cellTitle': {
-                writingMode: 'vertical-rl',
-                whiteSpace: 'nowrap',
-                textAlign: 'center',
-                padding: '4px 2px',
-                paddingBottom: '10px',
-                width: '36px !important',
-            },
-            '.ms-DetailsHeader': {
-                minHeight: '100px',
-            },
-        },
-    },
-});
-
 // #endregion -- Styles  --
 
 // #region -- Modal Popup View --
@@ -478,420 +139,23 @@ const StepOutcomeFocusModal: React.FC<StepOutcomeFocusModalProps> = ({ item, onD
 
 // #endregion -- Modal Popup View --
 
-/**
- * Gets the cell style given the provided StepOutcomeTableEntry.
- * @param entry The entry to base the style off of.
- * @returns The matching style given the state of the entry. 
- */
-function getCellStyle(entry: StepOutcomeTableEntry | null): { bgColor: string } {
-    const scolors = dashboard.getStatusColors();
-    let resultColor: string = 'transparent';
-
-    if (!entry) {
-        return { bgColor: resultColor };
-    }
-
-    const WarningStr: string = "Warnings";
-    const SuccessStr: string = "Success";
-    const SkippedStr: string = "Skipped";
-    const FailureStr: string = "Failure";
-
-    const { stepResponse } = entry;
-    const { outcome, state } = stepResponse;
-    const colors: Record<string, string> = {
-        "Success": scolors.get(StatusColor.Success)!,
-        "Failure": scolors.get(StatusColor.Failure)!,
-        "Warnings": scolors.get(StatusColor.Warnings)!,
-        "Skipped": scolors.get(StatusColor.Skipped)!,
-        "Running": scolors.get(StatusColor.Running)!,
-        "Waiting": scolors.get(StatusColor.Waiting)!,
-        "Ready": scolors.get(StatusColor.Ready)!,
-        "Unspecified": scolors.get(StatusColor.Skipped)!,
-    };
-
-    // There are some interesting mappings going on internally with respect to state & outcome. Outcome & State have overlapping "result view" space, and as such we need to 
-    // narrowly select the edge cases (e.g. outcome=success && state=completed uniquely signifies a completed run with no errors, otherwise outcome=success it could mean ready, waiting, or running).
-    // Then we can fall back to the state as the truth.
-    if (outcome === FailureStr) {
-        resultColor = state === SkippedStr ? colors[SkippedStr] : colors[FailureStr];
-    }
-    // A step can be emitting warnings whilst still running, so we narrowly select this colouring.
-    else if (outcome === WarningStr) {
-        resultColor = colors[WarningStr];
-    }
-    // A step can have a completed outcome, but merely be waiting to run, running, or ready.
-    else if (outcome === SuccessStr && state === "Completed") {
-        resultColor = colors[SuccessStr];
-    }
-    // The edge cases have been satisfied, yield to the state to define result view color.
-    else {
-        resultColor = colors[state];
-    }
-
-    return { bgColor: resultColor };
-}
-
-/**
- * React Component that represents the provided stepOutcomeTable. This is a standard HTML table.
- * @param stepOutcomeTable The data to base the table off of.
- * @returns React Component.
- */
-function GenerateClassicTable(stepOutcomeTable: StepOutcomeTable): JSX.Element {
-    const { modeColors } = getHordeStyling();
-    return (
-        <div style={{ overflow: 'auto', maxHeight: '600px', maxWidth: '100%' }}>
-            <table
-                style={{
-                    borderCollapse: 'separate',
-                    borderSpacing: '2px',
-                    width: '100%',
-                    tableLayout: 'fixed', // Fix column widths strictly
-                }}
-            >
-                <thead>
-                    <tr>
-                        <th
-                            style={{
-                                position: 'sticky',
-                                top: 0,
-                                left: 0,
-                                zIndex: 3,
-                                backgroundColor: modeColors.background,
-                                maxWidth: 200,
-                                minWidth: 180,
-                                width: 180,
-                                whiteSpace: 'nowrap',
-                                overflow: 'hidden',
-                                textOverflow: 'ellipsis',
-                                height: '100px',
-                            }}
-                        >
-                            Step
-                        </th>
-                        {stepOutcomeTable.changeColHeaders.slice(0, 20).map((col, idx) => (
-                            <th
-                                key={idx}
-                                style={{
-                                    position: 'sticky',
-                                    top: 0,
-                                    zIndex: 2,
-                                    backgroundColor: modeColors.background,
-                                    writingMode: 'vertical-rl',
-                                    textAlign: 'left',
-                                    whiteSpace: 'nowrap',
-                                    height: '40px',
-                                    width: '40px',
-                                }}
-                            >
-                                <div
-                                    style={{
-                                        transform: 'rotate(180deg)',
-                                        display: 'inline-block',
-                                        padding: '8px 0',
-                                    }}
-                                >
-                                    {col.change}
-                                </div>
-                            </th>
-                        ))}
-                    </tr>
-                </thead>
-                <tbody>
-                    {stepOutcomeTable.stepNameRowHeaders.slice(0, 20).map((key, rIdx) => (
-                        <tr key={rIdx} style={{ height: '40px' }}>
-                            <td
-                                title={key.stepName}
-                                style={{
-                                    position: 'sticky',
-                                    left: 0,
-                                    zIndex: 1,
-                                    backgroundColor: modeColors.background,
-                                    whiteSpace: 'nowrap',
-                                    overflow: 'hidden',
-                                    textOverflow: 'ellipsis',
-                                    maxWidth: 200,
-                                    minWidth: 180,
-                                    width: 180,
-                                    height: '40px',
-                                }}
-                            >
-                                {`${key.stepName}(${key.streamId})`}
-                            </td>
-                            {[...stepOutcomeTable.getStepOutputRow(encodeStepNameFromStepNameHeader(key))]
-                                .slice(0, 20)
-                                .map((entry, cIdx) => {
-                                    const { bgColor } = getCellStyle(entry);
-                                    return (
-                                        <td
-                                            key={cIdx}
-                                            style={{
-                                                backgroundColor: bgColor,
-                                                padding: 0,
-                                                margin: 0,
-                                                width: '40px',
-                                                height: '40px',
-                                                textAlign: 'center',
-                                            }}
-                                        >
-                                            <div
-                                                style={{
-                                                    display: 'block',
-                                                    width: '100%',
-                                                    height: '100%',
-                                                    textDecoration: 'none',
-                                                }}
-                                            >
-                                                &nbsp;
-                                            </div>
-                                        </td>
-                                    );
-                                })}
-                        </tr>
-                    ))}
-                </tbody>
-            </table>
-        </div>
-
-    );
-}
-
-/**
- * React Component that provides refresh spinner on incremental refresh. 
- * Note: This is a separate component to prevent rerender of primary table @see GenerateFluentUITable upon refresh spinner.
- * @returns React Component.
- */
-const IncrementalRefreshIndicatorPanel: React.FC<{ handler: StepOutcomeDataHandler }> = observer(({ handler }) => (
-    <div
-        style={{
-            gridRow: '1 / span 2',
-            gridColumn: 3,
-            overflowY: 'auto',
-            width: 25,
-        }}
-    >
-        {handler.isInIncrementalDataRefresh && <Spinner size={SpinnerSize.small} />}
-    </div>
-));
-
-/**
- * React Component that represents the provided stepOutcomeTable. This is a Fluent UI table.
- * @param stepOutcomeTable The data to base the table off of.
- * @returns React Component.
- */
-function GenerateFluentUITable(stepOutcomeTable: StepOutcomeTable, onCellClick: any): JSX.Element {
-    /**
-     * This is a fairly "virtual" (aka hacky) table in the following respects:
-     * - We need to have a sticky header, and first column
-     * - Sticky component works well for header, but not at all for first column
-     * - As a result, we have opted to take the following strategy:
-     *      - Create a dual "panel" view:
-     *          - Left panel is the 1st column; the Step Names
-     *          - Right panel - Top Stack is the Headers for all the changes
-     *          - Right panel - Bottom Stack is the cell data for all the steps
-     *      - Synchronize the Left panel vertical scroll to the Right Panel - Bottom Stack vertical scroll
-     *      - Synchronize the Rigth Panel Top Stack horizontal scroll to the Right Panel - Bottom Stack horizontal scroll
-     * 
-     * Whilst it's a tad hacky, it's predictable and easy to reason.
-     */
-
-    const leftPaneRef = useRef<HTMLDivElement>(null);
-    const rightTopPaneRef = useRef<HTMLDivElement>(null);
-    const rightBottomPaneRef = useRef<HTMLDivElement>(null);
-
-    const isSyncingVertical = useRef(false);
-    const isSyncingHorizontal = useRef(false);
-
-    const allColumns = getFluentColumns(stepOutcomeTable);
-    const items = getFluentRows(stepOutcomeTable);
-    const groups = getFluentGroups(stepOutcomeTable);
-
-    const firstColumn = allColumns.slice(0, 1);
-    const restColumns = allColumns.slice(1);
-    const [hasScrollbar, setHasScrollbar] = useState(false);
-
-    useEffect(() => {
-        const checkScroll = () => {
-            const el = rightBottomPaneRef.current;
-            if (el) setHasScrollbar(el.scrollWidth > el.clientWidth);
-        };
-
-        checkScroll();
-        const resizeObserver = new ResizeObserver(checkScroll);
-        rightBottomPaneRef.current && resizeObserver.observe(rightBottomPaneRef.current);
-        return () => resizeObserver.disconnect();
-    }, []);
-
-    // Stash data structure to not create on all scroll events.
-    const targetMap = {
-        vertical: {
-            left: rightBottomPaneRef,
-            rightBottom: leftPaneRef,
-        },
-        horizontal: {
-            rightTop: rightBottomPaneRef,
-            rightBottom: rightTopPaneRef,
-        },
-    };
-
-    const VERTICAL_SCROLL_AFFORDANCE = 8;
-
-    const handleScroll = (direction: 'vertical' | 'horizontal', source: 'left' | 'rightTop' | 'rightBottom') => {
-        const isVertical = direction === 'vertical';
-        const syncRef = isVertical ? isSyncingVertical : isSyncingHorizontal;
-        const targetRef = targetMap[direction][source];
-
-        return (e: React.UIEvent<HTMLDivElement>) => {
-            if (syncRef.current) return;
-            syncRef.current = true;
-
-            const el = e.currentTarget;
-            let scrollValue = isVertical ? el.scrollTop : el.scrollLeft;
-
-            // Clamp for rightTop horizontal due to how we need to overlay the top DetailList as a "sticky row header".
-            if (!isVertical && source === 'rightTop') {
-                const maxAllowed = el.scrollWidth - el.clientWidth - VERTICAL_SCROLL_AFFORDANCE;
-                scrollValue = Math.min(scrollValue, maxAllowed);
-                el.scrollLeft = scrollValue; // keep UI in sync immediately
-            }
-
-            if (targetRef?.current) {
-                if (isVertical) {
-                    targetRef.current.scrollTop = scrollValue;
-                } else {
-                    targetRef.current.scrollLeft = scrollValue;
-                }
-            }
-
-            requestAnimationFrame(() => {
-                syncRef.current = false;
-            });
-        };
-    };
-
-    const gridStyle = {
-        display: 'grid',
-        gridTemplateRows: 'auto 1fr',
-        gridTemplateColumns: 'auto 1fr',
-        maxHeight: '85vh',
-        overflow: 'hidden',
-    };
-
-    return (
-        <div style={gridStyle}>
-            {/* Left Panel - Left Column - Step Name - (sticky) */}
-            <div
-                ref={leftPaneRef}
-                onScroll={handleScroll('vertical', 'left')}
-                className="no-scrollbar"
-                style={{
-                    maxHeight: '100%',
-                    overflowY: 'auto',
-                    position: 'relative',
-                    gridRow: 2,
-                    gridColumn: 1,
-                    overflowX: 'scroll',
-                    paddingBottom: hasScrollbar ? '17px' : '0',
-                    zIndex: 1,
-                    scrollbarWidth: 'none',
-                    msOverflowStyle: 'none',
-                }}
-            >
-                <DetailsList
-                    items={items}
-                    columns={firstColumn}
-                    groups={groups}
-                    selectionMode={SelectionMode.none}
-                    layoutMode={DetailsListLayoutMode.fixedColumns}
-                    constrainMode={ConstrainMode.unconstrained}
-                    isHeaderVisible={false}
-                    className={stepOutcomeTableClasses.hideGroupCollapseButton}
-                />
-            </div>
-
-            {/* Right Panel - Top Stack - Change Header Row (sticky) */}
-            <div
-                ref={rightTopPaneRef}
-                onScroll={handleScroll('horizontal', 'rightTop')}
-                className="no-scrollbar"
-                style={{
-                    gridRow: 1,
-                    gridColumn: 2,
-                    overflowX: 'auto',
-                    minHeight: '80px',
-                    maxHeight: '400px',
-                    position: 'relative',
-                    zIndex: 2,
-                }}
-            >
-                <DetailsList
-                    items={[]} // header only
-                    columns={restColumns}
-                    isHeaderVisible={true}
-                    selectionMode={SelectionMode.none}
-                    className={stepOutcomeTableClasses.stepOutcomeTableHeaders}
-                    layoutMode={DetailsListLayoutMode.fixedColumns}
-                    constrainMode={ConstrainMode.unconstrained}
-                />
-            </div>
-
-            {/* Right Panel - Bottom Stack - Step Data Cells (scroll driver) */}
-            <div
-                ref={rightBottomPaneRef}
-                onScroll={(e) => {
-                    handleScroll('vertical', 'rightBottom')(e);
-                    handleScroll('horizontal', 'rightBottom')(e);
-                }}
-                style={{
-                    maxHeight: '100%',
-                    overflowY: 'auto',
-                    position: 'relative',
-                    gridRow: 2,
-                    gridColumn: 2,
-                    overflow: 'auto',
-                }}
-            >
-                <DetailsList
-                    items={items}
-                    columns={restColumns}
-                    groups={groups}
-                    selectionMode={SelectionMode.none}
-                    layoutMode={DetailsListLayoutMode.fixedColumns}
-                    constrainMode={ConstrainMode.unconstrained}
-                    isHeaderVisible={false}
-                    className={stepOutcomeTableClasses.hideGroupHeaderTitle}
-                    onRenderField={(fieldProps, defaultRender) => {
-
-                        if (!fieldProps) {
-                            return null;
-                        }
-                        const { item, column, itemIndex } = fieldProps;
-                        const cellValue = column.fieldName ? item[column.fieldName] : undefined;
-                        const stepData: StepOutcomeTableEntry = cellValue as StepOutcomeTableEntry;
-                        const id = `cell-${itemIndex}-${column.key}`;
-
-                        return (
-                            <div id={id} key={id} style={{ paddingLeft: 0, cursor: 'pointer' }} onClick={() => onCellClick?.(stepData)}>
-                                {defaultRender?.(fieldProps)}
-                            </div>
-                        );
-                    }}
-                />
-            </div>
-            <IncrementalRefreshIndicatorPanel handler={handler} />
-        </div>
-    );
-};
-
 /**
  * React Component that acts as a observable binder for data handler updates, and allows for easy swapping of other table types.
  */
 const StepOutcomeTableComponent: React.FC<StepOutcomeTableProps> = observer(({ handler, onCellSelected }) => {
     handler.updated;
-    const table = GenerateFluentUITable(handler.StepOutcomeTableData, onCellSelected);
 
     return (
         <Stack>
-            {handler.isInFullDataRefresh ? <Spinner size={SpinnerSize.large} /> : table}
+            {handler.isInFullDataRefresh ? (
+                <Spinner size={SpinnerSize.large} />
+            ) :
+                <StepOutcomeD3Table
+                    stepOutcomeTable={handler.StepOutcomeTableData}
+                    handler={handler}
+                    onCellSelected={onCellSelected}
+                />
+            }
         </Stack>
     );
 });
diff --git a/Engine/Source/Programs/Horde/HordeDashboard/src/components/buildhealth/stepoutcome/StepOutcomeD3Component.tsx b/Engine/Source/Programs/Horde/HordeDashboard/src/components/buildhealth/stepoutcome/StepOutcomeD3Component.tsx
new file mode 100644
--- /dev/null
+++ b/Engine/Source/Programs/Horde/HordeDashboard/src/components/buildhealth/stepoutcome/StepOutcomeD3Component.tsx
@@ -0,0 +1,460 @@
+// Copyright Epic Games, Inc. All Rights Reserved.
+
+import { mergeStyleSets } from "@fluentui/react";
+import * as d3 from "d3";
+import { getHordeStyling, preloadFonts } from "horde/styles/Styles";
+import { useEffect, useRef } from "react";
+import { StepOutcomeDataHandler } from "./StepOutcomeDataHandler";
+import { ChangeHeader, encodeStepNameFromStepNameHeader, INVALID_CHANGELIST, StepNameHeader, StepOutcomeTable, StepOutcomeTableEntry } from "./StepOutcomeDataTypes";
+import { getCellStyle, IncrementalRefreshIndicatorPanel, sharedBorderStyle } from "./StepOutcomeSharedUIComponents";
+
+// #region -- Styles --
+
+const baseTooltipLine = {
+    paddingLeft: 8,
+    margin: 0.2,
+};
+
+const tooltipClasses = mergeStyleSets({
+    tooltipLine: baseTooltipLine,
+    tooltipLineWithBorder: {
+        ...baseTooltipLine,
+        ...sharedBorderStyle.bottomBorder,
+    }
+});
+
+// #endregion -- Styles --
+
+/**
+ * Function to create a D3 format table for the step outcome table.
+ * @param stepOutcomeTable The data to base the table off of.
+ * @param handler The data handler used to query step outcome data.
+ * @param onCellSelected Cell selected callback.
+ * @returns React Component.
+ */
+function GenerateD3Table(stepOutcomeTable: StepOutcomeTable, handler: StepOutcomeDataHandler, onCellClick: any,): JSX.Element {
+    // #region -- Tooltip Helpers --
+
+    const TOOLTIP_SPACER: number = 4;
+
+    let attachDynamicPositionToolTip = (event: MouseEvent, tooltip: d3.Selection<HTMLDivElement, unknown, HTMLElement, any>, tooltipContent: HTMLDivElement) => {
+        const _d3 = d3 as any;
+
+        tooltip.html("");
+        const toolTipNode = tooltip.node();
+        toolTipNode?.appendChild(tooltipContent);
+
+        const tooltipWidth = toolTipNode?.offsetWidth ?? 0;
+        const tooltipHeight = toolTipNode?.offsetHeight ?? 0;
+
+        // Get absolute  mouse position.
+        const x = event.pageX;
+        const y = event.pageY;
+
+        let leftPos = x + TOOLTIP_SPACER;
+        let topPos = y - tooltipHeight / 2;
+
+        // If tooltip goes off right edge, flip to the left.
+        if (x + tooltipWidth + TOOLTIP_SPACER > window.innerWidth) {
+            leftPos = x - tooltipWidth - TOOLTIP_SPACER;
+        }
+
+        // Prevent it from going off the top/bottom.
+        if (topPos < 0) topPos = 0;
+        if (topPos + tooltipHeight > window.innerHeight) {
+            topPos = window.innerHeight - tooltipHeight;
+        }
+
+        tooltip
+            .style("opacity", 1)
+            .style("left", `${leftPos}px`)
+            .style("top", `${topPos}px`);
+    };
+
+    let attachTargettedToolTip = (element: HTMLElement, tooltip: d3.Selection<HTMLDivElement, unknown, HTMLElement, any>, tooltipContent: HTMLDivElement) => {
+        tooltip.html("");
+        const toolTipNode = tooltip.node();
+        toolTipNode?.appendChild(tooltipContent);
+
+        const tooltipWidth = toolTipNode?.offsetWidth ?? 0;
+        const tooltipHeight = toolTipNode?.offsetHeight ?? 0;
+
+        const cell = element as HTMLElement;
+        const rect = cell.getBoundingClientRect();
+
+        let leftPos = rect.right + TOOLTIP_SPACER; // default: to the right.
+
+        // if tooltip goes outside window, position to the left of cell.
+        if (rect.right + tooltipWidth + TOOLTIP_SPACER > window.innerWidth) {
+            leftPos = rect.left - tooltipWidth - TOOLTIP_SPACER;
+        }
+
+        tooltip
+            .style("opacity", 1)
+            .style("left", `${leftPos}px`)
+            .style("top", `${rect.bottom - tooltipHeight}px`);
+    };
+
+    // #endregion -- Tooltip Helpers --
+
+    const containerRef = useRef<HTMLDivElement>(null);
+    const { modeColors, _hordeClasses, _detailClasses } = getHordeStyling();
+
+    const scaffold = useRef<{
+        wrapper: d3.Selection<HTMLDivElement, unknown, null, undefined>;
+        table: d3.Selection<HTMLTableElement, unknown, null, undefined>;
+        thead: d3.Selection<HTMLTableSectionElement, unknown, null, undefined>;
+        theaderrow: d3.Selection<HTMLTableRowElement, unknown, null, undefined>;
+        tbody: d3.Selection<HTMLTableSectionElement, unknown, null, undefined>;
+        tooltip: d3.Selection<HTMLDivElement, unknown, HTMLElement, any>;
+    } | null>(null);
+
+    // Single execution useEffect to create static part of table.
+    useEffect(() => {
+        if (!containerRef.current || scaffold.current) return;
+
+        // Create our base container.
+        const wrapper = d3.select(containerRef.current)
+            .append("div")
+            .attr("id", "step-outcome-table-container");
+
+        // Create the table element.
+        const table = wrapper.append("table")
+            .style("border-collapse", "separate")
+            .style("border-spacing", "0px")
+            .style("width", "fit-content")
+            .style("table-layout", "fixed")
+            .style("font-family", preloadFonts[0])
+            .style("font-size", "12px")
+            .style("color", modeColors.text);
+
+        // Add header && tbody.
+        const thead = table.append("thead");
+        const tbody = table.append("tbody");
+
+        // Add tool tip div; will be used throughout.
+        const tooltip = d3.select("body")
+            .append("div")
+            .attr("class", "d3-tooltip")
+            .style("position", "absolute")
+            .style("pointer-events", "none")
+            .style("padding", "6px 12px")
+            .style("background", "rgba(0,0,0,0.95)")
+            .style("color", "#fff")
+            .style("z-index", "5")
+            .style("min-width", "250px")
+            .style("font-size", "12px")
+            .style("border-radius", "2px")
+            .style("line-height", "1.4")
+            .style("opacity", 0);
+
+        const theaderrow = thead.append("tr").style("height", "100px");
+
+        // First column header; (0,0) - we just want to hide this; sits on top of group; guarantee the width of groups TH.
+        theaderrow.append("th")
+            .style("position", "sticky")
+            .style("z-index", "6")
+            .style("top", "0")
+            .style("left", "0")
+            .style("background-color", modeColors.background)
+            .style("width", "30px")
+            .style("min-width", "30px")
+            .style("max-width", "30px");
+
+        // Second Column Header; (0, 1) - this sits on top of step name; guarantee the width of step name TH.
+        theaderrow.append("th")
+            .style("position", "sticky")
+            .style("z-index", "6")
+            .style("top", "0")
+            .style("left", "30px")
+            .style("background-color", modeColors.background)
+            .style("width", "320px")
+            .style("min-width", "320px")
+            .style("max-width", "320px");
+
+        scaffold.current = { wrapper, table, thead, theaderrow, tbody, tooltip };
+    }, []);
+
+    // Dependency array use effect for stepOutcomeTable data.
+    useEffect(() => {
+        if (!containerRef.current || !scaffold.current) return;
+
+        // Add all columns to the header row; these will be the changelists.
+        scaffold.current.theaderrow.selectAll("th.col")
+            .data([...stepOutcomeTable.getDateAnchoredChanges()] as ChangeHeader[])
+            .join("th")
+            .attr("class", "col")
+            .style("position", "sticky")
+            .style("top", "0")
+            .style("z-index", "3")
+            .style("background-color", modeColors.content)
+            .style("width", "30px")
+            .style("font-size", d => d.change === INVALID_CHANGELIST ? "large" : "small")
+            .style("box-sizing", "border-box")
+            .style("writing-mode", "vertical-rl")
+            .style("transform", "rotate(180deg)")
+            .style("padding-bottom", "5px")
+            .style("text-align", "right")
+            .style("border", `1px solid ${modeColors.background}`)
+            .text(d => d.change === INVALID_CHANGELIST ? (d.date ?? "") : d.change)
+            .on("mouseover", function (_event, d) {
+                const entry = d as unknown as ChangeHeader;
+
+                if (entry && entry.change !== INVALID_CHANGELIST) {
+                    const tooltipContent = document.createElement("div");
+                    tooltipContent.style.paddingLeft = "4px";
+
+                    const changeHeader = { change: entry.change ?? "-", date: entry.date ?? "-" };
+
+                    const changeDiv = document.createElement("div");
+                    changeDiv.innerHTML = `<strong>Change:</strong><p class="${tooltipClasses.tooltipLine}">${changeHeader.change}</p>`;
+                    tooltipContent.appendChild(changeDiv);
+
+                    const dateUTCDiv = document.createElement("div");
+                    dateUTCDiv.innerHTML = `<strong>Date (UTC):</strong><p class="${tooltipClasses.tooltipLine}">${new Date(changeHeader.date!).toISOString()}</p>`;
+                    tooltipContent.appendChild(dateUTCDiv);
+
+                    const dateLocalDiv = document.createElement("div");
+                    dateLocalDiv.innerHTML = `<strong>Date (Local):</strong><p class="${tooltipClasses.tooltipLine}">${new Date(changeHeader.date!).toString()}</p>`;
+                    tooltipContent.appendChild(dateLocalDiv);
+
+                    attachTargettedToolTip(this as HTMLElement, scaffold.current!.tooltip, tooltipContent);
+                }
+            })
+            .on("mouseout", function () {
+                scaffold.current!.tooltip.style("opacity", 0);
+            });
+
+        const rows = scaffold.current.tbody.selectAll("tr")
+            .data(stepOutcomeTable.stepNameRowHeaders as StepNameHeader[])
+            .join("tr")
+            .style("height", "30px")
+            .style("padding-top", "4px")
+            .style("padding-bottom", "4px");
+
+        // Add-Update all of the stream names as groups
+        const streamTracker = new Map<string, boolean>();
+
+        // Add all of our stream groupings
+        rows.each(function (rowData: StepNameHeader) {
+            const tr = d3.select(this);
+            const stream = rowData.streamId;
+
+            // Only add a TH for the first row of each stream
+            const streamData = !streamTracker.has(stream)
+                ? [{
+                    streamId: stream,
+                    rowspan: stepOutcomeTable.tableStreamMetadata.get(stream) ?? 1
+                }]
+                : [];
+
+            tr.selectAll("th.stream")
+                .data(streamData)
+                .join(
+                    enter => enter.append("th")
+                        .attr("class", "stream")
+                        .attr("rowspan", d => d.rowspan)
+                        .style("position", "sticky")
+                        .style("left", "0")
+                        .style("z-index", "2")
+                        .style("writing-mode", "vertical-rl")
+                        .style("transform", "rotate(180deg)")
+                        .style("background-color", modeColors.content)
+                        .style("white-space", "nowrap")
+                        .style("overflow", "hidden")
+                        .style("text-overflow", "ellipsis")
+                        .style("font-size", "16px")
+                        .style("text-align", "right")
+                        .style("padding-bottom", "10px")
+                        .style("box-sizing", "border-box")
+                        .style("border-left", "1px solid")
+                        .style("border-right", "1px solid")
+                        .style("border-bottom", "1px solid")
+                        .style("border-top", "1px solid")
+                        .text(d => d.streamId),
+                    update => update.text(d => d.streamId),
+                    exit => exit.remove()
+                );
+
+            // Attach tool-tip handlers.
+            tr.selectAll("th.stream").each(function (stepNameHeader: StepNameHeader) {
+                const th = d3.select(this);
+                th.on("mouseover", (event: MouseEvent) => {
+                    const tooltipContent = document.createElement("div");
+                    tooltipContent.style.paddingLeft = "4px";
+
+                    if (stepNameHeader) {
+                        const streamDiv = document.createElement("div");
+                        streamDiv.innerHTML = `<strong>Stream:</strong><p class="${tooltipClasses.tooltipLine}">${stepNameHeader.streamId}</p>`;
+                        tooltipContent.appendChild(streamDiv);
+                    }
+                    attachDynamicPositionToolTip(event, scaffold.current!.tooltip, tooltipContent);
+                });
+
+                th.on("mouseout", () => {
+                    scaffold.current!.tooltip.style("opacity", 0);
+                });
+            });
+
+            streamTracker.set(stream, true);
+        });
+
+        // Accessability row highlights.
+        rows.on("mouseover", function () {
+            d3.select(this)
+                .style("background-color", modeColors.content);
+        })
+            .on("mouseout", function () {
+                d3.select(this)
+                    .style("background-color", null);
+            });
+
+        // Add the first "step name" column item.
+        rows.each(function (rowData: StepNameHeader) {
+            const tr = d3.select(this);
+            const cells = tr.selectAll("td.step-name")
+                .data([rowData])
+                .join(
+                    enter => enter.append("td")
+                        .attr("class", "step-name").text(d => d.stepName),
+                    update => update.text(d => d.stepName),
+                    exit => exit.remove()
+                );
+
+            cells.style("position", "sticky")
+                .style("left", "32px")
+                .style("z-index", "1")
+                .style("background-color", modeColors.content)
+                .style("white-space", "nowrap")
+                .style("overflow", "hidden")
+                .style("padding-left", "8px")
+                .style("text-overflow", "ellipsis")
+                .style("box-sizing", "border-box")
+                .text(d => `${d.stepName}`)
+                .on("mouseover", function (_event, stepNameHeader) {
+                    const castedStepNameHeader = stepNameHeader as unknown as StepNameHeader;
+                    const tooltipContent = document.createElement("div");
+                    tooltipContent.style.paddingLeft = "4px";
+
+                    if (castedStepNameHeader) {
+                        const stepDiv = document.createElement("div");
+                        stepDiv.innerHTML = `<strong>Step:</strong><p class="${tooltipClasses.tooltipLine}">${castedStepNameHeader.stepName}</p>`;
+                        tooltipContent.appendChild(stepDiv);
+
+                        const streamDiv = document.createElement("div");
+                        streamDiv.innerHTML = `<strong>Stream:</strong><p class="${tooltipClasses.tooltipLine}">${castedStepNameHeader.streamId}</p>`;
+                        tooltipContent.appendChild(streamDiv);
+                    }
+
+                    attachTargettedToolTip(this as HTMLElement, scaffold.current!.tooltip, tooltipContent);
+                })
+                .on("mouseout", function () {
+                    scaffold.current!.tooltip.style("opacity", 0);
+                });
+        });
+
+        rows.selectAll("td.step-name").style("border", `1px solid ${modeColors.background}`);
+
+        const styleAndBindStepCell = (selection: d3.Selection<HTMLDivElement, StepOutcomeTableEntry, any, any>) => {
+            selection
+                .style("vertical-align", "middle")
+                .style("cursor", "pointer")
+                .style("width", "90%")
+                .style("height", "100%")
+                .style("background-color", d => getCellStyle(d).bgColor)
+                .style("display", "inline-block")
+                .style("text-align", "center")
+                .on("click", (event, d) => {
+                    onCellClick?.(d);
+                })
+                .on("mouseover", function (event, d) {
+                    if (!d) return;
+                    const entry = d as unknown as StepOutcomeTableEntry;
+                    const tooltipContent = document.createElement("div");
+                    tooltipContent.style.paddingLeft = "4px";
+
+                    const changeDiv = document.createElement("div");
+                    changeDiv.innerHTML = `<strong>Change:</strong><p class="${tooltipClasses.tooltipLineWithBorder}">${entry.change ?? "-"}</p>`;
+                    tooltipContent.appendChild(changeDiv);
+
+                    const durationDiv = document.createElement("div");
+                    durationDiv.innerHTML = `<strong>Duration:</strong><p class="${tooltipClasses.tooltipLine}">${entry.getDurationString?.() ?? "-"}</p>`;
+                    tooltipContent.appendChild(durationDiv);
+
+                    const stateDiv = document.createElement("div");
+                    stateDiv.innerHTML = `<strong>State:</strong><p class="${tooltipClasses.tooltipLineWithBorder}">${entry.stepResponse?.state ?? "-"}</p>`;
+                    tooltipContent.appendChild(stateDiv);
+
+                    const streamDiv = document.createElement("div");
+                    streamDiv.innerHTML = `<strong>Stream:</strong><p class="${tooltipClasses.tooltipLine}">${entry.streamId}</p>`;
+                    tooltipContent.appendChild(streamDiv);
+
+                    const jobDiv = document.createElement("div");
+                    jobDiv.innerHTML = `<strong>Job:</strong><p class="${tooltipClasses.tooltipLine}">${entry.jobName ?? "-"}</p>`;
+                    tooltipContent.appendChild(jobDiv);
+
+                    const stepDiv = document.createElement("div");
+                    stepDiv.innerHTML = `<strong>Step:</strong><p class="${tooltipClasses.tooltipLine}">${entry.stepResponse?.name ?? "-"}</p>`;
+                    tooltipContent.appendChild(stepDiv);
+
+                    attachTargettedToolTip(this as HTMLElement, scaffold.current!.tooltip, tooltipContent);
+                })
+                .on("mouseout", function () {
+                    scaffold.current!.tooltip.style("opacity", 0);
+                });
+        }
+
+        // Render all of the step outcome data
+        rows.each(function (rowData: StepNameHeader) {
+            const tr = d3.select(this);
+            const cellsRaw = stepOutcomeTable.getStepOutputRow(encodeStepNameFromStepNameHeader(rowData));
+            const cells = cellsRaw ? [...cellsRaw] : [];
+
+            const dataCells = tr.selectAll("td.step-outcome-data")
+                .data(cells)
+                .join("td")
+                .attr("class", "step-outcome-data");
+
+            dataCells.style("text-align", "center").style("border-right", `1px solid ${modeColors.content}`);
+
+            dataCells.selectAll("div")
+                .data(d => [d])
+                .join("div")
+                .call(div => styleAndBindStepCell(div as d3.Selection<HTMLDivElement, StepOutcomeTableEntry, any, any>))
+                .style("height", "25px");
+        });
+
+    }, [stepOutcomeTable.tableEntries, getCellStyle]);
+
+    const gridStyle: React.CSSProperties = {
+        display: 'grid',
+        gridTemplateRows: 'auto',
+        gridTemplateColumns: '1fr auto',
+    };
+
+    return (
+        <div style={gridStyle}>
+            {/* Left side: table */}
+            <div id="d3-container"
+                ref={containerRef}
+                style={{
+                    maxHeight: '85vh',
+                    overflow: 'auto'
+                }}
+            />
+            {/* Right side: incremental spinner */}
+            <IncrementalRefreshIndicatorPanel handler={handler} />
+        </div>
+    );
+};
+
+/**
+ * React Component that represents the provided stepOutcomeTable. This is a D3 table.
+ * @param stepOutcomeTable The data to base the table off of.
+ * @param handler The data handler used to query step outcome data.
+ * @param onCellSelected Cell selected callback.
+ * @returns React Component.
+ */
+export const StepOutcomeD3Table: React.FC<{ stepOutcomeTable: StepOutcomeTable; handler: StepOutcomeDataHandler, onCellSelected: any; }> = ({ stepOutcomeTable, handler, onCellSelected }) => {
+    return GenerateD3Table(stepOutcomeTable, handler, onCellSelected);
+};
diff --git a/Engine/Source/Programs/Horde/HordeDashboard/src/components/buildhealth/stepoutcome/StepOutcomeDataTypes.ts b/Engine/Source/Programs/Horde/HordeDashboard/src/components/buildhealth/stepoutcome/StepOutcomeDataTypes.ts
--- a/Engine/Source/Programs/Horde/HordeDashboard/src/components/buildhealth/stepoutcome/StepOutcomeDataTypes.ts
+++ b/Engine/Source/Programs/Horde/HordeDashboard/src/components/buildhealth/stepoutcome/StepOutcomeDataTypes.ts
@@ -3,7 +3,34 @@
 import { GetStepResponse, IssueData } from "horde/backend/Api";
 import { computeDuration, DURATION_NOT_SET, formatDurationFromMs, toDate } from "./StepOutcomeUtilities";
 
+/**
+ * Encodes a step outcome table entry to a stream qualified step name.
+ * Note: The template is not included in the hierarchy qualification since two separate templates within the same stream will collocate the same step row.
+ * @param stepOutcomeTableEntry The stepOutcomeTableEntry to encode.
+ * @returns The hierarchical stream qualified step name.
+ */
+export function encodeStepName(stepOutcomeTableEntry: StepOutcomeTableEntry) {
+    return `${stepOutcomeTableEntry.streamId.toLocaleLowerCase()}${KEY_SEPARATOR}${stepOutcomeTableEntry.stepResponse.name.toLocaleLowerCase()}`;
+}
+
+/**
+ * Encodes a step outcome table entry to a stream qualified step name.
+ * Note: The template is not included in the hierarchy qualification since two separate templates within the same stream will collocate the same step row.
+ * @param stepNameHeader The stepNameHeader to encode.
+ * @returns The hierarchical stream qualified step name.
+ */
+export function encodeStepNameFromStepNameHeader(stepNameHeader: StepNameHeader) {
+    return `${stepNameHeader.streamId.toLocaleLowerCase()}${KEY_SEPARATOR}${stepNameHeader.stepName.toLocaleLowerCase()}`;
+}
+
+// #region -- Constants --
+
 const KEY_SEPARATOR: string = "::";
+export const INVALID_CHANGELIST: number = -1;
+
+// #endregion -- Constants --
+
+// #region -- Data Types --
 
 /*
  * Convenience type to fold in extra metadata relevant for the StepResponse.
@@ -71,26 +98,6 @@ export class StepOutcomeTableEntry {
     // #endregion -- Public API --
 }
 
-/**
- * Encodes a step outcome table entry to a stream qualified step name.
- * Note: The template is not included in the hierarchy qualification since two separate templates within the same stream will collocate the same step row.
- * @param stepOutcomeTableEntry The stepOutcomeTableEntry to encode.
- * @returns The hierarchical stream qualified step name.
- */
-function encodeStepName(stepOutcomeTableEntry: StepOutcomeTableEntry) {
-    return `${stepOutcomeTableEntry.streamId.toLocaleLowerCase()}${KEY_SEPARATOR}${stepOutcomeTableEntry.stepResponse.name.toLocaleLowerCase()}`;
-}
-
-/**
- * Encodes a step outcome table entry to a stream qualified step name.
- * Note: The template is not included in the hierarchy qualification since two separate templates within the same stream will collocate the same step row.
- * @param stepNameHeader The stepNameHeader to encode.
- * @returns The hierarchical stream qualified step name.
- */
-export function encodeStepNameFromStepNameHeader(stepNameHeader: StepNameHeader) {
-    return `${stepNameHeader.streamId.toLocaleLowerCase()}${KEY_SEPARATOR}${stepNameHeader.stepName.toLocaleLowerCase()}`;
-}
-
 /**
  * Data struct that represents a change header. This header contains associated data related to rendering & ordering the header.
  */
@@ -118,6 +125,7 @@ export class StepOutcomeTable {
      */
     changeColLookup: Map<number, number> = new Map<number, number>();
     changeColHeaders: ChangeHeader[] = [];
+    changeOrderAscend: boolean = false;
 
     /**
      * Step name row headers.
@@ -228,7 +236,7 @@ export class StepOutcomeTable {
      */
     getStepOutputRow(rowName: string): Generator<StepOutcomeTableEntry>;
     getStepOutputRow(rowIndex: number): Generator<StepOutcomeTableEntry>;
-    *getStepOutputRow(row: string | number): Generator<StepOutcomeTableEntry> {
+    *getStepOutputRow(row: string | number): Generator<StepOutcomeTableEntry | null | undefined> {
         let rowIdx: number | undefined;
 
         if (typeof row === 'string') {
@@ -241,18 +249,12 @@ export class StepOutcomeTable {
             return;
         }
 
-        let rowEntries = this.tableEntries[rowIdx];
+        const rowEntries = this.tableEntries[rowIdx];
 
-        for (const entry of this.changeColHeaders) {
-            let colIdx = this.changeColLookup.get(entry.change);
-            if (colIdx === undefined) {
-                continue;
-            }
-
-            const value = rowEntries[colIdx];
-            if (value !== null && value !== undefined) {
-                yield value;
-            }
+        for (const entry of this.getDateAnchoredChanges()) {
+            const colIdx = this.changeColLookup.get(entry.change);
+            const value = colIdx !== undefined && colIdx > INVALID_CHANGELIST ? rowEntries[colIdx] : undefined;
+            yield value;
         }
     }
 
@@ -287,6 +289,7 @@ export class StepOutcomeTable {
      */
     orderTableDataByChange(ascending: boolean): void {
         this.changeColHeaders.sort((a, b) => ascending ? a.change - b.change : b.change - a.change);
+        this.changeOrderAscend = ascending;
     }
 
     /**
@@ -301,5 +304,39 @@ export class StepOutcomeTable {
         });
     }
 
+    /**
+     * Gets the the date anchored changes. 
+     * @returns All of the @see ChangeHeader items separated by entries that represent dates. These will have @see INVALID_CHANGELIST as their change property.
+     */
+    *getDateAnchoredChanges(): Generator<ChangeHeader> {
+        let prevDay: string | null = null;
+        let prevChangelist: string | null = null;
+        for (let i: number = 0; i < this.changeColHeaders.length; ++i) {
+
+            // If the date associated with the changelist header has changed since the last, we are on a new day. Set that as the colimn.
+            let changeHeader = this.changeColHeaders[i];
+            if (changeHeader.date !== undefined) {
+                const headerDate = new Date(changeHeader.date!);
+                const headerChangelist = changeHeader.change.toString();
+
+                // Extract the day portion as YYYY-MM-DD
+                const dayStr = `${headerDate.getUTCFullYear()}-${headerDate.getUTCMonth() + 1}-${headerDate.getUTCDate()}`;
+
+                if (prevDay !== dayStr) {
+                    prevDay = dayStr;
+                    prevChangelist = headerChangelist;
+                    let dateAnchorChangelistHeader: ChangeHeader = {
+                        change: INVALID_CHANGELIST,
+                        date: dayStr
+                    }
+                    yield dateAnchorChangelistHeader;
+                }
+                yield changeHeader;
+            }
+        }
+    }
+
     // #endregion -- Public API --
-}
\ No newline at end of file
+}
+
+// #endregion -- Data Types  --
\ No newline at end of file
diff --git a/Engine/Source/Programs/Horde/HordeDashboard/src/components/buildhealth/stepoutcome/StepOutcomeFluentUIComponent.tsx b/Engine/Source/Programs/Horde/HordeDashboard/src/components/buildhealth/stepoutcome/StepOutcomeFluentUIComponent.tsx
new file mode 100644
--- /dev/null
+++ b/Engine/Source/Programs/Horde/HordeDashboard/src/components/buildhealth/stepoutcome/StepOutcomeFluentUIComponent.tsx
@@ -0,0 +1,567 @@
+// Copyright Epic Games, Inc. All Rights Reserved.
+
+import { ConstrainMode, DetailsList, DetailsListLayoutMode, IColumn, IGroup, Label, mergeStyleSets, SelectionMode, Stack, TooltipHost } from "@fluentui/react";
+import { Text } from '@fluentui/react/lib/Text';
+import { useEffect, useRef, useState } from "react";
+import { StepOutcomeDataHandler } from "./StepOutcomeDataHandler";
+import { encodeStepNameFromStepNameHeader, StepNameHeader, StepOutcomeTable, StepOutcomeTableEntry } from "./StepOutcomeDataTypes";
+import { getCellStyle, IncrementalRefreshIndicatorPanel, ROW_HEADER_NAME, sharedBorderStyle } from "./StepOutcomeSharedUIComponents";
+
+type StepOutcomeFluentRow = {
+    [ROW_HEADER_NAME]: StepNameHeader;
+    [change: string]: StepOutcomeTableEntry | StepNameHeader | string | undefined;
+};
+
+// #region -- Styles --
+
+const stepOutcomeTooltipClasses = mergeStyleSets(sharedBorderStyle);
+
+const headerStylesHighlight = mergeStyleSets({
+    headerCell: {
+        fontWeight: '700 !important',
+        fontSize: 'larger !important'
+    },
+});
+
+const stepOutcomeTableClasses = mergeStyleSets({
+    stepOutcomeTableHeaders: {
+        selectors: {
+            '.ms-DetailsHeader-cellName': {
+                fontWeight: 200
+            },
+            '.ms-DetailsList-headerWrapper': {
+                position: 'sticky'
+            },
+            '.ms-DetailsHeader-cellTitle': {
+                writingMode: 'vertical-lr',
+                transform: 'rotate(180deg)',
+                whiteSpace: 'nowrap',
+                textAlign: 'center',
+                paddingBottom: '10px',
+                width: '36px !important',
+                userSelect: 'text',
+            },
+            '.ms-DetailsHeader': {
+                minHeight: '100px',
+                lineHeight: 'normal',
+                paddingTop: '0 !important',
+            },
+            '.ms-DetailsHeader-cell': {
+                verticalAlign: 'bottom',
+                minHeight: '80px',
+                borderLeft: '1px solid rgba(41, 41, 41, 0.35)', borderRight: '1px solid  rgba(41, 41, 41, 0.35)'
+            },
+            '.ms-DetailsHeader-cell span': {
+                overflow: 'visible',
+            },
+            '.ms-List-cell': {
+                minHeight: 40,
+                height: 40,
+            },
+            '.ms-SelectionZone': {
+                display: 'none',
+            },
+        },
+    },
+    removeCellContentPadding: {
+        selectors: {
+            '.ms-DetailsRow-cell': {
+                paddingLeft: '0 !important',
+            }
+        }
+    },
+    hideGroupCollapseButton: {
+        selectors: {
+            'button.ms-GroupHeader-expand': {
+                display: 'none',
+            },
+            '.ms-GroupHeader': {
+                height: '35px'
+            }
+        },
+    },
+    hideGroupHeaderTitle: {
+        selectors: {
+            '.ms-GroupHeader': {
+                height: '35px'
+            },
+            '.ms-GroupHeader-title': {
+                display: 'none',
+            },
+            'button.ms-GroupHeader-expand': {
+                display: 'none',
+            },
+            '.ms-GroupSpacer': {
+                width: '0 !important',
+            },
+            '.ms-DetailsRow-cell': {
+                padding: '10px 4px',
+                // Cleanup work for proper squares
+                // maxHeight: '35px',
+                // width: '35px !important',
+                // maxWdith: '35px',
+                borderLeft: '1px solid  rgba(41, 41, 41, 0.35)', borderRight: '1px solid  rgba(41, 41, 41, 0.35)'
+            },
+
+        },
+    },
+    disableGroupCollapseButton: {
+        selectors: {
+            '.ms-GroupHeader-dropIcon': {
+                display: 'none',
+                pointerEvents: 'none',
+                opacity: 0.5,
+            },
+        },
+    },
+    stickyHeaderOverride: {
+        selectors: {
+            '.ms-DetailsHeader-cellTitle': {
+                writingMode: 'vertical-rl',
+                whiteSpace: 'nowrap',
+                textAlign: 'center',
+                padding: '4px 2px',
+                paddingBottom: '10px',
+                width: '36px !important',
+            },
+            '.ms-DetailsHeader': {
+                minHeight: '100px',
+            },
+        },
+    },
+});
+
+// #endregion -- Styles --
+
+// #region -- Fluent Helpers --
+
+/**
+ * Produces Fluent UI compatible groups from the provided Step Outcome Table. Groups are currently derived based on Streams.
+ * @param stepOutcomeTable The step outcome table to use in producing Fluent UI groups.
+ * @returns The Fluent UI compatible groups.
+ */
+function getFluentGroups(stepOutcomeTable: StepOutcomeTable): any[] {
+    let returnGroups: IGroup[] = [];
+    let startIdx = 0;
+
+    for (const [key, value] of stepOutcomeTable.tableStreamMetadata) {
+        let streamGroup: IGroup = {
+            key: key,
+            name: key,
+            count: value,
+            startIndex: startIdx,
+            isCollapsed: false
+        };
+        startIdx += value;
+        returnGroups.push(streamGroup);
+    }
+    return returnGroups;
+}
+
+/**
+ * Produces Fluent UI compatible columns from the provided Step Outcome Table.
+ * @param stepOutcomeTable The step outcome table to use in producing Fluent UI columns.
+ * @returns The Fluent UI compatible columns.
+ */
+function getFluentColumns(stepOutcomeTable: StepOutcomeTable): IColumn[] {
+    const columns: IColumn[] = [
+        {
+            key: ROW_HEADER_NAME,
+            name: '',
+            fieldName: ROW_HEADER_NAME,
+            minWidth: 300,
+            onRender: (item) => {
+                return (
+                    <TooltipHost content={
+                        <div>
+                            <Stack>
+                                <Label>Step:</Label>
+                                <Text style={{ paddingLeft: 8 }}>
+                                    {item.stepName.stepName}
+                                </Text>
+                            </Stack>
+                            <Stack>
+                                <Label>Stream:</Label>
+                                <Text style={{ paddingLeft: 8 }}>
+                                    {item.stepName.streamId}
+                                </Text>
+                            </Stack>
+                        </div>
+                    }>
+                        <div style={{ whiteSpace: 'normal', wordWrap: 'break-word', minHeight: 20, maxHeight: 20 }}>
+                            {` ${item.stepName.stepName}`}
+                        </div>
+                    </TooltipHost>
+                );
+            },
+        }
+    ];
+
+    // We want to inject dates into the DetailsList column headers in order to assist with data consumption.
+    // Keep accumulators for the corresponding date strings.
+    let prevDay: string | null = null;
+    let prevChangelist: string | null = null;
+
+    for (let i: number = 0; i < stepOutcomeTable.changeColHeaders.length; ++i) {
+
+        // If the date associated with the changelist header has changed since the last, we are on a new day. Set that as the colimn.
+        let changeHeader = stepOutcomeTable.changeColHeaders[i];
+        if (changeHeader.date !== undefined) {
+            const headerDate = new Date(changeHeader.date!);
+            const headerChangelist = changeHeader.change.toString();
+
+            // Extract the day portion as YYYY-MM-DD
+            const dayStr = `${headerDate.getUTCFullYear()}-${headerDate.getUTCMonth() + 1}-${headerDate.getUTCDate()}`;
+
+            if (prevDay !== dayStr) {
+                prevDay = dayStr;
+                prevChangelist = headerChangelist;
+                columns.push({
+                    key: `${prevDay}_${changeHeader.change}`,
+                    name: changeHeader.change.toString(),
+                    minWidth: 10,
+
+                    onRenderHeader: (props, defaultRender) => {
+                        if (!props || !defaultRender) return null;
+
+                        return (
+                            <div className={headerStylesHighlight.headerCell}>
+                                <TooltipHost>
+                                    {dayStr}
+                                </TooltipHost>
+                            </div>
+                        );
+                    }
+                });
+            }
+        }
+
+        // Push all columns that relate to changelists.
+        columns.push({
+            key: changeHeader.change.toString(),
+            name: changeHeader.change.toString(),
+            fieldName: changeHeader.change.toString(),
+            minWidth: 10,
+            onRenderHeader: (props, defaultRender) => {
+                if (!props) {
+                    return null;
+                }
+                let innerChange = stepOutcomeTable.changeColHeaders[props?.columnIndex];
+                return props && defaultRender ? (
+                    <TooltipHost content={
+                        <Stack>
+                            <Label>Change:</Label>
+                            <Text style={{ paddingLeft: 8 }}>{changeHeader.change}</Text>
+                            <Label>Date (UTC):</Label>
+                            <Text style={{ paddingLeft: 8 }} >{new Date(changeHeader.date!).toISOString()}</Text>
+                            <Label>Date (Local):</Label>
+                            <Text style={{ paddingLeft: 8 }} >{new Date(changeHeader.date!).toString()}</Text>
+                        </Stack>
+                    }
+                    >
+                        {defaultRender(props)}
+                    </TooltipHost>
+                ) : null
+            },
+            onRender: (item: any) => {
+                const entry = item[changeHeader.change.toString()] as StepOutcomeTableEntry | null;
+                let bgColor = 'transparent';
+
+                if (entry) {
+                    let retrievedCellStype = getCellStyle(entry);
+                    bgColor = retrievedCellStype.bgColor;
+                }
+
+                return (
+                    <TooltipHost content={
+                        <Stack>
+                            <Label>Change:</Label>
+                            <Text className={stepOutcomeTooltipClasses.bottomBorder} style={{ paddingLeft: 8 }}>{changeHeader.change}</Text>
+                            {entry && (
+                                <>
+                                    <Label>Duration:</Label>
+                                    <Text style={{ paddingLeft: 8 }}>{entry.getDurationString()}</Text>
+                                    <Label>State:</Label>
+                                    <Text className={stepOutcomeTooltipClasses.bottomBorder} style={{ paddingLeft: 8 }}>{entry.stepResponse.state}</Text>
+                                    <Label>Stream:</Label>
+                                    <Text style={{ paddingLeft: 8 }}>{entry.streamId}</Text>
+                                    <Label>Job:</Label>
+                                    <Text style={{ paddingLeft: 8 }}>{entry.jobName}</Text>
+                                    <Label>Step:</Label>
+                                    <Text style={{ paddingLeft: 8 }}>{entry.stepResponse.name}</Text>
+                                </>)}
+                        </Stack>
+                    }>
+                        <div
+                            style={{
+                                backgroundColor: bgColor,
+                                minHeight: 20, maxHeight: 20,
+                                minWidth: 20, maxWidth: 20,
+                            }}
+                        >
+                            &nbsp;
+                        </div>
+                    </ TooltipHost>
+
+                );
+            }
+        });
+    }
+
+    return columns;
+}
+
+/**
+ * Produces Fluent UI compatible rows from the provided Step Outcome Table.
+ * @param stepOutcomeTable The step outcome table to use in producing Fluent UI rows.
+ * @returns The Fluent UI compatible rows.
+ */
+function getFluentRows(stepOutcomeTable: StepOutcomeTable, handler: StepOutcomeDataHandler): StepOutcomeFluentRow[] {
+    let items: StepOutcomeFluentRow[] = [];
+
+    for (let i: number = 0; i < stepOutcomeTable.stepNameRowHeaders.length; ++i) {
+        let rows: StepOutcomeTableEntry[] = [...handler.StepOutcomeTableData.getStepOutputRow(encodeStepNameFromStepNameHeader(stepOutcomeTable.stepNameRowHeaders[i]))];
+        let fluentRowEntry: StepOutcomeFluentRow = { [ROW_HEADER_NAME]: stepOutcomeTable.stepNameRowHeaders[i] };
+        for (let y: number = 0; y < rows.length; ++y) {
+            let tableEntry: StepOutcomeTableEntry = rows[y];
+            if (!tableEntry) {
+                continue;
+            }
+            if (tableEntry.change === undefined) {
+                continue;
+            }
+            fluentRowEntry[tableEntry.change!.toString()] = tableEntry;
+        }
+        items.push(fluentRowEntry);
+    }
+
+    return items;
+}
+
+// #endregion -- Fluent Helpers --
+
+/**
+ * Function to create a FLuent UI format table for the step outcome table.
+ * @param stepOutcomeTable The data to base the table off of.
+ * @param handler The data handler used to query step outcome data.
+ * @param onCellSelected Cell selected callback.
+ * @returns React Component.
+ */
+function GenerateFluentUITable(stepOutcomeTable: StepOutcomeTable, handler: StepOutcomeDataHandler, onCellClick: any): JSX.Element {
+    /**
+     * This is a fairly "virtual" (aka hacky) table in the following respects:
+     * - We need to have a sticky header, and first column
+     * - Sticky component works well for header, but not at all for first column
+     * - As a result, we have opted to take the following strategy:
+     *      - Create a dual "panel" view:
+     *          - Left panel is the 1st column; the Step Names
+     *          - Right panel - Top Stack is the Headers for all the changes
+     *          - Right panel - Bottom Stack is the cell data for all the steps
+     *      - Synchronize the Left panel vertical scroll to the Right Panel - Bottom Stack vertical scroll
+     *      - Synchronize the Rigth Panel Top Stack horizontal scroll to the Right Panel - Bottom Stack horizontal scroll
+     * 
+     * Whilst it's a tad hacky, it's predictable and easy to reason.
+     */
+
+    const leftPaneRef = useRef<HTMLDivElement>(null);
+    const rightTopPaneRef = useRef<HTMLDivElement>(null);
+    const rightBottomPaneRef = useRef<HTMLDivElement>(null);
+
+    const isSyncingVertical = useRef(false);
+    const isSyncingHorizontal = useRef(false);
+
+    const allColumns = getFluentColumns(stepOutcomeTable);
+    const items = getFluentRows(stepOutcomeTable, handler);
+    const groups = getFluentGroups(stepOutcomeTable);
+
+    const firstColumn = allColumns.slice(0, 1);
+    const restColumns = allColumns.slice(1);
+    const [hasScrollbar, setHasScrollbar] = useState(false);
+
+    useEffect(() => {
+        const checkScroll = () => {
+            const el = rightBottomPaneRef.current;
+            if (el) setHasScrollbar(el.scrollWidth > el.clientWidth);
+        };
+
+        checkScroll();
+        const resizeObserver = new ResizeObserver(checkScroll);
+        rightBottomPaneRef.current && resizeObserver.observe(rightBottomPaneRef.current);
+        return () => resizeObserver.disconnect();
+    }, []);
+
+    // Stash data structure to not create on all scroll events.
+    const targetMap = {
+        vertical: {
+            left: rightBottomPaneRef,
+            rightBottom: leftPaneRef,
+        },
+        horizontal: {
+            rightTop: rightBottomPaneRef,
+            rightBottom: rightTopPaneRef,
+        },
+    };
+
+    const VERTICAL_SCROLL_AFFORDANCE = 8;
+
+    const handleScroll = (direction: 'vertical' | 'horizontal', source: 'left' | 'rightTop' | 'rightBottom') => {
+        const isVertical = direction === 'vertical';
+        const syncRef = isVertical ? isSyncingVertical : isSyncingHorizontal;
+        const targetRef = targetMap[direction][source];
+
+        return (e: React.UIEvent<HTMLDivElement>) => {
+            if (syncRef.current) return;
+            syncRef.current = true;
+
+            const el = e.currentTarget;
+            let scrollValue = isVertical ? el.scrollTop : el.scrollLeft;
+
+            // Clamp for rightTop horizontal due to how we need to overlay the top DetailList as a "sticky row header".
+            if (!isVertical && source === 'rightTop') {
+                const maxAllowed = el.scrollWidth - el.clientWidth - VERTICAL_SCROLL_AFFORDANCE;
+                scrollValue = Math.min(scrollValue, maxAllowed);
+                el.scrollLeft = scrollValue; // keep UI in sync immediately
+            }
+
+            if (targetRef?.current) {
+                if (isVertical) {
+                    targetRef.current.scrollTop = scrollValue;
+                } else {
+                    targetRef.current.scrollLeft = scrollValue;
+                }
+            }
+
+            requestAnimationFrame(() => {
+                syncRef.current = false;
+            });
+        };
+    };
+
+    const gridStyle = {
+        display: 'grid',
+        gridTemplateRows: 'auto 1fr',
+        gridTemplateColumns: 'auto 1fr',
+        maxHeight: '85vh',
+        overflow: 'hidden',
+    };
+
+    return (
+        <div style={gridStyle}>
+            {/* Left Panel - Left Column - Step Name - (sticky) */}
+            <div
+                ref={leftPaneRef}
+                onScroll={handleScroll('vertical', 'left')}
+                className="no-scrollbar"
+                style={{
+                    maxHeight: '100%',
+                    overflowY: 'auto',
+                    position: 'relative',
+                    gridRow: 2,
+                    gridColumn: 1,
+                    overflowX: 'scroll',
+                    paddingBottom: hasScrollbar ? '17px' : '0',
+                    zIndex: 1,
+                    scrollbarWidth: 'none',
+                    msOverflowStyle: 'none',
+                }}
+            >
+                <DetailsList
+                    items={items}
+                    columns={firstColumn}
+                    groups={groups}
+                    selectionMode={SelectionMode.none}
+                    layoutMode={DetailsListLayoutMode.fixedColumns}
+                    constrainMode={ConstrainMode.unconstrained}
+                    isHeaderVisible={false}
+                    className={stepOutcomeTableClasses.hideGroupCollapseButton}
+                />
+            </div>
+
+            {/* Right Panel - Top Stack - Change Header Row (sticky) */}
+            <div
+                ref={rightTopPaneRef}
+                onScroll={handleScroll('horizontal', 'rightTop')}
+                className="no-scrollbar"
+                style={{
+                    gridRow: 1,
+                    gridColumn: 2,
+                    overflowX: 'auto',
+                    minHeight: '80px',
+                    maxHeight: '400px',
+                    position: 'relative',
+                    zIndex: 2,
+                }}
+            >
+                <DetailsList
+                    items={[]} // header only
+                    columns={restColumns}
+                    isHeaderVisible={true}
+                    selectionMode={SelectionMode.none}
+                    className={stepOutcomeTableClasses.stepOutcomeTableHeaders}
+                    layoutMode={DetailsListLayoutMode.fixedColumns}
+                    constrainMode={ConstrainMode.unconstrained}
+                />
+            </div>
+
+            {/* Right Panel - Bottom Stack - Step Data Cells (scroll driver) */}
+            <div
+                ref={rightBottomPaneRef}
+                onScroll={(e) => {
+                    handleScroll('vertical', 'rightBottom')(e);
+                    handleScroll('horizontal', 'rightBottom')(e);
+                }}
+                style={{
+                    maxHeight: '100%',
+                    overflowY: 'auto',
+                    position: 'relative',
+                    gridRow: 2,
+                    gridColumn: 2,
+                    overflow: 'auto',
+                }}
+            >
+                <DetailsList
+                    items={items}
+                    columns={restColumns}
+                    groups={groups}
+                    selectionMode={SelectionMode.none}
+                    layoutMode={DetailsListLayoutMode.fixedColumns}
+                    constrainMode={ConstrainMode.unconstrained}
+                    isHeaderVisible={false}
+                    className={stepOutcomeTableClasses.hideGroupHeaderTitle}
+                    onRenderField={(fieldProps, defaultRender) => {
+
+                        if (!fieldProps) {
+                            return null;
+                        }
+                        const { item, column, itemIndex } = fieldProps;
+                        const cellValue = column.fieldName ? item[column.fieldName] : undefined;
+                        const stepData: StepOutcomeTableEntry = cellValue as StepOutcomeTableEntry;
+                        const id = `cell-${itemIndex}-${column.key}`;
+
+                        return (
+                            <div id={id} key={id} style={{ paddingLeft: 0, cursor: 'pointer' }} onClick={() => onCellClick?.(stepData)}>
+                                {defaultRender?.(fieldProps)}
+                            </div>
+                        );
+                    }}
+                />
+            </div>
+            <IncrementalRefreshIndicatorPanel handler={handler} />
+        </div>
+    );
+};
+
+/**
+ * React Component that represents the provided stepOutcomeTable. This is a Fluent UI table.
+ * @param stepOutcomeTable The data to base the table off of.
+ * @param handler The data handler used to query step outcome data.
+ * @param onCellSelected Cell selected callback.
+ * @returns React Component.
+ */
+export const StepOutcomeFluentUITable: React.FC<{
+    stepOutcomeTable: StepOutcomeTable;
+    handler: StepOutcomeDataHandler;
+    onCellSelected: any;
+}> = ({ stepOutcomeTable, handler, onCellSelected }) => {
+    return GenerateFluentUITable(stepOutcomeTable, handler, onCellSelected);
+};
\ No newline at end of file
diff --git a/Engine/Source/Programs/Horde/HordeDashboard/src/components/buildhealth/stepoutcome/StepOutcomeSharedUIComponents.tsx b/Engine/Source/Programs/Horde/HordeDashboard/src/components/buildhealth/stepoutcome/StepOutcomeSharedUIComponents.tsx
new file mode 100644
--- /dev/null
+++ b/Engine/Source/Programs/Horde/HordeDashboard/src/components/buildhealth/stepoutcome/StepOutcomeSharedUIComponents.tsx
@@ -0,0 +1,95 @@
+// Copyright Epic Games, Inc. All Rights Reserved.
+
+import { Spinner, SpinnerSize } from "@fluentui/react";
+import dashboard, { StatusColor } from "horde/backend/Dashboard";
+import { observer } from "mobx-react-lite";
+import { StepOutcomeDataHandler } from "./StepOutcomeDataHandler";
+import { StepOutcomeTableEntry } from "./StepOutcomeDataTypes";
+
+// #region -- Constants --
+
+export const ROW_HEADER_NAME: string = "stepName" as const;
+
+// #endregion -- Constants --
+
+// #region -- Styles --
+
+export const sharedBorderStyle = {
+    bottomBorder: {
+        borderBottom: `1px solid rgba(80, 80, 80, 1)`
+    }
+};
+
+// #endregion -- Styles --
+
+/**
+ * Gets the cell style given the provided StepOutcomeTableEntry.
+ * @param entry The entry to base the style off of.
+ * @returns The matching style given the state of the entry. 
+ */
+export function getCellStyle(entry: StepOutcomeTableEntry | null): { bgColor: string } {
+    const scolors = dashboard.getStatusColors();
+    let resultColor: string = 'transparent';
+
+    if (!entry) {
+        return { bgColor: resultColor };
+    }
+
+    const WarningStr: string = "Warnings";
+    const SuccessStr: string = "Success";
+    const SkippedStr: string = "Skipped";
+    const FailureStr: string = "Failure";
+
+    const { stepResponse } = entry;
+    const { outcome, state } = stepResponse;
+    const colors: Record<string, string> = {
+        "Success": scolors.get(StatusColor.Success)!,
+        "Failure": scolors.get(StatusColor.Failure)!,
+        "Warnings": scolors.get(StatusColor.Warnings)!,
+        "Skipped": scolors.get(StatusColor.Skipped)!,
+        "Running": scolors.get(StatusColor.Running)!,
+        "Waiting": scolors.get(StatusColor.Waiting)!,
+        "Ready": scolors.get(StatusColor.Ready)!,
+        "Unspecified": scolors.get(StatusColor.Skipped)!,
+    };
+
+    // There are some interesting mappings going on internally with respect to state & outcome. Outcome & State have overlapping "result view" space, and as such we need to 
+    // narrowly select the edge cases (e.g. outcome=success && state=completed uniquely signifies a completed run with no errors, otherwise outcome=success it could mean ready, waiting, or running).
+    // Then we can fall back to the state as the truth.
+    if (outcome === FailureStr) {
+        resultColor = state === SkippedStr ? colors[SkippedStr] : colors[FailureStr];
+    }
+    // A step can be emitting warnings whilst still running, so we narrowly select this colouring.
+    else if (outcome === WarningStr) {
+        resultColor = colors[WarningStr];
+    }
+    // A step can have a completed outcome, but merely be waiting to run, running, or ready.
+    else if (outcome === SuccessStr && state === "Completed") {
+        resultColor = colors[SuccessStr];
+    }
+    // The edge cases have been satisfied, yield to the state to define result view color.
+    else {
+        resultColor = colors[state];
+    }
+
+    return { bgColor: resultColor };
+}
+
+/**
+ * React Component that provides refresh spinner on incremental refresh. 
+ * Note: This is a separate component to prevent rerender of primary table @see GenerateFluentUITable upon refresh spinner.
+ * @param handler The data handler used to query step outcome data.
+ * @returns React Component.
+ */
+export const IncrementalRefreshIndicatorPanel: React.FC<{ handler: StepOutcomeDataHandler }> = observer(({ handler }) => (
+    <div
+        style={{
+            gridRow: '1 / span 2',
+            gridColumn: 3,
+            overflowY: 'auto',
+            width: 25,
+        }}
+    >
+        {handler.isInIncrementalDataRefresh && <Spinner size={SpinnerSize.small} />}
+    </div>
+));
\ No newline at end of file
