Horde:  Temporary fix for hung job steps on farm, catch log signal that outputs in cases where UAT or subprocesses are hanging.

diff --git a/Engine/Source/Programs/Horde/Drivers/JobDriver/Execution/JobExecutor.cs b/Engine/Source/Programs/Horde/Drivers/JobDriver/Execution/JobExecutor.cs
--- a/Engine/Source/Programs/Horde/Drivers/JobDriver/Execution/JobExecutor.cs
+++ b/Engine/Source/Programs/Horde/Drivers/JobDriver/Execution/JobExecutor.cs
@@ -1687,17 +1687,49 @@ namespace JobDriver.Execution
 					using (ManagedProcess process = new ManagedProcess(processGroup, fileName, arguments, null, newEnvironment, null, ProcessPriorityClass.Normal))
 					{
 						bool hasExited = false;
+						StringBuilder hangDetection = new StringBuilder();
 						try
 						{
 							async Task WaitForExitOrCancelAsync()
 							{
 								await process.WaitForExitAsync(cancellationToken);
 								hasExited = true;
-								cancellationSource.CancelAfter(TimeSpan.FromSeconds(60.0));
+								if (!cancellationSource.IsCancellationRequested)
+								{
+									cancellationSource.CancelAfter(TimeSpan.FromSeconds(60.0));
+								}								
 							}
 
 							Task cancelTask = WaitForExitOrCancelAsync();
-							await process.CopyToAsync((buffer, offset, length) => filter.WriteData(buffer.AsMemory(offset, length)), 4096, cancellationSource.Token);
+							await process.CopyToAsync((buffer, offset, length) => { 
+								filter.WriteData(buffer.AsMemory(offset, length));
+
+								// UAT currently can hang a step and generally we're seeing the below logging in these cases
+								// this is a temporary fix to stop farm issues, until we can properly address UAT hangs 
+								// that are often caused by subprocesses, and/or until we can put better signalling in place
+								if (!cancellationSource.IsCancellationRequested)
+								{
+									try
+									{
+										hangDetection.Append(Encoding.UTF8.GetString(buffer));
+										string hang = hangDetection.ToString();
+
+										if (hang.Contains("some output may be missing; we gave up", StringComparison.OrdinalIgnoreCase))
+										{
+											cancellationSource.Cancel();
+										}										
+										else if (hang.Length > 8192)
+										{
+											// 8k window
+											hangDetection.Clear();
+										}
+									}
+									catch (Exception ex)
+									{
+										logger.LogWarning(ex, "Temporary hang detected threw an exception");
+									}
+								}
+							}, 4096, cancellationSource.Token);
 							await cancelTask;
 						}
 						catch (OperationCanceledException) when (hasExited)
