Horde: Guard against multiple hang detection in same step and fix string decoding on signal

diff --git a/Engine/Source/Programs/Horde/Drivers/JobDriver/Execution/JobExecutor.cs b/Engine/Source/Programs/Horde/Drivers/JobDriver/Execution/JobExecutor.cs
--- a/Engine/Source/Programs/Horde/Drivers/JobDriver/Execution/JobExecutor.cs
+++ b/Engine/Source/Programs/Horde/Drivers/JobDriver/Execution/JobExecutor.cs
@@ -1687,6 +1687,7 @@ namespace JobDriver.Execution
 					using (ManagedProcess process = new ManagedProcess(processGroup, fileName, arguments, null, newEnvironment, null, ProcessPriorityClass.Normal))
 					{
 						bool hasExited = false;
+						bool hangException = false;
 						StringBuilder hangDetection = new StringBuilder();
 						try
 						{
@@ -1701,17 +1702,19 @@ namespace JobDriver.Execution
 							}
 
 							Task cancelTask = WaitForExitOrCancelAsync();
-							await process.CopyToAsync((buffer, offset, length) => { 
-								filter.WriteData(buffer.AsMemory(offset, length));
+							await process.CopyToAsync((buffer, offset, length) => {
+
+								Memory<byte> bufferMemory = buffer.AsMemory(offset, length);
+								filter.WriteData(bufferMemory);
 
 								// UAT currently can hang a step and generally we're seeing the below logging in these cases
 								// this is a temporary fix to stop farm issues, until we can properly address UAT hangs 
 								// that are often caused by subprocesses, and/or until we can put better signalling in place
-								if (!cancellationSource.IsCancellationRequested)
+								if (!hangException && !cancellationSource.IsCancellationRequested)
 								{
 									try
 									{
-										hangDetection.Append(Encoding.UTF8.GetString(buffer));
+										hangDetection.Append(Encoding.UTF8.GetString(bufferMemory.Span));
 										string hang = hangDetection.ToString();
 
 										if (hang.Contains("some output may be missing; we gave up", StringComparison.OrdinalIgnoreCase))
@@ -1726,6 +1729,7 @@ namespace JobDriver.Execution
 									}
 									catch (Exception ex)
 									{
+										hangException = true;
 										logger.LogWarning(ex, "Temporary hang detected threw an exception");
 									}
 								}
