Horde Dashboard: Reduce JobDetails V1 QA Fixes

diff --git a/Engine/Source/Programs/Horde/HordeDashboard/src/components/error/ErrorBar.tsx b/Engine/Source/Programs/Horde/HordeDashboard/src/components/error/ErrorBar.tsx
--- a/Engine/Source/Programs/Horde/HordeDashboard/src/components/error/ErrorBar.tsx
+++ b/Engine/Source/Programs/Horde/HordeDashboard/src/components/error/ErrorBar.tsx
@@ -57,7 +57,7 @@ export const ErrorBar: React.FC = observer(() => {
                 </div>
             }
         >
-            <span>{error.title}: {errorBarStore.message}. {(error.url || anyHelp) && "Check your browser's developer console."}</span>
+            <span>{error.title}: {errorBarStore.message}. {(!!error.url || anyHelp) && "Check your browser's developer console."}</span>
         </MessageBar>
     );
 });
\ No newline at end of file
diff --git a/Engine/Source/Programs/Horde/HordeDashboard/src/components/error/ErrorStore.ts b/Engine/Source/Programs/Horde/HordeDashboard/src/components/error/ErrorStore.ts
--- a/Engine/Source/Programs/Horde/HordeDashboard/src/components/error/ErrorStore.ts
+++ b/Engine/Source/Programs/Horde/HordeDashboard/src/components/error/ErrorStore.ts
@@ -157,6 +157,11 @@ class ErrorStore {
       }
 
       if (error.message) {
+         // Catch(reason) defaults to type 'any,' missing eroneous assignments of incorrect types.
+         if(typeof error.message !== "string"){
+            console.error("Error message not set to string. Check error store set arguments.")
+            return "[BAD ERROR MESSAGE]"
+         }
          return error.message;
       }
 
diff --git a/Engine/Source/Programs/Horde/HordeDashboard/src/components/logs/LogSource.ts b/Engine/Source/Programs/Horde/HordeDashboard/src/components/logs/LogSource.ts
--- a/Engine/Source/Programs/Horde/HordeDashboard/src/components/logs/LogSource.ts
+++ b/Engine/Source/Programs/Horde/HordeDashboard/src/components/logs/LogSource.ts
@@ -41,8 +41,9 @@ export abstract class LogSource extends PollBase {
        
         try {
            logData = await backend.getLogData(logId);
-        } catch (err) {
-            throw new Error(`Failed to load log data for ${logId}: ${err}`);
+        } catch (reason) {
+            errorBarStore.set({title: "Error Loading Log Data", reason: reason, message: `Failed to load log data for ${logId}`});
+            throw new Error(reason);
         }
 
         // JobId may come in as a string of 0s. This converts to boolean.
@@ -115,7 +116,7 @@ export abstract class LogSource extends PollBase {
         try {
             await this.updateLogData();
         } catch(reason) {
-            errorBarStore.set({title: "Error Loading Log Data", message: reason}, true)
+            errorBarStore.set({title: "Error Loading Log Data", reason: reason}, true)
         }
     }
 
@@ -127,7 +128,7 @@ export abstract class LogSource extends PollBase {
                 this.resize(this.logData.lineCount);
             })
         } catch(reason) {
-            errorBarStore.set({title: "Error Updating Log Data", message: reason}, true);
+            errorBarStore.set({title: "Error Updating Log Data", reason: reason}, true);
         }
     }
 
@@ -141,8 +142,8 @@ export abstract class LogSource extends PollBase {
         const count = lineCount - this.logItems.length;
 
         if(count <= 0) return;
-
-        this.logItems.push(...this.getPlaceholderLogItems(count, this.logItems.length + 1));
+        // Do not use spread operator to push placeholder items. Causes Range Error (overloads Call Stack).
+        this.logItems = this.logItems.concat(this.getPlaceholderLogItems(count, this.logItems.length + 1));
         this._itemsVersion++;
     }
     
@@ -210,7 +211,7 @@ export abstract class LogSource extends PollBase {
                 this._itemsVersion++;
             })
         } catch(reason) {
-            errorBarStore.set({title: "Error Loading Line Data", message: reason}, true)
+            errorBarStore.set({title: "Error Loading Line Data", reason: reason}, true)
         }
     }
 
@@ -224,7 +225,7 @@ export abstract class LogSource extends PollBase {
         try {
             await this.updateLogData();
         } catch(reason) {
-            errorBarStore.set({title: "Error Polling Log Data", message: reason}, true)
+            errorBarStore.set({title: "Error Polling Log Data", reason: reason}, true)
         }
     }
 }
@@ -402,7 +403,7 @@ export class JobLogSource extends LogSource {
                 this._events = events;
             })
         } catch (reason) {
-            errorBarStore.set({title: "Error Getting Events", message: reason}, true);
+            errorBarStore.set({title: "Error Getting Events", reason: reason}, true);
         }
     }
 
@@ -418,7 +419,7 @@ export class JobLogSource extends LogSource {
                 this._issues = [...unresolvedIssues, ...resolvedIssues];
             })
         } catch (reason) {
-            errorBarStore.set({title: "Error Getting Issues", message: reason}, true);
+            errorBarStore.set({title: "Error Getting Issues", reason: reason}, true);
         }
     }
 
@@ -438,8 +439,8 @@ export class JobLogSource extends LogSource {
             runInAction(() => {
                 this.artifactsV2 = artifactResponse.artifacts;
             })
-        } catch (err) {
-            errorBarStore.set({title: "Error Getting Artifacts", message: err}, true);
+        } catch (reason) {
+            errorBarStore.set({title: "Error Getting Artifacts", reason: reason}, true);
         }
     }
 
@@ -557,7 +558,7 @@ export class LeaseLogSource extends LogSource {
                 this.agent = agent;
             })
         } catch (reason) {
-            errorBarStore.set({title: "Error Updating Agent", message: reason}, true);
+            errorBarStore.set({title: "Error Updating Agent", reason: reason}, true);
         }
     }
 
@@ -568,7 +569,7 @@ export class LeaseLogSource extends LogSource {
                 this.lease = lease;
             })
         } catch (reason) {
-            errorBarStore.set({title: "Error Updating Lease", message: reason}, true);
+            errorBarStore.set({title: "Error Updating Lease", reason: reason}, true);
         }
     }
 
diff --git a/Engine/Source/Programs/Horde/HordeDashboard/src/components/logs/LogView.tsx b/Engine/Source/Programs/Horde/HordeDashboard/src/components/logs/LogView.tsx
--- a/Engine/Source/Programs/Horde/HordeDashboard/src/components/logs/LogView.tsx
+++ b/Engine/Source/Programs/Horde/HordeDashboard/src/components/logs/LogView.tsx
@@ -387,9 +387,8 @@ const LogList: React.FC<{ logId: string }> = observer(({ logId }) => {
    const [searchState, setSearchState] = useState<{ search?: string, results?: number[], curRequest?: any, text?: string }>({});
    const [issueHistory, setIssueHistory] = useState(false);
    const [logHistory, setLogHistory] = useState(false);
-   const [logArtifacts, setLogArtifacts] = useState("");
    const [logTelemetry, setLogTelemetry] = useState(false);
-   const [logError, setLogError] = useState("")
+   const [logError, setLogError] = useState<Error | undefined>(undefined);
 
    let [historyAgentId, setHistoryAgentId] = useState<string | undefined>(undefined);
 
@@ -458,13 +457,6 @@ const LogList: React.FC<{ logId: string }> = observer(({ logId }) => {
 
       searchSelectedText();
    };
-   
-
-   if (logError) {
-      return <Stack horizontalAlign='center' style={{ paddingTop: 24 }} >
-         <Text variant='mediumPlus'>{`Unable to load log data - ${logError}`}</Text>
-      </Stack>
-   }
 
    if (handler && handler.logSource && handler.logSource.logId !== logId) {
       selection.setItems([], true);
@@ -501,6 +493,15 @@ const LogList: React.FC<{ logId: string }> = observer(({ logId }) => {
 
    }
 
+   // Important: This should come before handler init below
+   if(logError) {
+      return (
+         <Stack horizontalAlign='center' style={{ paddingTop: 24 }} >
+            <Text variant='mediumPlus'>{`Unable to load log data - ${logError.message}`}</Text>
+         </Stack>
+      )
+   }
+
    if(!handler){
       let lineIndexQuery: number | undefined;
       if (query.get("lineIndex")) {
@@ -512,6 +513,9 @@ const LogList: React.FC<{ logId: string }> = observer(({ logId }) => {
       const handler = new LogHandler();
       handler.init(logId, lineIndexQuery ?? undefined).then(() => {
          setHandler(handler);
+      }).catch(error => {
+         // This should primarily catch errors bubbling up from logSource.create()
+         setLogError(error);
       });
 
       return <Spinner size={SpinnerSize.large} />;
@@ -696,7 +700,8 @@ const LogList: React.FC<{ logId: string }> = observer(({ logId }) => {
                                     <Text style={{ fontSize: 14, fontFamily: "Horde Open Sans SemiBold" }}>Structured Log Line</Text>
                                  </Stack>
                                  <Stack style={{ paddingLeft: 12 }}>
-                                    <Text style={{ fontSize: 11, whiteSpace: "pre-wrap", fontFamily: "Horde Cousine Regular" }}>{JSON.stringify(item.line, undefined, 2).replaceAll("\\r", "").replaceAll("\\n", "\n")}</Text></Stack>
+                                    <Text style={{ fontSize: 11, whiteSpace: "pre-wrap", fontFamily: "Horde Cousine Regular" }}>{JSON.stringify(item.line, undefined, 2).replaceAll("\\r", "").replaceAll("\\n", "\n")}</Text>
+                                 </Stack>
                               </Stack>
                            </Callout>}
 
@@ -1164,17 +1169,15 @@ const LogList: React.FC<{ logId: string }> = observer(({ logId }) => {
       
       {
          logSourceIsJobLogSource 
-         && !!artifactContext 
-         && logArtifacts !== "legacy" 
-         && 
-            <JobArtifactsModal 
+         && !!artifactContext
+         && <JobArtifactsModal 
                jobId={jobDetails!.jobData!.id} 
                stepId={jobDetails!.stepByLogId(logId)?.id!} 
                artifacts={(logSource as JobLogSource).artifactsV2} 
                contextType={artifactContext} artifactPath={artifactPath} 
                artifactId={artifactId} 
                onClose={() => { navigate(window.location.pathname, { replace: true }) }} 
-            />
+            /> 
       }
 
       {!!historyAgentId && <HistoryModal agentId={historyAgentId} onDismiss={() => { navigate(baseUrl, { replace: true }); setHistoryAgentId(undefined) }} />}
