Allow querying for job testdata with just the job id

#rnx

diff --git a/Engine/Source/Programs/Horde/Plugins/Build/HordeServer.Build.Tests/Jobs/TestData/TestDataTests.cs b/Engine/Source/Programs/Horde/Plugins/Build/HordeServer.Build.Tests/Jobs/TestData/TestDataTests.cs
--- a/Engine/Source/Programs/Horde/Plugins/Build/HordeServer.Build.Tests/Jobs/TestData/TestDataTests.cs
+++ b/Engine/Source/Programs/Horde/Plugins/Build/HordeServer.Build.Tests/Jobs/TestData/TestDataTests.cs
@@ -342,6 +342,7 @@ namespace HordeServer.Tests.Jobs.TestData
 			IJobStep step2 = job2.Batches[0].Steps[0];
 			IJob job3 = CreateJob(_mainStreamId, 106, "Test Build", _graph);
 			IJobStep step3 = job3.Batches[0].Steps[0];
+			IJobStep step4 = job3.Batches[0].Steps[1];
 
 			BsonDocument testData = BsonDocument.Parse(String.Join('\n', _sessionTestDataV2Lines));
 			List<(string key, BsonDocument)> data = new List<(string key, BsonDocument)>();
@@ -369,6 +370,7 @@ namespace HordeServer.Tests.Jobs.TestData
 			await TestDataCollectionV2.AddAsync(job2, step2, data.ToArray());
 			await TestDataCollectionV2.AddAsync(job3, step3, data.ToArray());
 			await TestDataCollectionV2.AddAsync(job3, step3, data2.ToArray());
+			await TestDataCollectionV2.AddAsync(job3, step4, data.ToArray());
 
 			await TestDataService.TickForTestingAsync();
 
@@ -432,7 +434,7 @@ namespace HordeServer.Tests.Jobs.TestData
 			Assert.IsNotNull(testSessionResult);
 			Assert.IsNotNull(testSessionResult.Value);
 			List<GetTestSessionResponse> sessions = testSessionResult.Value;
-			Assert.AreEqual(5, sessions.Count);
+			Assert.AreEqual(6, sessions.Count);
 			Assert.AreEqual(TestOutcome.Failure, sessions[0].Outcome);
 
 			GetPhaseSessionsRequest phaseSessionRequest = new GetPhaseSessionsRequest() { StreamIds = streamIds.ToList(), PhaseIds = tests.SelectMany(t => t.Phases.Select(p => p.Id)).ToList() };
@@ -440,7 +442,7 @@ namespace HordeServer.Tests.Jobs.TestData
 			Assert.IsNotNull(phaseSessionResult);
 			Assert.IsNotNull(phaseSessionResult.Value);
 			List<GetTestPhaseSessionResponse> phaseSessions = phaseSessionResult.Value;
-			Assert.AreEqual(12, phaseSessions.Count);
+			Assert.AreEqual(15, phaseSessions.Count);
 			Assert.AreEqual(TestPhaseOutcome.Failed, phaseSessions.First(p => p.PhaseRef == phaseId).Outcome);
 
 			// Get the sessions for one phase in one stream
@@ -449,7 +451,7 @@ namespace HordeServer.Tests.Jobs.TestData
 			Assert.IsNotNull(phaseSessionResult);
 			Assert.IsNotNull(phaseSessionResult.Value);
 			phaseSessions = phaseSessionResult.Value;
-			Assert.AreEqual(3, phaseSessions.Count);
+			Assert.AreEqual(4, phaseSessions.Count);
 
 			ActionResult<object> testdata =  await TestDataController.GetTestDataV2Async(sessions[0].TestDataId.ToString());
 			Assert.IsNotNull(testdata);
@@ -461,6 +463,11 @@ namespace HordeServer.Tests.Jobs.TestData
 			Assert.AreEqual(3, testDataIds.Value.Count);
 			Assert.IsTrue(testDataIds.Value.All(item => sessions.Any(s => s.TestDataId == item.TestDataId)));
 			Assert.IsTrue(testDataIds.Value.All(item => sessions.First(s => s.TestDataId == item.TestDataId).NameRef == item.TestNameRef));
+
+			testDataIds = await TestDataController.GetJobStepTestDataV2Async(job3.Id.ToString());
+			Assert.IsNotNull(testDataIds);
+			Assert.IsNotNull(testDataIds.Value);
+			Assert.AreEqual(4, testDataIds.Value.Count);
 		}
 
 		private readonly string[] _simpleTestDataLines =
diff --git a/Engine/Source/Programs/Horde/Plugins/Build/HordeServer.Build/Jobs/TestData/ITestDataCollection.cs b/Engine/Source/Programs/Horde/Plugins/Build/HordeServer.Build/Jobs/TestData/ITestDataCollection.cs
--- a/Engine/Source/Programs/Horde/Plugins/Build/HordeServer.Build/Jobs/TestData/ITestDataCollection.cs
+++ b/Engine/Source/Programs/Horde/Plugins/Build/HordeServer.Build/Jobs/TestData/ITestDataCollection.cs
@@ -189,7 +189,7 @@ namespace HordeServer.Jobs.TestData
 		/// <param name="stepId">Unique id of the step</param>
 		/// <param name="cancellationToken">Cancellation token for the operation</param>
 		/// <returns>The testdata ids</returns>
-		Task<IReadOnlyList<ITestData>> GetJobStepTestDataAsync(JobId jobId, JobStepId stepId, CancellationToken cancellationToken = default);
+		Task<IReadOnlyList<ITestData>> GetJobStepTestDataAsync(JobId jobId, JobStepId? stepId = null, CancellationToken cancellationToken = default);
 
 		/// <summary>
 		/// Gets a testdata by ID
diff --git a/Engine/Source/Programs/Horde/Plugins/Build/HordeServer.Build/Jobs/TestData/TestDataCollectionV2.cs b/Engine/Source/Programs/Horde/Plugins/Build/HordeServer.Build/Jobs/TestData/TestDataCollectionV2.cs
--- a/Engine/Source/Programs/Horde/Plugins/Build/HordeServer.Build/Jobs/TestData/TestDataCollectionV2.cs
+++ b/Engine/Source/Programs/Horde/Plugins/Build/HordeServer.Build/Jobs/TestData/TestDataCollectionV2.cs
@@ -684,9 +684,19 @@ namespace HordeServer.Jobs.TestData
 		}
 
 		/// <inheritdoc/>
-		public async Task<IReadOnlyList<ITestData>> GetJobStepTestDataAsync(JobId jobId, JobStepId stepId, CancellationToken cancellationToken)
+		public async Task<IReadOnlyList<ITestData>> GetJobStepTestDataAsync(JobId jobId, JobStepId? stepId = null, CancellationToken cancellationToken = default)
 		{
-			return await _testDataDocuments.Find(x => x.JobId == jobId && x.StepId == stepId).ToListAsync(cancellationToken);
+			FilterDefinition<TestDataDocument> filter = FilterDefinition<TestDataDocument>.Empty;
+			FilterDefinitionBuilder<TestDataDocument> filterBuilder = Builders<TestDataDocument>.Filter;
+
+			filter &= filterBuilder.Eq(x => x.JobId, jobId);
+
+			if (stepId != null)
+			{
+				filter &= filterBuilder.Eq(x => x.StepId, stepId);
+			}
+
+			return await _testDataDocuments.Find(filter).ToListAsync(cancellationToken);
 		}
 
 		/// <inheritdoc/>
diff --git a/Engine/Source/Programs/Horde/Plugins/Build/HordeServer.Build/Jobs/TestData/TestDataController.cs b/Engine/Source/Programs/Horde/Plugins/Build/HordeServer.Build/Jobs/TestData/TestDataController.cs
--- a/Engine/Source/Programs/Horde/Plugins/Build/HordeServer.Build/Jobs/TestData/TestDataController.cs
+++ b/Engine/Source/Programs/Horde/Plugins/Build/HordeServer.Build/Jobs/TestData/TestDataController.cs
@@ -468,7 +468,7 @@ namespace HordeServer.Jobs.TestData
 		}
 
 		/// <summary>
-		/// Retrieve list of testdata ids run within a job step
+		/// Retrieve list of testdata ids run within a job or job step
 		/// </summary>
 		/// <param name="jobId">job Id from which the test was run in</param>
 		/// <param name="stepId">step Id from which the test was run in</param>
@@ -477,14 +477,14 @@ namespace HordeServer.Jobs.TestData
 		[HttpGet]
 		[Route("/api/v3/testdata/")]
 		[ProducesResponseType(typeof(List<GetJobStepTestDataIDResponse>), 200)]
-		public async Task<ActionResult<List<GetJobStepTestDataIDResponse>>> GetJobStepTestDataV2Async([FromQuery] string jobId, [FromQuery] string stepId, CancellationToken cancellationToken = default)
+		public async Task<ActionResult<List<GetJobStepTestDataIDResponse>>> GetJobStepTestDataV2Async([FromQuery] string jobId, [FromQuery] string? stepId = null, CancellationToken cancellationToken = default)
 		{
 			if (!await _jobService.AuthorizeAsync(JobId.Parse(jobId), JobAclAction.ViewJob, User, _buildConfig.Value, cancellationToken))
 			{
 				return Forbid();
 			}
 
-			IReadOnlyList<ITestData> testData = await _testDataService.GetJobStepTestDataV2Async(JobId.Parse(jobId), JobStepId.Parse(stepId), cancellationToken);
+			IReadOnlyList<ITestData> testData = await _testDataService.GetJobStepTestDataV2Async(JobId.Parse(jobId), stepId != null ? JobStepId.Parse(stepId) : null, cancellationToken);
 			IReadOnlyList<ITestNameRef> tests = await _testDataService.FindTestNameRefsV2Async(keys: testData.Select(x => x.Key).ToArray(), cancellationToken: cancellationToken);
 
 			return testData.Select(item => new GetJobStepTestDataIDResponse() {
diff --git a/Engine/Source/Programs/Horde/Plugins/Build/HordeServer.Build/Jobs/TestData/TestDataService.cs b/Engine/Source/Programs/Horde/Plugins/Build/HordeServer.Build/Jobs/TestData/TestDataService.cs
--- a/Engine/Source/Programs/Horde/Plugins/Build/HordeServer.Build/Jobs/TestData/TestDataService.cs
+++ b/Engine/Source/Programs/Horde/Plugins/Build/HordeServer.Build/Jobs/TestData/TestDataService.cs
@@ -182,7 +182,7 @@ namespace HordeServer.Jobs.TestData
 		/// <param name="stepId">Unique id of the step</param>
 		/// <param name="cancellationToken">Cancellation token for the operation</param>
 		/// <returns>The testdata ids</returns>
-		public async Task<IReadOnlyList<ITestData>> GetJobStepTestDataV2Async(JobId jobId, JobStepId stepId, CancellationToken cancellationToken)
+		public async Task<IReadOnlyList<ITestData>> GetJobStepTestDataV2Async(JobId jobId, JobStepId? stepId = null, CancellationToken cancellationToken = default)
 		{
 			return await _testDataV2.GetJobStepTestDataAsync(jobId, stepId, cancellationToken);
 		}
