[Horde][BuildHealth][Server] Expose DateUtc on GetCommitResponse

diff --git a/Engine/Source/Programs/Horde/Plugins/Build/HordeServer.Build/Streams/StreamsController.cs b/Engine/Source/Programs/Horde/Plugins/Build/HordeServer.Build/Streams/StreamsController.cs
--- a/Engine/Source/Programs/Horde/Plugins/Build/HordeServer.Build/Streams/StreamsController.cs
+++ b/Engine/Source/Programs/Horde/Plugins/Build/HordeServer.Build/Streams/StreamsController.cs
@@ -226,7 +226,7 @@ namespace HordeServer.Streams
 
 		static GetCommitResponse CreateGetCommitResponse(ICommit commit, IUser author, IReadOnlyList<CommitTag>? tags, IReadOnlyList<string>? files)
 		{
-			GetCommitResponse response = new GetCommitResponse(author.ToThinApiResponse(), commit.Description) { Id = commit.Id };
+			GetCommitResponse response = new GetCommitResponse(author.ToThinApiResponse(), commit.Description, commit.DateUtc) { Id = commit.Id };
 			if (tags != null)
 			{
 				response.Tags = new List<CommitTag>(tags);
@@ -235,6 +235,7 @@ namespace HordeServer.Streams
 			{
 				response.Files = new List<string>(files);
 			}
+
 			return response;
 		}
 
diff --git a/Engine/Source/Programs/Shared/EpicGames.Horde/Commits/CommitMessages.cs b/Engine/Source/Programs/Shared/EpicGames.Horde/Commits/CommitMessages.cs
--- a/Engine/Source/Programs/Shared/EpicGames.Horde/Commits/CommitMessages.cs
+++ b/Engine/Source/Programs/Shared/EpicGames.Horde/Commits/CommitMessages.cs
@@ -49,6 +49,11 @@ namespace EpicGames.Horde.Commits
 		/// </summary>
 		public string Description { get; set; }
 
+		/// <summary>
+		/// Date/time that change was committed
+		/// </summary>
+		public DateTime DateUtc{ get; set; }
+
 		/// <summary>
 		/// Tags for this commit
 		/// </summary>
@@ -62,11 +67,12 @@ namespace EpicGames.Horde.Commits
 		/// <summary>
 		/// Constructor
 		/// </summary>
-		public GetCommitResponse(GetThinUserInfoResponse authorInfo, string description)
+		public GetCommitResponse(GetThinUserInfoResponse authorInfo, string description, DateTime dateUtc)
 		{
 			Author = authorInfo.Name;
 			AuthorInfo = authorInfo;
 			Description = description;
+			DateUtc = dateUtc;
 		}
 	}
 }
