[Horde][BuildHealth][Dashboard] Memoization pass for render performance, minor highlighting improvements, and deletion of FluentUI table.

diff --git a/Engine/Source/Programs/Horde/HordeDashboard/src/components/buildhealth/BuildhealthView.tsx b/Engine/Source/Programs/Horde/HordeDashboard/src/components/buildhealth/BuildhealthView.tsx
--- a/Engine/Source/Programs/Horde/HordeDashboard/src/components/buildhealth/BuildhealthView.tsx
+++ b/Engine/Source/Programs/Horde/HordeDashboard/src/components/buildhealth/BuildhealthView.tsx
@@ -1,9 +1,9 @@
 // Copyright Epic Games, Inc. All Rights Reserved.
 
-import { ComboBox, DirectionalHint, Dropdown, DropdownMenuItemType, IDropdownOption, Label, Stack, TooltipHost } from "@fluentui/react";
+import { ComboBox, DirectionalHint, Dropdown, DropdownMenuItemType, IDropdownOption, IDropdownStyles, Label, Slider, Stack, TooltipHost } from "@fluentui/react";
 import { GetProjectResponse, GetStreamResponse } from "horde/backend/Api";
 import { observer } from "mobx-react-lite";
-import { useEffect, useMemo, useState } from "react";
+import { useCallback, useEffect, useMemo, useState } from "react";
 import { useLocation, useNavigate } from "react-router";
 import { useWindowSize } from "../../base/utilities/hooks";
 import { getHordeStyling } from "../../styles/Styles";
@@ -95,114 +95,143 @@ const MultiListParameter: React.FC<{ handler: BuildHealthDataHandler, buildHealt
         }
     });
 
-    const groups = Array.from(gset).sort((a, b) => a.localeCompare(b));
+    // #region -- Data Memos --
 
-    const doptions: IDropdownOption<BuildHealthOptionData>[] = [];
+    const doptions = useMemo(() => {
+        const gset: Set<string> = new Set();
 
-    groups.forEach(group => {
-        if (group !== "__nogroup") {
-            doptions.push({
-                key: `group_${group}`,
-                text: group,
-                itemType: DropdownMenuItemType.Header
-            });
-        }
+        param.items.forEach(item => gset.add(item.group ?? "__nogroup"));
 
-        param.items.forEach(item => {
-            if (item.group === group) {
-                const key = item.id;
-                doptions.push({
-                    key: key,
-                    data: item,
-                    text: item.text
+        const groups = Array.from(gset).sort((a, b) => a.localeCompare(b));
+
+        const options: IDropdownOption<BuildHealthOptionData>[] = [];
+
+        groups.forEach(group => {
+            if (group !== "__nogroup") {
+                options.push({
+                    key: `group_${group}`,
+                    text: group,
+                    itemType: DropdownMenuItemType.Header,
                 });
             }
+            param.items.forEach(item => {
+                if ((item.group ?? "__nogroup") === group) {
+                    options.push({
+                        key: item.id,
+                        text: item.text,
+                        data: item,
+                    });
+                }
+            });
         });
-    });
+
+        return options;
+    }, [param.items]);
+
+    const memoizedSelectedKeys = useMemo(() => selectedKeys, [selectedKeys]);
+
+    // #endregion -- Data Memos --
+
+    // #region -- Style Memos --
+
+    const dropdownStyles: Partial<IDropdownStyles> = useMemo(() => ({
+        dropdown: { width: 300 },
+        callout: {
+            selectors: {
+                [".ms-Callout-main"]: {
+                    padding: "4px 4px 12px 12px",
+                    overflow: "auto" as "auto",
+                },
+            },
+        },
+        dropdownItemHeader: { fontSize: 12, color: modeColors.text },
+        dropdownOptionText: { fontSize: 12 },
+        dropdownItem: {
+            minHeight: 18,
+            lineHeight: 18,
+            selectors: {
+                [".ms-Checkbox-checkbox"]: {
+                    width: 14,
+                    height: 14,
+                    fontSize: 11,
+                },
+            },
+        },
+        dropdownItemSelected: {
+            minHeight: 18,
+            lineHeight: 18,
+            backgroundColor: "inherit",
+            selectors: {
+                [".ms-Checkbox-checkbox"]: {
+                    width: 14,
+                    height: 14,
+                    fontSize: 11,
+                },
+            },
+        },
+    }), [modeColors.text]);
+
+    // #endregion -- Style Memos --
+
+    // #region -- Dropdown Callbacks --
+
+    const renderItem = useCallback((item, defaultRender) => {
+        let castedOption = item as IDropdownOption<BuildHealthOptionData>;
+        if (!castedOption || !defaultRender) return null;
+
+        if (!castedOption.data) {
+            return defaultRender(item);
+        }
+        return (
+            <TooltipHost key={`tooltip_${castedOption.text}`} content={
+                <div>
+                    <Stack>
+                        <Label>{castedOption.data?.group}</Label>
+                    </Stack>
+                    <Stack>
+                        {castedOption.text}
+                    </Stack>
+                </div>
+            }
+                directionalHint={DirectionalHint.rightCenter}
+                calloutProps={{
+                    styles: {
+                        root: { maxWidth: 600 }
+                    }
+                }}
+            >
+                {defaultRender(item)}
+            </TooltipHost>
+        );
+    }, []);
+
+    const handleChange = useCallback((_ev, option) => {
+        const castedOption = option as IDropdownOption<BuildHealthOptionData>;
+        if (!castedOption) return;
+        onChange(castedOption);
+    }, [onChange]);
+
+    const handleDismiss = useCallback(() => {
+        handler.requestHierarchicalRefresh(request);
+        buildHealthOptions.synchronizeDerivedKeys();
+    }, [handler, request, buildHealthOptions]);
+
+    // #endregion -- Dropdown Callbacks --
 
     return (
-        <Dropdown id={changeVersion.toString() + key}
+        <Dropdown
+            id={changeVersion.toString() + key}
             key={key}
             label={param.label}
-            placeholder={"Select options"}
+            placeholder="Select options"
             options={doptions}
-            selectedKeys={selectedKeys}
+            selectedKeys={memoizedSelectedKeys}
             dropdownWidth={400}
-            calloutProps={{
-                directionalHint: DirectionalHint.rightCenter
-            }}
-            styles={{
-                dropdown: { width: 300 },
-                callout: {
-                    selectors: {
-                        ".ms-Callout-main": {
-                            padding: "4px 4px 12px 12px",
-                            overflow: "auto",
-                        }
-                    }
-                },
-                dropdownItemHeader: { fontSize: 12, color: modeColors.text },
-                dropdownOptionText: { fontSize: 12 },
-                dropdownItem: {
-                    minHeight: 18, lineHeight: 18, selectors: {
-                        '.ms-Checkbox-checkbox': {
-                            width: 14,
-                            height: 14,
-                            fontSize: 11
-                        }
-                    }
-                },
-                dropdownItemSelected: {
-                    minHeight: 18, lineHeight: 18, backgroundColor: "inherit",
-                    selectors: {
-                        '.ms-Checkbox-checkbox': {
-                            width: 14,
-                            height: 14,
-                            fontSize: 11
-                        }
-                    }
-                }
-            }}
-            onRenderItem={(item, defaultRender) => {
-                let castedOption = item as IDropdownOption<BuildHealthOptionData>;
-                if (!castedOption || !defaultRender) return null;
-
-                if (!castedOption.data) {
-                    return defaultRender(item);
-                }
-                return (
-                    <TooltipHost key={`tooltip_${castedOption.text}`} content={
-                        <div>
-                            <Stack>
-                                <Label>{castedOption.data?.group}</Label>
-                            </Stack>
-                            <Stack>
-                                {castedOption.text}
-                            </Stack>
-                        </div>
-                    }
-                        directionalHint={DirectionalHint.rightCenter}
-                        calloutProps={{
-                            styles: {
-                                root: { maxWidth: 600 }
-                            }
-                        }}
-                    >
-                        {defaultRender(item)}
-                    </TooltipHost>
-                );
-            }}
-            onChange={(_ev, option) => {
-                let castedOption = option as IDropdownOption<BuildHealthOptionData>;
-                if (!castedOption) {
-                    return;
-                }
-                onChange(castedOption);
-            }}
-            onDismiss={() => {
-                handler.requestHierarchicalRefresh(request);
-                buildHealthOptions.synchronizeDerivedKeys();
-            }}
+            calloutProps={{ directionalHint: DirectionalHint.rightCenter }}
+            styles={dropdownStyles}
+            onRenderItem={renderItem}
+            onChange={handleChange}
+            onDismiss={handleDismiss}
             multiSelect
         />);
 };
diff --git a/Engine/Source/Programs/Horde/HordeDashboard/src/components/buildhealth/stepoutcome/StepOutcomeD3Component.tsx b/Engine/Source/Programs/Horde/HordeDashboard/src/components/buildhealth/stepoutcome/StepOutcomeD3Component.tsx
--- a/Engine/Source/Programs/Horde/HordeDashboard/src/components/buildhealth/stepoutcome/StepOutcomeD3Component.tsx
+++ b/Engine/Source/Programs/Horde/HordeDashboard/src/components/buildhealth/stepoutcome/StepOutcomeD3Component.tsx
@@ -7,6 +7,7 @@ import { useEffect, useRef } from "react";
 import { StepOutcomeDataHandler } from "./StepOutcomeDataHandler";
 import { ChangeHeader, encodeStepNameFromStepNameHeader, INVALID_CHANGELIST, StepNameHeader, StepOutcomeTable, StepOutcomeTableEntry } from "./StepOutcomeDataTypes";
 import { getCellStyle, IncrementalRefreshIndicatorPanel, sharedBorderStyle } from "./StepOutcomeSharedUIComponents";
+import { getHordeTheme } from "horde/styles/theme";
 
 // #region -- Styles --
 
@@ -99,6 +100,7 @@ function GenerateD3Table(stepOutcomeTable: StepOutcomeTable, handler: StepOutcom
 
     const containerRef = useRef<HTMLDivElement>(null);
     const { modeColors, _hordeClasses, _detailClasses } = getHordeStyling();
+    const theme = getHordeTheme();
 
     const scaffold = useRef<{
         wrapper: d3.Selection<HTMLDivElement, unknown, null, undefined>;
@@ -133,6 +135,8 @@ function GenerateD3Table(stepOutcomeTable: StepOutcomeTable, handler: StepOutcom
         const tbody = table.append("tbody");
 
         // Add tool tip div; will be used throughout.
+        d3.selectAll(".d3-tooltip").remove();
+
         const tooltip = d3.select("body")
             .append("div")
             .attr("class", "d3-tooltip")
@@ -173,6 +177,10 @@ function GenerateD3Table(stepOutcomeTable: StepOutcomeTable, handler: StepOutcom
             .style("max-width", "320px");
 
         scaffold.current = { wrapper, table, thead, theaderrow, tbody, tooltip };
+
+        return () => {
+            d3.selectAll(".d3-tooltip").remove();
+        };
     }, []);
 
     // Dependency array use effect for stepOutcomeTable data.
@@ -198,6 +206,7 @@ function GenerateD3Table(stepOutcomeTable: StepOutcomeTable, handler: StepOutcom
             .style("border", `1px solid ${modeColors.background}`)
             .text(d => d.change === INVALID_CHANGELIST ? (d.date ?? "") : d.change)
             .on("mouseover", function (_event, d) {
+                d3.select(this).style("background-color", theme.semanticColors.bodyBackgroundHovered);
                 const entry = d as unknown as ChangeHeader;
 
                 if (entry && entry.change !== INVALID_CHANGELIST) {
@@ -223,6 +232,7 @@ function GenerateD3Table(stepOutcomeTable: StepOutcomeTable, handler: StepOutcom
             })
             .on("mouseout", function () {
                 scaffold.current!.tooltip.style("opacity", 0);
+                d3.select(this).style("background-color", modeColors.content);
             });
 
         const rows = scaffold.current.tbody.selectAll("tr")
@@ -280,6 +290,8 @@ function GenerateD3Table(stepOutcomeTable: StepOutcomeTable, handler: StepOutcom
             tr.selectAll("th.stream").each(function (stepNameHeader: StepNameHeader) {
                 const th = d3.select(this);
                 th.on("mouseover", (event: MouseEvent) => {
+                    d3.select(this).style("background-color", theme.semanticColors.bodyBackgroundHovered);
+
                     const tooltipContent = document.createElement("div");
                     tooltipContent.style.paddingLeft = "4px";
 
@@ -293,6 +305,7 @@ function GenerateD3Table(stepOutcomeTable: StepOutcomeTable, handler: StepOutcom
 
                 th.on("mouseout", () => {
                     scaffold.current!.tooltip.style("opacity", 0);
+                    d3.select(this).style("background-color", modeColors.content);
                 });
             });
 
@@ -301,12 +314,10 @@ function GenerateD3Table(stepOutcomeTable: StepOutcomeTable, handler: StepOutcom
 
         // Accessability row highlights.
         rows.on("mouseover", function () {
-            d3.select(this)
-                .style("background-color", modeColors.content);
+            d3.select(this).style("background-color", theme.semanticColors.bodyBackgroundHovered);
         })
             .on("mouseout", function () {
-                d3.select(this)
-                    .style("background-color", null);
+                d3.select(this).style("background-color", null);
             });
 
         // Add the first "step name" column item.
@@ -332,6 +343,8 @@ function GenerateD3Table(stepOutcomeTable: StepOutcomeTable, handler: StepOutcom
                 .style("box-sizing", "border-box")
                 .text(d => `${d.stepName}`)
                 .on("mouseover", function (_event, stepNameHeader) {
+                    d3.select(this).style("background-color", theme.semanticColors.bodyBackgroundHovered);
+
                     const castedStepNameHeader = stepNameHeader as unknown as StepNameHeader;
                     const tooltipContent = document.createElement("div");
                     tooltipContent.style.paddingLeft = "4px";
@@ -350,6 +363,7 @@ function GenerateD3Table(stepOutcomeTable: StepOutcomeTable, handler: StepOutcom
                 })
                 .on("mouseout", function () {
                     scaffold.current!.tooltip.style("opacity", 0);
+                    d3.select(this).style("background-color", modeColors.content);
                 });
         });
 
diff --git a/Engine/Source/Programs/Horde/HordeDashboard/src/components/buildhealth/stepoutcome/StepOutcomeFluentUIComponent.tsx b/Engine/Source/Programs/Horde/HordeDashboard/src/components/buildhealth/stepoutcome/StepOutcomeFluentUIComponent.tsx
deleted file mode 100644
