Horde Dashboard: Show job description in summary when available

diff --git a/Engine/Source/Programs/Horde/HordeDashboard/src/components/jobDetailsV2/JobDetailViewSummary.tsx b/Engine/Source/Programs/Horde/HordeDashboard/src/components/jobDetailsV2/JobDetailViewSummary.tsx
--- a/Engine/Source/Programs/Horde/HordeDashboard/src/components/jobDetailsV2/JobDetailViewSummary.tsx
+++ b/Engine/Source/Programs/Horde/HordeDashboard/src/components/jobDetailsV2/JobDetailViewSummary.tsx
@@ -1,6 +1,6 @@
 // Copyright Epic Games, Inc. All Rights Reserved.
 
-import { Label, Stack, Text } from '@fluentui/react';
+import { Label, Separator, Stack, Text } from '@fluentui/react';
 import { observer } from 'mobx-react-lite';
 import React, { useEffect } from 'react';
 import { DashboardPreference, JobState, ReportPlacement } from '../../backend/Api';
@@ -85,11 +85,11 @@ export const SummaryPanel: React.FC<{ jobDetails: JobDetailsV2 }> = observer(({
    if (jobDetails.jobError) {
       document.title = "Error Loading Job";
       let errorTextColor: string | undefined = dashboard.getPreference(DashboardPreference.ColorError) ?? "#c44525";
-      
+
       return <Stack style={{ paddingTop: 64, paddingBottom: 82 }}>
-               <Text variant="large" styles={{ root: { fontFamily: "Horde Open Sans SemiBold", padding: "25px 30px 25px", color: errorTextColor } }}>Error loading job, please check network connection.</Text>
-               <Text variant="mediumPlus" styles={{ root: { fontFamily: "Horde Open Sans SemiBold", padding: "10px 30px 25px", color: errorTextColor } }}>{jobDetails.jobError}</Text>
-         </Stack>
+         <Text variant="large" styles={{ root: { fontFamily: "Horde Open Sans SemiBold", padding: "25px 30px 25px", color: errorTextColor } }}>Error loading job, please check network connection.</Text>
+         <Text variant="mediumPlus" styles={{ root: { fontFamily: "Horde Open Sans SemiBold", padding: "10px 30px 25px", color: errorTextColor } }}>{jobDetails.jobError}</Text>
+      </Stack>
    }
 
    if (!jobData) {
@@ -118,9 +118,11 @@ export const SummaryPanel: React.FC<{ jobDetails: JobDetailsV2 }> = observer(({
    const summary = text.join(".  ");
 
    const reportData = jobDetails.getSummaryReport();
-   
+
    const parentJobUrl = jobData.parentJobId && jobData.parentJobId ? `/job/${jobData.parentJobId}?step=${jobData.parentJobStepId}` : undefined;
 
+   const desc = jobDetails.template?.description;
+
    return (<Stack id={sideRail.url} styles={{ root: { paddingTop: 0, paddingRight: 12 } }}>
       <Stack className={hordeClasses.raised} >
          <Stack tokens={{ childrenGap: 12 }} grow>
@@ -142,6 +144,7 @@ export const SummaryPanel: React.FC<{ jobDetails: JobDetailsV2 }> = observer(({
                   {!!price && <Stack>
                      <Text>{`Estimated cost: $${price.toFixed(2)}`}</Text>
                   </Stack>}
+                  {!!desc && <Stack> <Separator /> <Text variant="mediumPlus" styles={{ root: { fontFamily: "Horde Open Sans SemiBold", paddingBottom: "12px" } }}>Description</Text><Markdown>{desc}</Markdown></Stack>}
                </Stack>
             </Stack>
          </Stack>
