Horde Dashboard: Tweak job description display

diff --git a/Engine/Source/Programs/Horde/HordeDashboard/src/components/jobDetailsV2/JobDetailViewSummary.tsx b/Engine/Source/Programs/Horde/HordeDashboard/src/components/jobDetailsV2/JobDetailViewSummary.tsx
--- a/Engine/Source/Programs/Horde/HordeDashboard/src/components/jobDetailsV2/JobDetailViewSummary.tsx
+++ b/Engine/Source/Programs/Horde/HordeDashboard/src/components/jobDetailsV2/JobDetailViewSummary.tsx
@@ -144,7 +144,7 @@ export const SummaryPanel: React.FC<{ jobDetails: JobDetailsV2 }> = observer(({
                   {!!price && <Stack>
                      <Text>{`Estimated cost: $${price.toFixed(2)}`}</Text>
                   </Stack>}
-                  {!!desc && <Stack> <Separator /> <Text variant="mediumPlus" styles={{ root: { fontFamily: "Horde Open Sans SemiBold", paddingBottom: "12px" } }}>Description</Text><Markdown>{desc}</Markdown></Stack>}
+                  {!!desc && <Stack> <Separator /> <Text variant="medium" styles={{ root: { fontFamily: "Horde Open Sans SemiBold", paddingBottom: "12px" } }}>Description</Text><Markdown>{desc}</Markdown></Stack>}
                </Stack>
             </Stack>
          </Stack>
