[HordeDashboard] Enables modifying tags of a reserved device
#tests unit tests

diff --git a/Engine/Source/Programs/Horde/HordeDashboard/src/components/devices/DeviceEditor.tsx b/Engine/Source/Programs/Horde/HordeDashboard/src/components/devices/DeviceEditor.tsx
--- a/Engine/Source/Programs/Horde/HordeDashboard/src/components/devices/DeviceEditor.tsx
+++ b/Engine/Source/Programs/Horde/HordeDashboard/src/components/devices/DeviceEditor.tsx
@@ -300,7 +300,8 @@ export const DeviceEditor: React.FC<{ handler: DeviceHandler, deviceIn?: GetDevi
             if (editNote) {
 
                await backend.modifyDevice(device.id!, {
-                  notes: device.notes ? device.notes.trim() : ""
+                  notes: device.notes ? device.notes.trim() : "",
+                  tags: device.tags ? device.tags : [""]
                }).then(() => {
                   close();
                }).catch((reason) => {
@@ -312,7 +313,8 @@ export const DeviceEditor: React.FC<{ handler: DeviceHandler, deviceIn?: GetDevi
                await backend.modifyDevice(device.id!, {
                   problem: false,
                   maintenance: device.maintenance,
-                  enabled: true
+                  enabled: true,
+                  tags: device.tags ? device.tags : [""]
                }).then(() => {
                   close();
                }).catch((reason) => {
