[Horde][BuildHealth][Dashboard] UI improvements for 5.7:
- Max width for ultra wide monitors
- No data for current filter label
- Title that summarizes the filter

diff --git a/Engine/Source/Programs/Horde/HordeDashboard/src/components/buildhealth/BuildhealthView.tsx b/Engine/Source/Programs/Horde/HordeDashboard/src/components/buildhealth/BuildhealthView.tsx
--- a/Engine/Source/Programs/Horde/HordeDashboard/src/components/buildhealth/BuildhealthView.tsx
+++ b/Engine/Source/Programs/Horde/HordeDashboard/src/components/buildhealth/BuildhealthView.tsx
@@ -14,6 +14,7 @@ import { BuildHealthDataHandler, DataHandlerRefreshRequest, expandRefreshRequest
 import { KEY_SEPARATOR, PARAMETER_KEY_PREFIX, StepRefData, TemplateRefData } from "./BuildHealthDataTypes";
 import { BuildHealthOptionsController, BuildHealthOptionsState, JobHistoryTimeSpans, loadBuildHealthOptionsFromParams, parseBuildHealthQueryParams, TimeSpan } from "./BuildHealthOptions";
 import { encodeStepKey, encodeTemplateKey } from "./BuildHealthUtilities";
+import { stepOutcomeTableClasses } from "./stepoutcome/StepOutcomeSharedUIComponents";
 
 // #region -- Dropdown Components --
 
@@ -547,7 +548,7 @@ export const BuildHealthView: React.FC = observer(function ConstructBuildHealthV
                     <Stack styles={{ root: { width: "100%" } }}>
                         <Stack horizontal styles={{ root: { minHeight: '85vh' } }} >
                             <BuildHealthSidebar handler={handler} buildHealthOptions={buildHealthController} />
-                            <Stack id="parent-buildhealth-stepoutcome" grow style={{ overflowX: "auto", overflowY: "visible", minWidth: "1128px", position: "relative", border: `${modeColors.content} solid 2px` }}>
+                            <Stack id="parent-buildhealth-stepoutcome" className={stepOutcomeTableClasses.root} grow style={{ overflowX: "auto", overflowY: "visible", minWidth: "1200px", position: "relative", border: `${modeColors.content} solid 2px` }}>
                                 <StepOutcomeView filter={filter} />
                             </Stack>
                         </Stack>
diff --git a/Engine/Source/Programs/Horde/HordeDashboard/src/components/buildhealth/stepoutcome/StepOutcome.tsx b/Engine/Source/Programs/Horde/HordeDashboard/src/components/buildhealth/stepoutcome/StepOutcome.tsx
--- a/Engine/Source/Programs/Horde/HordeDashboard/src/components/buildhealth/stepoutcome/StepOutcome.tsx
+++ b/Engine/Source/Programs/Horde/HordeDashboard/src/components/buildhealth/stepoutcome/StepOutcome.tsx
@@ -5,9 +5,9 @@ import { Text } from '@fluentui/react/lib/Text';
 import { observer } from "mobx-react-lite";
 import React, { useEffect, useState } from 'react';
 import { StepOutcomeD3Table } from "./StepOutcomeD3Component";
-import { StepOutcomeDataHandler, StepOutcomeFilters } from "./StepOutcomeDataHandler";
+import { prettyPrintStepOutcomeFilters, prettyPrintStepOutcomeFiltersSummary, StepOutcomeDataHandler, StepOutcomeFilters } from "./StepOutcomeDataHandler";
 import { StepOutcomeTableEntry } from "./StepOutcomeDataTypes";
-import { sharedBorderStyle } from "./StepOutcomeSharedUIComponents";
+import { sharedBorderStyle, stepOutcomeTableClasses } from "./StepOutcomeSharedUIComponents";
 import { toDate } from "./StepOutcomeUtilities";
 
 // #region -- View Components --
@@ -142,19 +142,50 @@ const StepOutcomeFocusModal: React.FC<StepOutcomeFocusModalProps> = ({ item, onD
 /**
  * React Component that acts as a observable binder for data handler updates, and allows for easy swapping of other table types.
  */
-const StepOutcomeTableComponent: React.FC<StepOutcomeTableProps> = observer(({ handler, onCellSelected }) => {
+const StepOutcomeTableComponent: React.FC<StepOutcomeTableProps> = observer(({ handler, filters, onCellSelected }) => {
     handler.updated;
+    const [dataCount, setDataCount] = useState(() => handler.StepOutcomeTableData.dataCount);
+
+    useEffect(() => {
+        setDataCount(handler.StepOutcomeTableData.dataCount);
+    }), [handler.StepOutcomeTableData.dataCount];
+
+    const [header, subHeader, subSubHeader] = prettyPrintStepOutcomeFiltersSummary(filters);
 
     return (
         <Stack>
             {handler.isInFullDataRefresh ? (
-                <Spinner size={SpinnerSize.large} />
+                <Spinner style={{ paddingTop: 16 }} size={SpinnerSize.large} />
             ) :
-                <StepOutcomeD3Table
-                    stepOutcomeTable={handler.StepOutcomeTableData}
-                    handler={handler}
-                    onCellSelected={onCellSelected}
-                />
+                <Stack>
+                    {dataCount > 0 ? (
+                        <Stack>
+                            <Stack horizontalAlign="center" styles={{ root: { paddingTop: 8, paddingBottom: 8 } }}>
+                                <Text styles={{ root: { fontWeight: 'bold' } }}>{header}</Text>
+                                <Text >{subHeader} - {subSubHeader}</Text>
+                            </Stack>
+                            <StepOutcomeD3Table
+                                stepOutcomeTable={handler.StepOutcomeTableData}
+                                handler={handler}
+                                onCellSelected={onCellSelected}
+                            />
+                        </Stack>
+                    ) : (
+                        <Text style={{ display: "block", paddingTop: 16, textAlign: "center" }}>
+                            {!handler.isInitialRefresh ? (
+                                <>
+                                    No results were found with the currently selected filters:
+                                    <br />
+                                    <pre style={{ textAlign: "left", display: "inline-block", marginTop: "8px" }}>
+                                        {prettyPrintStepOutcomeFilters(filters)}
+                                    </pre>
+                                </>
+                            ) : (
+                                <>Please select filters and run a search to see results.</>
+                            )}
+                        </Text>
+                    )}
+                </Stack>
             }
         </Stack>
     );
@@ -189,7 +220,7 @@ export const StepOutcomeView: React.FC<StepOutcomeViewProps> = ({ filter, refres
 
     return (
         <Stack>
-            <Stack>
+            <Stack key="step-outcome-table">
                 {selectedItem && (<StepOutcomeFocusModal item={selectedItem} onDismiss={() => setSelectedItem(null)} />)}
                 <StepOutcomeTableComponent handler={handler} filters={filter} onCellSelected={setSelectedItem} />
             </Stack>
diff --git a/Engine/Source/Programs/Horde/HordeDashboard/src/components/buildhealth/stepoutcome/StepOutcomeD3Component.tsx b/Engine/Source/Programs/Horde/HordeDashboard/src/components/buildhealth/stepoutcome/StepOutcomeD3Component.tsx
--- a/Engine/Source/Programs/Horde/HordeDashboard/src/components/buildhealth/stepoutcome/StepOutcomeD3Component.tsx
+++ b/Engine/Source/Programs/Horde/HordeDashboard/src/components/buildhealth/stepoutcome/StepOutcomeD3Component.tsx
@@ -452,7 +452,7 @@ function GenerateD3Table(stepOutcomeTable: StepOutcomeTable, handler: StepOutcom
             <div id="d3-container"
                 ref={containerRef}
                 style={{
-                    maxHeight: '85vh',
+                    maxHeight: '83vh',
                     overflow: 'auto'
                 }}
             />
diff --git a/Engine/Source/Programs/Horde/HordeDashboard/src/components/buildhealth/stepoutcome/StepOutcomeDataHandler.ts b/Engine/Source/Programs/Horde/HordeDashboard/src/components/buildhealth/stepoutcome/StepOutcomeDataHandler.ts
--- a/Engine/Source/Programs/Horde/HordeDashboard/src/components/buildhealth/stepoutcome/StepOutcomeDataHandler.ts
+++ b/Engine/Source/Programs/Horde/HordeDashboard/src/components/buildhealth/stepoutcome/StepOutcomeDataHandler.ts
@@ -9,6 +9,66 @@ import { action, computed, makeObservable, observable, runInAction } from "mobx"
 
 // #region -- API Data Types & Utilities --
 
+/**
+ * Utility function to create a formatted string of the StepOutcomeFilters object for messaging.
+ * @param filters The filters object to use in creating the string.
+ * @returns A formated string representation of the StepOutcomeFilters.
+ */
+export function prettyPrintStepOutcomeFilters(filters: StepOutcomeFilters): string {
+    const parts: string[] = [];
+
+    parts.push(`Streams: ${filters.streams.join(", ") || "(none)"}`);
+
+    const start = filters.jobHistorySpan.start.toISOString().split("T")[0];
+    const end = filters.jobHistorySpan.end ? filters.jobHistorySpan.end.toISOString().split("T")[0] : "(ongoing)";
+
+    parts.push(`Job History: ${start} → ${end}`);
+
+    if (filters.jobs?.length) {
+        parts.push(`Jobs: ${filters.jobs.join(", ")}`);
+    }
+
+    if (filters.steps?.length) {
+        parts.push(`Steps: ${filters.steps.join(", ")}`);
+    }
+
+    if (filters.maxJobCount !== undefined) {
+        parts.push(`Max Job Count: ${filters.maxJobCount}`);
+    }
+
+    return parts.join("\n");
+}
+
+/**
+ * Utility function to create a formatted string of the StepOutcomeFilters object for table summary.
+ * @param filters The filters object to use in creating the strings.
+ * @returns A three string representation of the filter as a summary: header, sub header, secondary sub header.
+ */
+export function prettyPrintStepOutcomeFiltersSummary(filters: StepOutcomeFilters): [string, string, string] {
+    // Header: Stream(s)
+    const header = `Streams: ${filters.streams.join(", ") || "(none)"}`;
+
+    // Subheader: Job(s) - Steps
+    const jobs = filters.jobs?.length
+        ? (filters.jobs.length > 3 ? "(many)" : filters.jobs.join(", "))
+        : "(all jobs)";
+
+    const steps = filters.steps?.length
+        ? (filters.steps.length > 3 ? "(many)" : filters.steps.join(", "))
+        : "(all steps)";
+
+    const subHeader = `Jobs: ${jobs} - Steps: ${steps}`;
+
+    // Sub-subheader: Last N Days
+    const start = filters.jobHistorySpan.start;
+    const end = filters.jobHistorySpan.end || new Date();
+    const diffTime = Math.abs(end.getTime() - start.getTime());
+    const diffDays = Math.ceil(diffTime / (1000 * 60 * 60 * 24));
+    const subSubHeader = `Last ${diffDays} day${diffDays > 1 ? "s" : ""}`;
+
+    return [header, subHeader, subSubHeader];
+}
+
 /**
  * Data structure that controls the filtering & queries of the Data Handler.
  * @todo UE-309726 - This data structure can be simplified to reduce duplication, and reduce the need for the internal container to decode.
@@ -137,7 +197,7 @@ export class StepOutcomeDataHandler extends PollBase {
             }
         }
 
-        console.info(`Retrieved ${this.jobData.length} results for StepOutcome query.`);
+        console.info(`Retrieved ${this.jobData.length} job instance results for StepOutcome query.`);
 
         this.stepOutcomeTableData.orderTableDataByChange(true);
         this.setUpdated();
@@ -154,12 +214,12 @@ export class StepOutcomeDataHandler extends PollBase {
      * Will populate the provided stepOutcomeTableEntry's issue data field,
      * @param stepOutcomeTableEntry The step outcome table entry to attempt to populate with issue data.
      */
-    async populateStepOutcomeTableEntryIssueData(stepOutcomeTableEntry : StepOutcomeTableEntry, onCompleteCallback? : (boolean) => void) : Promise<void> {
-        let unresolvedIssueQuery : IssueQuery = {jobId: stepOutcomeTableEntry.jobId, stepId : stepOutcomeTableEntry.stepResponse.id, resolved:false}
+    async populateStepOutcomeTableEntryIssueData(stepOutcomeTableEntry: StepOutcomeTableEntry, onCompleteCallback?: (boolean) => void): Promise<void> {
+        let unresolvedIssueQuery: IssueQuery = { jobId: stepOutcomeTableEntry.jobId, stepId: stepOutcomeTableEntry.stepResponse.id, resolved: false }
         const unresolvedIssueResults = await backend.getIssues(unresolvedIssueQuery);
         stepOutcomeTableEntry.issuesData = unresolvedIssueResults;
 
-        let resolvedIssueQuery : IssueQuery = {jobId: stepOutcomeTableEntry.jobId, stepId : stepOutcomeTableEntry.stepResponse.id, resolved:true}
+        let resolvedIssueQuery: IssueQuery = { jobId: stepOutcomeTableEntry.jobId, stepId: stepOutcomeTableEntry.stepResponse.id, resolved: true }
         const resolvedIssueResults = await backend.getIssues(resolvedIssueQuery);
         stepOutcomeTableEntry.issuesData.push(...resolvedIssueResults);
 
@@ -171,6 +231,15 @@ export class StepOutcomeDataHandler extends PollBase {
 
     // #region -- Public API  --
 
+    /**
+     * Gets whether the handler has yet to initiate the initial refresh.
+     * @returns True if the handler has not issued a initial refresh, false otherwise.
+     */
+    @computed
+    get isInitialRefresh(): boolean {
+        return this.intialRefresh;
+    }
+
     /**
      * Gets whether an full data refresh is occurring.
      * @returns True if an full data refresh is occurring, false otherwise.
diff --git a/Engine/Source/Programs/Horde/HordeDashboard/src/components/buildhealth/stepoutcome/StepOutcomeDataTypes.ts b/Engine/Source/Programs/Horde/HordeDashboard/src/components/buildhealth/stepoutcome/StepOutcomeDataTypes.ts
--- a/Engine/Source/Programs/Horde/HordeDashboard/src/components/buildhealth/stepoutcome/StepOutcomeDataTypes.ts
+++ b/Engine/Source/Programs/Horde/HordeDashboard/src/components/buildhealth/stepoutcome/StepOutcomeDataTypes.ts
@@ -120,6 +120,8 @@ export type StepNameHeader = {
 export class StepOutcomeTable {
     // #region -- Public Members --
 
+    dataCount: number = 0;
+
     /**
      * Change column headers.
      */
@@ -167,6 +169,7 @@ export class StepOutcomeTable {
      * Resets the table data. This will remove all data from the container.
      */
     reset(): void {
+        this.dataCount = 0;
         this.tableEntries = [];
         this.stepNameRowHeaders = [];
         this.stepRowLookup.clear();
@@ -225,7 +228,7 @@ export class StepOutcomeTable {
             });
             return;
         }
-
+        this.dataCount++;
         this.tableEntries[insertionRow][insertionCol] = tableEntry;
     }
 
diff --git a/Engine/Source/Programs/Horde/HordeDashboard/src/components/buildhealth/stepoutcome/StepOutcomeSharedUIComponents.tsx b/Engine/Source/Programs/Horde/HordeDashboard/src/components/buildhealth/stepoutcome/StepOutcomeSharedUIComponents.tsx
--- a/Engine/Source/Programs/Horde/HordeDashboard/src/components/buildhealth/stepoutcome/StepOutcomeSharedUIComponents.tsx
+++ b/Engine/Source/Programs/Horde/HordeDashboard/src/components/buildhealth/stepoutcome/StepOutcomeSharedUIComponents.tsx
@@ -1,6 +1,6 @@
 // Copyright Epic Games, Inc. All Rights Reserved.
 
-import { Spinner, SpinnerSize } from "@fluentui/react";
+import { mergeStyleSets, Spinner, SpinnerSize } from "@fluentui/react";
 import dashboard, { StatusColor } from "horde/backend/Dashboard";
 import { observer } from "mobx-react-lite";
 import { StepOutcomeDataHandler } from "./StepOutcomeDataHandler";
@@ -20,6 +20,12 @@ export const sharedBorderStyle = {
     }
 };
 
+export const stepOutcomeTableClasses = mergeStyleSets({
+    root: {
+        maxWidth: "min(85vw, 2400px)",
+    },
+});
+
 // #endregion -- Styles --
 
 /**
