[Horde][BuildHealth][Dashboard] Filter option to prevent preflight data.

diff --git a/Engine/Source/Programs/Horde/HordeDashboard/src/components/buildhealth/BuildHealthDataTypes.ts b/Engine/Source/Programs/Horde/HordeDashboard/src/components/buildhealth/BuildHealthDataTypes.ts
--- a/Engine/Source/Programs/Horde/HordeDashboard/src/components/buildhealth/BuildHealthDataTypes.ts
+++ b/Engine/Source/Programs/Horde/HordeDashboard/src/components/buildhealth/BuildHealthDataTypes.ts
@@ -8,6 +8,7 @@ export const TEMPLATES_URL_PARAM_NAME: string = "templates";
 export const STEPS_URL_PARAM_NAME: string = "stepNames";
 export const PARAMETER_KEY_PREFIX: string = "parameter_key";
 export const TIME_SPAN_URL_PARAM: string = "lastTimeRange";
+export const INCLUDE_PREFLIGHT_PARAM: string = "includePreflight";
 export const KEY_SEPARATOR: string = "::";
 
 /**
diff --git a/Engine/Source/Programs/Horde/HordeDashboard/src/components/buildhealth/BuildHealthOptions.ts b/Engine/Source/Programs/Horde/HordeDashboard/src/components/buildhealth/BuildHealthOptions.ts
--- a/Engine/Source/Programs/Horde/HordeDashboard/src/components/buildhealth/BuildHealthOptions.ts
+++ b/Engine/Source/Programs/Horde/HordeDashboard/src/components/buildhealth/BuildHealthOptions.ts
@@ -1,7 +1,7 @@
 // Copyright Epic Games, Inc. All Rights Reserved.
 
 import { action, makeObservable, observable, runInAction } from "mobx";
-import { PROJECTS_URL_PARAM_NAME, STEPS_URL_PARAM_NAME, STREAMS_URL_PARAM_NAME, TemplateRefData, TEMPLATES_URL_PARAM_NAME, TIME_SPAN_URL_PARAM } from "../buildhealth/BuildHealthDataTypes";
+import { INCLUDE_PREFLIGHT_PARAM, PROJECTS_URL_PARAM_NAME, STEPS_URL_PARAM_NAME, STREAMS_URL_PARAM_NAME, TemplateRefData, TEMPLATES_URL_PARAM_NAME, TIME_SPAN_URL_PARAM } from "../buildhealth/BuildHealthDataTypes";
 import { BuildHealthDataHandler } from "./BuildHealthDataHandler";
 import { decodeStepKey, decodeTemplateKey, encodeStepKey, encodeTemplateKey } from "./BuildHealthUtilities";
 import { Location } from 'react-router-dom';
@@ -19,6 +19,7 @@ type BuildHealthQueryParams = {
     templateIds: string[];
     stepNames: string[];
     timeSpanKey: string;
+    includePreflight: boolean;
 };
 
 /**
@@ -85,12 +86,15 @@ export function parseBuildHealthQueryParams(location: Location): BuildHealthQuer
         const raw = searchParams.get(TIME_SPAN_URL_PARAM);
         return raw ? LZString.decompressFromEncodedURIComponent(raw) ?? "" : "";
     })();
-    
+
     if (!JobHistoryTimeSpans.some(ts => ts.key === timeSpanKey)) {
         timeSpanKey = DEFAULT_JOB_HISTORY_TIME_SPAN.key;
     }
 
-    return { projectIds, streamIds, templateIds: templateIds, stepNames, timeSpanKey };
+    // Read from URLSearchParams
+    const includePreflight = (searchParams.get(INCLUDE_PREFLIGHT_PARAM) ? LZString.decompressFromEncodedURIComponent(searchParams.get(INCLUDE_PREFLIGHT_PARAM)!) : "false") === "true";
+
+    return { projectIds, streamIds, templateIds: templateIds, stepNames, timeSpanKey, includePreflight };
 }
 
 /**
@@ -104,6 +108,10 @@ export async function loadBuildHealthOptionsFromParams(
     options: BuildHealthOptionsController,
     handler: BuildHealthDataHandler) {
 
+    runInAction(() => {
+        options.setIncludePreflight(params.includePreflight);
+    })
+
     if (params.timeSpanKey) {
         runInAction(() => {
             options.setJobHistoryTimeSpan(JobHistoryTimeSpans.find(ts => ts.key === params.timeSpanKey)!);
@@ -201,6 +209,7 @@ export class BuildHealthOptionsState {
     @observable enabledTemplates: Record<string, string> = {};
     @observable enabledSteps: Record<string, string> = {};
     @observable jobHistoryTimeSpan: TimeSpan = DEFAULT_JOB_HISTORY_TIME_SPAN;
+    @observable includePreflight: boolean = false;
 
     // #endregion -- User Selected Options --
 
@@ -352,6 +361,16 @@ export class BuildHealthOptionsController {
         this.setOptionsChanged();
     }
 
+    /**
+     * Sets whether to filter out preflight jobs, or not.
+     * @param includePreflight True to include preflight job instances, false otherwise.
+     */
+    @action
+    setIncludePreflight(includePreflight: boolean) {
+        this.state.includePreflight = includePreflight;
+        this.setOptionsChanged();
+    }
+
     /**
       * Clear all the projects tracked by the options.
       * @param cascade Whether to cascade this down to dependent options.
@@ -431,6 +450,7 @@ export class BuildHealthOptionsController {
         const jobs = this.state.stepOutcomeEnabledJobKeys.join(",");
         const steps = this.state.stepOutcomeEnabledStepKeys.join(",");
         const lastTimeRange = this.state.jobHistoryTimeSpan.key.toString();
+        const includePreflight = this.state.includePreflight;
 
         if (!projects.length) return params;
 
@@ -439,6 +459,7 @@ export class BuildHealthOptionsController {
         params.set(TEMPLATES_URL_PARAM_NAME, LZString.compressToEncodedURIComponent(jobs));
         params.set(STEPS_URL_PARAM_NAME, LZString.compressToEncodedURIComponent(steps));
         params.set(TIME_SPAN_URL_PARAM, LZString.compressToEncodedURIComponent(lastTimeRange));
+        params.set(INCLUDE_PREFLIGHT_PARAM, LZString.compressToEncodedURIComponent(String(includePreflight)));
 
         return params;
     }
diff --git a/Engine/Source/Programs/Horde/HordeDashboard/src/components/buildhealth/BuildhealthView.tsx b/Engine/Source/Programs/Horde/HordeDashboard/src/components/buildhealth/BuildhealthView.tsx
--- a/Engine/Source/Programs/Horde/HordeDashboard/src/components/buildhealth/BuildhealthView.tsx
+++ b/Engine/Source/Programs/Horde/HordeDashboard/src/components/buildhealth/BuildhealthView.tsx
@@ -1,6 +1,6 @@
 // Copyright Epic Games, Inc. All Rights Reserved.
 
-import { ComboBox, DirectionalHint, Dropdown, DropdownMenuItemType, IDropdownOption, IDropdownStyles, Label, Slider, Stack, TooltipHost } from "@fluentui/react";
+import { Checkbox, ComboBox, DirectionalHint, Dropdown, DropdownMenuItemType, IDropdownOption, IDropdownStyles, Label, Slider, Stack, TooltipHost } from "@fluentui/react";
 import { GetProjectResponse, GetStreamResponse } from "horde/backend/Api";
 import { observer } from "mobx-react-lite";
 import { useCallback, useEffect, useMemo, useState } from "react";
@@ -473,6 +473,7 @@ const BuildHealthUrlSync: React.FC<{ buildHealthController: BuildHealthOptionsCo
 const BuildHealthSidebar: React.FC<{ handler: BuildHealthDataHandler, buildHealthOptions: BuildHealthOptionsController }> = function ConstrucBuildHealthSidebar({ handler, buildHealthOptions }) {
     const { hordeClasses } = getHordeStyling();
     const [timeSpanKey, setTimeSpanKey] = useState(() => buildHealthOptions.state.jobHistoryTimeSpan.key);
+    const [includePreflight, setIncludePreflight] = useState(() => buildHealthOptions.state.includePreflight)
 
     const updateJobHistoryTimeSpan = (_ev, option) => {
         const select = option as TimeSpan;
@@ -480,10 +481,19 @@ const BuildHealthSidebar: React.FC<{ handler: BuildHealthDataHandler, buildHealt
         buildHealthOptions.setJobHistoryTimeSpan(select);
     };
 
+    const updateIncludePreflight = (isChecked) => {
+        setIncludePreflight(isChecked);
+        buildHealthOptions.setIncludePreflight(isChecked);
+    }
+
     useEffect(() => {
         setTimeSpanKey(buildHealthOptions.state.jobHistoryTimeSpan.key);
     }, [buildHealthOptions.state.jobHistoryTimeSpan.key]);
 
+    useEffect(() => {
+        setIncludePreflight(buildHealthOptions.state.includePreflight);
+    }, [buildHealthOptions.state.includePreflight]);
+
     let timeComboWidth = 180;
 
     return (
@@ -510,6 +520,13 @@ const BuildHealthSidebar: React.FC<{ handler: BuildHealthDataHandler, buildHealt
                                 buildHealthOptions.synchronizeDerivedKeys();
                             }}
                         />
+                        <TooltipHost content="Include preflights in the job history results, or not.">
+                            <Checkbox
+                                label="Include Preflights"
+                                checked={includePreflight}
+                                onChange={(_, isChecked) => updateIncludePreflight(!!isChecked)}
+                            />
+                        </TooltipHost>
                     </Stack>
                 </Stack>
             </Stack>
@@ -529,12 +546,14 @@ export const BuildHealthView: React.FC = observer(function ConstructBuildHealthV
         streams: buildHealthState.stepOutcomeEnabledStreamKeys,
         jobs: buildHealthState.stepOutcomeEnabledJobKeys,
         steps: buildHealthState.stepOutcomeEnabledStepKeys,
-        jobHistorySpan: { start: buildHealthState.startDate, end: buildHealthState.endDate }
+        jobHistorySpan: { start: buildHealthState.startDate, end: buildHealthState.endDate },
+        includePreflights: buildHealthState.includePreflight
     }), [
         buildHealthState.stepOutcomeEnabledStreamKeys,
         buildHealthState.stepOutcomeEnabledJobKeys,
         buildHealthState.stepOutcomeEnabledStepKeys,
-        { start: buildHealthState.startDate, end: buildHealthState.endDate }
+        { start: buildHealthState.startDate, end: buildHealthState.endDate },
+        buildHealthState.includePreflight
     ]);
 
     return (
diff --git a/Engine/Source/Programs/Horde/HordeDashboard/src/components/buildhealth/stepoutcome/StepOutcomeDataHandler.ts b/Engine/Source/Programs/Horde/HordeDashboard/src/components/buildhealth/stepoutcome/StepOutcomeDataHandler.ts
--- a/Engine/Source/Programs/Horde/HordeDashboard/src/components/buildhealth/stepoutcome/StepOutcomeDataHandler.ts
+++ b/Engine/Source/Programs/Horde/HordeDashboard/src/components/buildhealth/stepoutcome/StepOutcomeDataHandler.ts
@@ -79,6 +79,7 @@ export interface StepOutcomeFilters {
     jobs?: string[];
     steps?: string[];
     maxJobCount?: number;
+    includePreflights: boolean;
 }
 
 // #endregion -- API Data Types & Utilities --
@@ -88,7 +89,7 @@ export interface StepOutcomeFilters {
  * @todo Data sharing & caching would ideally occur between the BuildHealth View & Step Outcome.
  */
 export class StepOutcomeDataHandler extends PollBase {
-    private readonly STREAM_JOBS_FILTER = "id,streamId,name,change,templateId,state,createTime,updateTime,batches,preflightDescription,cancellationReason";
+    private readonly STREAM_JOBS_FILTER = "id,streamId,name,change,templateId,state,createTime,updateTime,batches,preflightCommitId,preflightDescription,cancellationReason";
     private readonly MAX_JOB_COUNT_DEFAULT = 150;
 
     // #region -- Private Members -- 
@@ -181,6 +182,12 @@ export class StepOutcomeDataHandler extends PollBase {
         for (const streamResult of results) {
             // Process all job responses
             for (const jobResponse of streamResult) {
+                // Filter preflights as appropriate
+                if(!this.activeFilter.includePreflights && (jobResponse.preflightChange || jobResponse.preflightDescription ))
+                {
+                    continue;
+                }
+
                 const batchResponses = jobResponse.batches ?? [];
                 for (const batch of batchResponses) {
                     for (const step of batch.steps) {
