[Horde][BuildHealth][Dashboard] Topnav to have build health link.

diff --git a/Engine/Source/Programs/Horde/HordeDashboard/src/components/TopNav.tsx b/Engine/Source/Programs/Horde/HordeDashboard/src/components/TopNav.tsx
--- a/Engine/Source/Programs/Horde/HordeDashboard/src/components/TopNav.tsx
+++ b/Engine/Source/Programs/Horde/HordeDashboard/src/components/TopNav.tsx
@@ -490,6 +490,12 @@ export const TopNav: React.FC<{ suppressServer?: boolean }> = observer(({ suppre
          },
       });
 
+      serviceItems.push({
+         key: "build_health_dashboard",
+         text: "Build Health",
+         link: "/buildhealth"
+      });
+
       serviceItems = serviceItems.sort((a, b) => {
          return (a.text ?? "").localeCompare(b.text ?? "")
       })
