Horde: Improve error message design for job view errors

diff --git a/Engine/Source/Programs/Horde/HordeDashboard/src/components/error/ErrorView.tsx b/Engine/Source/Programs/Horde/HordeDashboard/src/components/error/ErrorView.tsx
new file mode 100644
--- /dev/null
+++ b/Engine/Source/Programs/Horde/HordeDashboard/src/components/error/ErrorView.tsx
@@ -0,0 +1,35 @@
+// Copyright Epic Games, Inc. All Rights Reserved.
+
+import React from 'react';
+import dashboard from '../../backend/Dashboard';
+import { getHordeStyling } from "horde/styles/Styles.tsx";
+import { DashboardPreference } from "horde/backend/Api.ts";
+
+export const ErrorView: React.FC<{ title: string, description: string, reason: string }> = ({title, description, reason}) => {
+    const {modeColors} = getHordeStyling();
+    const errorColor = dashboard.getPreference(DashboardPreference.ColorError) ?? "#c44525";
+    return (
+        <div style={{
+            backgroundColor: modeColors.background,
+            padding: '60px 0',
+            display: 'flex',
+            color: errorColor,
+            flexDirection: 'column',
+            justifyContent: 'center',
+            alignItems: 'center'
+        }}>
+            <div style={{
+                maxWidth: 600,
+                width: "90%",
+                padding: '0px 15px',
+                backgroundColor: `color-mix(in srgb, ${errorColor} 10%, white)`,
+                border: "1px solid " + errorColor
+            }}>
+                <h1 style={{border: 1, marginBottom: '8px'}}>{title}</h1>
+                <p>{description}</p>
+                <h2 style={{fontSize: 14, margin: '0 0'}}>Reason</h2>
+                <p style={{marginTop: 5, wordBreak: 'break-word'}}>{reason}</p>
+            </div>
+        </div>
+    );
+}
diff --git a/Engine/Source/Programs/Horde/HordeDashboard/src/components/jobDetailsV2/JobDetailViewV2.tsx b/Engine/Source/Programs/Horde/HordeDashboard/src/components/jobDetailsV2/JobDetailViewV2.tsx
--- a/Engine/Source/Programs/Horde/HordeDashboard/src/components/jobDetailsV2/JobDetailViewV2.tsx
+++ b/Engine/Source/Programs/Horde/HordeDashboard/src/components/jobDetailsV2/JobDetailViewV2.tsx
@@ -21,6 +21,7 @@ import { SummaryPanel } from "./JobDetailViewSummary";
 import { JobOperations } from "./JobOperationsBar";
 import { StepDetailView } from "./StepDetailView";
 import userInactivity from "horde/backend/UserInactivity";
+import { ErrorView } from "horde/components/error/ErrorView.tsx";
 
 class BreadcrumbDataView extends JobDataView {
 
@@ -342,7 +343,7 @@ const ScrollRestore: React.FC<{ jobDetails: JobDetailsV2, scrollRef: React.RefOb
 }
 
 
-const DetailsView: React.FC<{ jobDetails: JobDetailsV2 }> = ({ jobDetails }) => {
+const DetailsView: React.FC<{ jobDetails: JobDetailsV2 }> = observer(({ jobDetails }) => {
 
    const windowSize = useWindowSize();
    const scrollRef = useRef<HTMLDivElement>(null);
@@ -352,6 +353,10 @@ const DetailsView: React.FC<{ jobDetails: JobDetailsV2 }> = ({ jobDetails }) =>
 
    const details = jobDetails;
 
+   if (jobDetails.updated) { } // Reference an observable value to ensure component is reloaded
+   if (jobDetails?.jobError) {
+      return <ErrorView title="Unable to load job" description="Verify job and stream exists and potentially your internet connection." reason={jobDetails.jobError} />;
+   }
 
    const centerAlign = vw / 2 - 720 /*890*/;
    const key = `windowsize_jobdetail_view_${windowSize.width}_${windowSize.height}`;
@@ -391,7 +396,7 @@ const DetailsView: React.FC<{ jobDetails: JobDetailsV2 }> = ({ jobDetails }) =>
          </Stack>
       </Stack>
    </Stack>
-};
+});
 
 export const JobDetailViewV2: React.FC = () => {
 
