Fix missing test and phase session tag ids

#rnx

diff --git a/Engine/Source/Programs/Horde/Plugins/Build/HordeServer.Build.Tests/Jobs/TestData/TestDataTests.cs b/Engine/Source/Programs/Horde/Plugins/Build/HordeServer.Build.Tests/Jobs/TestData/TestDataTests.cs
--- a/Engine/Source/Programs/Horde/Plugins/Build/HordeServer.Build.Tests/Jobs/TestData/TestDataTests.cs
+++ b/Engine/Source/Programs/Horde/Plugins/Build/HordeServer.Build.Tests/Jobs/TestData/TestDataTests.cs
@@ -436,6 +436,7 @@ namespace HordeServer.Tests.Jobs.TestData
 			List<GetTestSessionResponse> sessions = testSessionResult.Value;
 			Assert.AreEqual(6, sessions.Count);
 			Assert.AreEqual(TestOutcome.Failure, sessions[0].Outcome);
+			Assert.IsTrue(sessions[0].TagIds?.All(id => tags.Any(t => t.Id == id)));
 
 			GetPhaseSessionsRequest phaseSessionRequest = new GetPhaseSessionsRequest() { StreamIds = streamIds.ToList(), PhaseIds = tests.SelectMany(t => t.Phases.Select(p => p.Id)).ToList() };
 			ActionResult<List<GetTestPhaseSessionResponse>> phaseSessionResult = await TestDataController.GetTestPhaseSessionsV2Async(phaseSessionRequest);
@@ -444,6 +445,7 @@ namespace HordeServer.Tests.Jobs.TestData
 			List<GetTestPhaseSessionResponse> phaseSessions = phaseSessionResult.Value;
 			Assert.AreEqual(15, phaseSessions.Count);
 			Assert.AreEqual(TestPhaseOutcome.Failed, phaseSessions.First(p => p.PhaseRef == phaseId).Outcome);
+			Assert.IsTrue(phaseSessions.Any(s => s.TagIds?.All(id => tags.Any(t => t.Id == id)) ?? false));
 
 			// Get the sessions for one phase in one stream
 			phaseSessionRequest = new GetPhaseSessionsRequest() { StreamIds = [streamIds[0]], PhaseIds = [tests[0].Phases[0].Id] };
@@ -853,6 +855,7 @@ namespace HordeServer.Tests.Jobs.TestData
 			@"                        ""dateTime"": ""2025-03-03T18:08:08"",",
 			@"                        ""timeElapseSec"": 0.008163903,",
 			@"                        ""outcome"": ""Success"",",
+			@"                        ""tags"": [""Group""],",
 			@"                        ""deviceKeys"": [",
 			@"                            ""9AF6B8CC484096D351EC56ABDFBB919F""",
 			@"                        ]",
@@ -864,6 +867,7 @@ namespace HordeServer.Tests.Jobs.TestData
 			@"                        ""timeElapseSec"": 0.1961886,",
 			@"                        ""outcome"": ""Failed"",",
 			@"                        ""errorFingerprint"": ""a6efc189"",",
+			@"                        ""tags"": [""Engine""],",
 			@"                        ""deviceKeys"": [",
 			@"                            ""9AF6B8CC484096D351EC56ABDFBB919F""",
 			@"                        ],",
@@ -934,6 +938,7 @@ namespace HordeServer.Tests.Jobs.TestData
 			@"                        ""dateTime"": ""2025-03-03T18:08:08"",",
 			@"                        ""timeElapseSec"": 0.008163903,",
 			@"                        ""outcome"": ""Success"",",
+			@"                        ""tags"": [""Parser""],",
 			@"                        ""deviceKeys"": [",
 			@"                            ""9AF6B8CC484096D351EC56ABDFBB919F""",
 			@"                        ]",
@@ -945,6 +950,7 @@ namespace HordeServer.Tests.Jobs.TestData
 			@"                        ""timeElapseSec"": 0.1961886,",
 			@"                        ""outcome"": ""Failed"",",
 			@"                        ""errorFingerprint"": ""a6efc189"",",
+			@"                        ""tags"": [""Engine""],",
 			@"                        ""deviceKeys"": [",
 			@"                            ""9AF6B8CC484096D351EC56ABDFBB919F""",
 			@"                        ],",
diff --git a/Engine/Source/Programs/Horde/Plugins/Build/HordeServer.Build/Jobs/TestData/ITestData.cs b/Engine/Source/Programs/Horde/Plugins/Build/HordeServer.Build/Jobs/TestData/ITestData.cs
--- a/Engine/Source/Programs/Horde/Plugins/Build/HordeServer.Build/Jobs/TestData/ITestData.cs
+++ b/Engine/Source/Programs/Horde/Plugins/Build/HordeServer.Build/Jobs/TestData/ITestData.cs
@@ -519,6 +519,11 @@ namespace HordeServer.Jobs.TestData
 		/// </summary>
 		string? ErrorFingerprint { get; }
 
+		/// <summary>
+		/// The tags associated with the phase session
+		/// </summary>
+		IReadOnlyList<TestTagId>? Tags { get; }
+
 		/// <summary>
 		/// The associated stream
 		/// </summary>
diff --git a/Engine/Source/Programs/Horde/Plugins/Build/HordeServer.Build/Jobs/TestData/TestDataCollectionV2.cs b/Engine/Source/Programs/Horde/Plugins/Build/HordeServer.Build/Jobs/TestData/TestDataCollectionV2.cs
--- a/Engine/Source/Programs/Horde/Plugins/Build/HordeServer.Build/Jobs/TestData/TestDataCollectionV2.cs
+++ b/Engine/Source/Programs/Horde/Plugins/Build/HordeServer.Build/Jobs/TestData/TestDataCollectionV2.cs
@@ -339,6 +339,9 @@ namespace HordeServer.Jobs.TestData
 			[BsonIgnoreIfNull, BsonElement("ef")]
 			public string? ErrorFingerprint { get; set; }
 
+			[BsonIgnoreIfNull, BsonElement("tg")]
+			public IReadOnlyList<TestTagId>? Tags { get; set; }
+
 			[BsonRequired, BsonElement("jid")]
 			public JobId JobId { get; set; }
 
@@ -744,6 +747,8 @@ namespace HordeServer.Jobs.TestData
 				public bool? HasWarning { get; set; }
 				[BsonIgnoreIfNull, BsonElement("eventStreamPath")]
 				public string? EventStreamPath { get; set; }
+				[BsonIgnoreIfNull, BsonElement("tags")]
+				public List<string>? Tags { get; set; }
 			}
 
 			[BsonRequired, BsonElement("summary")]
@@ -808,6 +813,8 @@ namespace HordeServer.Jobs.TestData
 					sessionDoc.StartDateTime = startDateTime;
 				}
 				sessionDoc.Duration = TimeSpan.FromSeconds(testSession.Summary.TimeElapseSec);
+				// tags
+				sessionDoc.Tags = tags?.Select(t => t.Id).ToList();
 
 				// phase session
 				List<TestPhaseSessionDocument> phaseSessions = new();
@@ -831,6 +838,8 @@ namespace HordeServer.Jobs.TestData
 						phaseDoc.StartDateTime = startDateTime;
 					}
 					phaseDoc.Duration = TimeSpan.FromSeconds(phase.TimeElapseSec);
+					// tags
+					phaseDoc.Tags = tags?.Where(tag => phase.Tags?.Contains(tag.Name) ?? false).Select(tag => tag.Id).ToList();
 
 					phaseSessions.Add(phaseDoc);
 				}
diff --git a/Engine/Source/Programs/Horde/Plugins/Build/HordeServer.Build/Jobs/TestData/TestDataController.cs b/Engine/Source/Programs/Horde/Plugins/Build/HordeServer.Build/Jobs/TestData/TestDataController.cs
--- a/Engine/Source/Programs/Horde/Plugins/Build/HordeServer.Build/Jobs/TestData/TestDataController.cs
+++ b/Engine/Source/Programs/Horde/Plugins/Build/HordeServer.Build/Jobs/TestData/TestDataController.cs
@@ -854,6 +854,7 @@ namespace HordeServer.Jobs.TestData
 					SessionId = testData.SessionId.ToString(),
 					Outcome = testData.Outcome,
 					HasWarning = testData.HasWarning,
+					TagIds = testData.Tags?.Select(t => t.ToString()).ToArray(),
 					ErrorFingerprint = testData.ErrorFingerprint,
 					EventStreamPath = testData.EventStreamPath
 				});
diff --git a/Engine/Source/Programs/Shared/EpicGames.Horde/Jobs/TestData/TestDataMessages.cs b/Engine/Source/Programs/Shared/EpicGames.Horde/Jobs/TestData/TestDataMessages.cs
--- a/Engine/Source/Programs/Shared/EpicGames.Horde/Jobs/TestData/TestDataMessages.cs
+++ b/Engine/Source/Programs/Shared/EpicGames.Horde/Jobs/TestData/TestDataMessages.cs
@@ -825,6 +825,11 @@ namespace EpicGames.Horde.Jobs.TestData
 		/// </summary>
 		public string SessionId { get; set; } = String.Empty;
 
+		/// <summary>
+		/// The list of tag ids associated with the phase session
+		/// </summary>
+		public string[]? TagIds { get; set; }
+
 		/// <summary>
 		/// The commit id associated with the phase session
 		/// </summary>
