[Horde][BuildHealth][Dashboard] Minor UI enhancements for 5.7:

- Disable date anchors in headers for preflights. Changelist <> Dates can fall out of sync easily, causing date header anchor instability.

diff --git a/Engine/Source/Programs/Horde/HordeDashboard/src/components/buildhealth/BuildhealthView.tsx b/Engine/Source/Programs/Horde/HordeDashboard/src/components/buildhealth/BuildhealthView.tsx
--- a/Engine/Source/Programs/Horde/HordeDashboard/src/components/buildhealth/BuildhealthView.tsx
+++ b/Engine/Source/Programs/Horde/HordeDashboard/src/components/buildhealth/BuildhealthView.tsx
@@ -520,7 +520,17 @@ const BuildHealthSidebar: React.FC<{ handler: BuildHealthDataHandler, buildHealt
                                 buildHealthOptions.synchronizeDerivedKeys();
                             }}
                         />
-                        <TooltipHost content="Include preflights in the job history results, or not.">
+                        {/* @todo - UE-315881 will remove the need for this constraint. */}
+                        <TooltipHost content={
+                            <div style={{ maxWidth: 300 }}>
+                                <div>
+                                    Include preflights in the job history results, or not.
+                                    <p>
+                                        <i>Note: Including preflights disables date anchors in header.</i>
+                                    </p>
+                                </div>
+                            </div>
+                        }>
                             <Checkbox
                                 label="Include Preflights"
                                 checked={includePreflight}
diff --git a/Engine/Source/Programs/Horde/HordeDashboard/src/components/buildhealth/stepoutcome/StepOutcomeDataHandler.ts b/Engine/Source/Programs/Horde/HordeDashboard/src/components/buildhealth/stepoutcome/StepOutcomeDataHandler.ts
--- a/Engine/Source/Programs/Horde/HordeDashboard/src/components/buildhealth/stepoutcome/StepOutcomeDataHandler.ts
+++ b/Engine/Source/Programs/Horde/HordeDashboard/src/components/buildhealth/stepoutcome/StepOutcomeDataHandler.ts
@@ -207,6 +207,10 @@ export class StepOutcomeDataHandler extends PollBase {
         console.info(`Retrieved ${this.jobData.length} job instance results for StepOutcome query.`);
 
         this.stepOutcomeTableData.orderTableDataByChange(true);
+        
+        // @todo - UE-315881 will remove the need for this constraint.
+        this.StepOutcomeTableData.supportsDateAnchoredChanges = !this.activeFilter.includePreflights;
+        
         this.setUpdated();
 
         runInAction(() => {
diff --git a/Engine/Source/Programs/Horde/HordeDashboard/src/components/buildhealth/stepoutcome/StepOutcomeDataTypes.ts b/Engine/Source/Programs/Horde/HordeDashboard/src/components/buildhealth/stepoutcome/StepOutcomeDataTypes.ts
--- a/Engine/Source/Programs/Horde/HordeDashboard/src/components/buildhealth/stepoutcome/StepOutcomeDataTypes.ts
+++ b/Engine/Source/Programs/Horde/HordeDashboard/src/components/buildhealth/stepoutcome/StepOutcomeDataTypes.ts
@@ -122,6 +122,12 @@ export class StepOutcomeTable {
 
     dataCount: number = 0;
 
+    /**
+     * Sets whether @see getDateAnchoredChanges is supported for the given data.
+     * @todo - UE-315881 will remove the need for this constraint.
+     */
+    supportsDateAnchoredChanges: boolean = true;
+
     /**
      * Change column headers.
      */
@@ -312,6 +318,14 @@ export class StepOutcomeTable {
      * @returns All of the @see ChangeHeader items separated by entries that represent dates. These will have @see INVALID_CHANGELIST as their change property.
      */
     *getDateAnchoredChanges(): Generator<ChangeHeader> {
+        if (!this.supportsDateAnchoredChanges) {
+            for (const header of this.changeColHeaders) {
+                yield header;
+            }
+
+            return;
+        }
+
         let prevDay: string | null = null;
         let prevChangelist: string | null = null;
         for (let i: number = 0; i < this.changeColHeaders.length; ++i) {
