[Horde][ReleaseNotes] Updated Release Notes for [REDACTED_PHONE] deploy

diff --git a/Engine/Source/Programs/Horde/Docs/ReleaseNotes.md b/Engine/Source/Programs/Horde/Docs/ReleaseNotes.md
--- a/Engine/Source/Programs/Horde/Docs/ReleaseNotes.md
+++ b/Engine/Source/Programs/Horde/Docs/ReleaseNotes.md
@@ -1,5 +1,37 @@
 # Release Notes
 
+## 2025-08-28
+
+* [Horde][BuildHealth][Dashboard] Minor UI enhancements for 5.7:
+  - Disable date anchors in headers for preflights. Changelist <> Dates can fall out of sync easily, causing date header anchor instability. (45280879)
+* Improve error message design for job view errors (45270337)
+* [Horde][BuildHealth][Dashboard] Topnav to have build health link. (45255816)
+* [Horde][BuildHealth][Dashboard] Filter option to prevent preflight data. (45255787)
+* [Horde][BuildHealth][Dashboard] UI improvements for 5.7:
+  - Max width for ultra wide monitors
+  - No data for current filter label
+  - Title that summarizes the filter (45255627)
+* [HordeDashboard] Enables modifying tags of a reserved device (45241906)
+* Horde Dashboard: Tweak job description display (45239969)
+* Horde Dashboard: Show job description in summary when available (45239842)
+* [Horde][BuildHealth][Dashboard] Memoization pass for render performance, minor highlighting improvements, and deletion of FluentUI table. (45238473)
+* [Horde][BuildHealth][Server] Expose DateUtc on GetCommitResponse (45237652)
+* Remove unused Password property from EpicGames.Horde.Streams.IWorkspaceType (45228808)
+* Horde Dashboard: Reduce JobDetails V1 QA Fixes (45220314)
+* Guard against multiple hang detection in same step and fix string decoding on signal (45203470)
+* Temporary fix for hung job steps on farm, catch log signal that outputs in cases where UAT or subprocesses are hanging. (45198954)
+* [Horde][BuildHleath][Dashboard] D3 StepOutcomeTable variant for improved render performance. (45197220)
+* Horde Dashboard: Reduce JobDetails V1 (45195387)
+* [Horde] Add GET endpoints for fetching devices to the Horde HTTP client (45127432)
+* Added issue audit logging to new World Leak issue handler (45116295)
+* Horde Dashboard: Add Epic copyright notice to files missing it (45110422)
+* Renamed Type structured logging field to avoid overlap with audit log channel. (45106576)
+* Pass down the interactive boolean to the request (45105852)
+* EpicGames.Perforce: Add more metrics to ManagedWorkspace (45101670)
+* [Horde][BuildHealth][Dashboard] Minor UI enhancements for:
+  - Drop downs in Build Health View
+  - Step Outcome Table tool-tips (45070406)
+
 ## 2025-08-21
 
 * [Horde] Adds a PUT function for updating devices to the HordeHttpClient (45045274)
