[Horde][BuildHealth][Dashboard] Minor UI enhancements for 5.7:
- Last refresh time
- Snap to right (most recent date)
- Add a 2 week option for time filter

diff --git a/Engine/Source/Programs/Horde/HordeDashboard/src/components/buildhealth/BuildHealthOptions.ts b/Engine/Source/Programs/Horde/HordeDashboard/src/components/buildhealth/BuildHealthOptions.ts
--- a/Engine/Source/Programs/Horde/HordeDashboard/src/components/buildhealth/BuildHealthOptions.ts
+++ b/Engine/Source/Programs/Horde/HordeDashboard/src/components/buildhealth/BuildHealthOptions.ts
@@ -1,12 +1,11 @@
 // Copyright Epic Games, Inc. All Rights Reserved.
 
 import { action, makeObservable, observable, runInAction } from "mobx";
-import { INCLUDE_PREFLIGHT_PARAM, PROJECTS_URL_PARAM_NAME, STEPS_URL_PARAM_NAME, STREAMS_URL_PARAM_NAME, TemplateRefData, TEMPLATES_URL_PARAM_NAME, TIME_SPAN_URL_PARAM } from "../buildhealth/BuildHealthDataTypes";
+import { INCLUDE_PREFLIGHT_PARAM, PROJECTS_URL_PARAM_NAME, STEPS_URL_PARAM_NAME, STREAMS_URL_PARAM_NAME, TEMPLATES_URL_PARAM_NAME, TIME_SPAN_URL_PARAM } from "../buildhealth/BuildHealthDataTypes";
 import { BuildHealthDataHandler } from "./BuildHealthDataHandler";
 import { decodeStepKey, decodeTemplateKey, encodeStepKey, encodeTemplateKey } from "./BuildHealthUtilities";
 import { Location } from 'react-router-dom';
 import LZString from "lz-string"
-import { useQuery } from "horde/base/utilities/hooks";
 
 // #region -- Options Utilities & Data Types --
 
@@ -52,6 +51,9 @@ export const JobHistoryTimeSpans: TimeSpan[] = [
     },
     {
         text: "Past 1 Week", key: "time_1_week", minutes: 60 * 24 * 7
+    },
+        {
+        text: "Past 2 Weeks", key: "time_2_week2", minutes: 60 * 24 * 7 * 2
     },
     {
         text: "Past 1 Month", key: "time_1_month", minutes: 60 * 24 * 31
diff --git a/Engine/Source/Programs/Horde/HordeDashboard/src/components/buildhealth/stepoutcome/StepOutcome.tsx b/Engine/Source/Programs/Horde/HordeDashboard/src/components/buildhealth/stepoutcome/StepOutcome.tsx
--- a/Engine/Source/Programs/Horde/HordeDashboard/src/components/buildhealth/stepoutcome/StepOutcome.tsx
+++ b/Engine/Source/Programs/Horde/HordeDashboard/src/components/buildhealth/stepoutcome/StepOutcome.tsx
@@ -172,7 +172,7 @@ const StepOutcomeTableComponent: React.FC<StepOutcomeTableProps> = observer(({ h
                         </Stack>
                     ) : (
                         <Text style={{ display: "block", paddingTop: 16, textAlign: "center" }}>
-                            {!handler.isInitialRefresh ? (
+                            {!handler.isInitialRefresh && !handler.isInFullDataRefresh && !handler.isInIncrementalDataRefresh ? (
                                 <>
                                     No results were found with the currently selected filters:
                                     <br />
diff --git a/Engine/Source/Programs/Horde/HordeDashboard/src/components/buildhealth/stepoutcome/StepOutcomeD3Component.tsx b/Engine/Source/Programs/Horde/HordeDashboard/src/components/buildhealth/stepoutcome/StepOutcomeD3Component.tsx
--- a/Engine/Source/Programs/Horde/HordeDashboard/src/components/buildhealth/stepoutcome/StepOutcomeD3Component.tsx
+++ b/Engine/Source/Programs/Horde/HordeDashboard/src/components/buildhealth/stepoutcome/StepOutcomeD3Component.tsx
@@ -165,17 +165,6 @@ function GenerateD3Table(stepOutcomeTable: StepOutcomeTable, handler: StepOutcom
             .style("min-width", "30px")
             .style("max-width", "30px");
 
-        // Second Column Header; (0, 1) - this sits on top of step name; guarantee the width of step name TH.
-        theaderrow.append("th")
-            .style("position", "sticky")
-            .style("z-index", "6")
-            .style("top", "0")
-            .style("left", "30px")
-            .style("background-color", modeColors.background)
-            .style("width", "320px")
-            .style("min-width", "320px")
-            .style("max-width", "320px");
-
         scaffold.current = { wrapper, table, thead, theaderrow, tbody, tooltip };
 
         return () => {
@@ -187,6 +176,59 @@ function GenerateD3Table(stepOutcomeTable: StepOutcomeTable, handler: StepOutcom
     useEffect(() => {
         if (!containerRef.current || !scaffold.current) return;
 
+        // Second Column Header; (0, 1) - used to display refresh time - this sits on top of step name; guarantee the width of step name TH.
+        scaffold.current.theaderrow
+            .selectAll("th.refresh-header")
+            .data([handler.lastRefreshDate])
+            .join("th")
+            .attr("class", "refresh-header")
+            .style("position", "sticky")
+            .style("z-index", "6")
+            .style("top", "0")
+            .style("left", "30px")
+            .style("background-color", modeColors.background)
+            .style("width", "320px")
+            .style("min-width", "320px")
+            .style("max-width", "320px")
+            .style("cursor", "pointer")
+            .text(d =>
+                d
+                    ? `Last Refresh: ${d.toLocaleTimeString([], {
+                        hour: '2-digit',
+                        minute: '2-digit',
+                        second: '2-digit',
+                        timeZoneName: 'short'
+                    })}`
+                    : ""
+            )
+            .on("mouseover", function (_event, _d) {
+                d3.select(this).style("background-color", theme.semanticColors.bodyBackgroundHovered);
+                const totalSeconds = Math.floor(handler.refreshTime / 1000);
+                const minutes = Math.floor(totalSeconds / 60);
+                const seconds = totalSeconds % 60;
+                const formatted = `${minutes}m ${seconds}s`;
+
+                const tooltipContent = document.createElement("div");
+                tooltipContent.style.paddingLeft = "4px";
+
+                const refreshRateDiv = document.createElement("div");
+                refreshRateDiv.innerHTML = `<strong>Refresh Rate:</strong><p class="${tooltipClasses.tooltipLine}">${formatted}</p><p>Click to refresh.</p>`;
+                tooltipContent.appendChild(refreshRateDiv);
+
+                attachTargettedToolTip(this as HTMLElement, scaffold.current!.tooltip, tooltipContent);
+            })
+            .on("mouseout", function () {
+                scaffold.current!.tooltip.style("opacity", 0);
+                d3.select(this).style("background-color", modeColors.content);
+            })
+            .on("click", (_event, _d) => {
+                if (handler.isInFullDataRefresh || handler.isInIncrementalDataRefresh) {
+                    return;
+                }
+
+                handler.poll();
+            });
+
         // Add all columns to the header row; these will be the changelists.
         scaffold.current.theaderrow.selectAll("th.col")
             .data([...stepOutcomeTable.getDateAnchoredChanges()] as ChangeHeader[])
@@ -378,10 +420,10 @@ function GenerateD3Table(stepOutcomeTable: StepOutcomeTable, handler: StepOutcom
                 .style("background-color", d => getCellStyle(d).bgColor)
                 .style("display", "inline-block")
                 .style("text-align", "center")
-                .on("click", (event, d) => {
+                .on("click", (_event, d) => {
                     onCellClick?.(d);
                 })
-                .on("mouseover", function (event, d) {
+                .on("mouseover", function (_event, d) {
                     if (!d) return;
                     const entry = d as unknown as StepOutcomeTableEntry;
                     const tooltipContent = document.createElement("div");
@@ -446,6 +488,14 @@ function GenerateD3Table(stepOutcomeTable: StepOutcomeTable, handler: StepOutcom
         gridTemplateColumns: '1fr auto',
     };
 
+    // Binding to the reference of the StepOutcomeTableData is by design. Once we have performed the initial load, we want to snap to the right side.
+    // After the user regains control and can horizontally scroll, we do *not* want to snap to right anymore.
+    useEffect(() => {
+        if (containerRef.current) {
+            containerRef.current.scrollLeft = containerRef.current.scrollWidth;
+        }
+    }, [handler.StepOutcomeTableData]);
+
     return (
         <div style={gridStyle}>
             {/* Left side: table */}
diff --git a/Engine/Source/Programs/Horde/HordeDashboard/src/components/buildhealth/stepoutcome/StepOutcomeDataHandler.ts b/Engine/Source/Programs/Horde/HordeDashboard/src/components/buildhealth/stepoutcome/StepOutcomeDataHandler.ts
--- a/Engine/Source/Programs/Horde/HordeDashboard/src/components/buildhealth/stepoutcome/StepOutcomeDataHandler.ts
+++ b/Engine/Source/Programs/Horde/HordeDashboard/src/components/buildhealth/stepoutcome/StepOutcomeDataHandler.ts
@@ -94,6 +94,9 @@ export class StepOutcomeDataHandler extends PollBase {
 
     // #region -- Private Members -- 
 
+    @observable
+    private internalLastRefreshDate?: Date;
+
     @observable
     private intialRefresh = true;
     @observable
@@ -216,6 +219,7 @@ export class StepOutcomeDataHandler extends PollBase {
         runInAction(() => {
             this.activeRefresh = false;
             this.intialRefresh = false;
+            this.internalLastRefreshDate = new Date();
         });
 
         return Promise.resolve();
@@ -242,6 +246,15 @@ export class StepOutcomeDataHandler extends PollBase {
 
     // #region -- Public API  --
 
+    /**
+     * Gets the last refresh date.
+     * @return The last refresh date, if a refresh has occurred. 
+     */
+    @computed
+    get lastRefreshDate(): Date | undefined {
+        return this.internalLastRefreshDate;
+    }
+
     /**
      * Gets whether the handler has yet to initiate the initial refresh.
      * @returns True if the handler has not issued a initial refresh, false otherwise.
@@ -288,6 +301,14 @@ export class StepOutcomeDataHandler extends PollBase {
         this.poll();
     }
 
+    /**
+     * Gets the current refresh cadence of the data handler.
+     * @returns The handler refresh cadence in ms.
+     */
+    get refreshTime(): number {
+        return this.pollTime;
+    }
+
     /**
      *  Sets a new refresh time for the handler, in ms.
      * @param newRefreshTime The new refresh rate.
