[Horde][BuildHealth][Dashboard] Selective dropdown updates on change & Step filters no longer being applied across streams

diff --git a/Engine/Source/Programs/Horde/HordeDashboard/src/components/buildhealth/BuildHealthOptions.ts b/Engine/Source/Programs/Horde/HordeDashboard/src/components/buildhealth/BuildHealthOptions.ts
--- a/Engine/Source/Programs/Horde/HordeDashboard/src/components/buildhealth/BuildHealthOptions.ts
+++ b/Engine/Source/Programs/Horde/HordeDashboard/src/components/buildhealth/BuildHealthOptions.ts
@@ -52,7 +52,7 @@ export const JobHistoryTimeSpans: TimeSpan[] = [
     {
         text: "Past 1 Week", key: "time_1_week", minutes: 60 * 24 * 7
     },
-        {
+    {
         text: "Past 2 Weeks", key: "time_2_week2", minutes: 60 * 24 * 7 * 2
     },
     {
@@ -315,7 +315,7 @@ export class BuildHealthOptionsController {
             this.state.enabledStreams[id] = name;
         } else {
             delete this.state.enabledStreams[id];
-            this.clearTemplates(true);
+            this.clearTemplates(true, id);
         }
         this.setOptionsChanged();
     }
@@ -332,7 +332,8 @@ export class BuildHealthOptionsController {
             this.state.enabledTemplates[id] = name;
         } else {
             delete this.state.enabledTemplates[id];
-            this.clearSteps();
+            let { streamId, templateId } = decodeTemplateKey(id);
+            this.clearSteps({streamId, templateId});
         }
         this.setOptionsChanged();
     }
@@ -394,32 +395,43 @@ export class BuildHealthOptionsController {
     clearStreams(cascade?: boolean) {
         Object.keys(this.state.enabledStreams).forEach(key => {
             delete this.state.enabledStreams[key];
+
+            if (cascade) {
+                this.clearTemplates(cascade, key);
+            }
         });
-        if (cascade) {
-            this.clearTemplates(cascade);
-        }
+
     }
 
     /**
       * Clear all the templates tracked by the options.
       * @param cascade Whether to cascade this down to dependent options.
+      * @param owningStreamId Whether to clear an item if it belongs to a stream. If not provided, will remove any template.
       */
-    clearTemplates(cascade?: boolean) {
+    clearTemplates(cascade?: boolean, owningStreamId?: string) {
         Object.keys(this.state.enabledTemplates).forEach(key => {
-            delete this.state.enabledTemplates[key];
+            let { streamId, templateId } = decodeTemplateKey(key);
+            if (!owningStreamId || (streamId === owningStreamId)) {
+                delete this.state.enabledTemplates[key];
+
+                if (cascade) {
+                    this.clearSteps(owningStreamId ? { streamId: owningStreamId!, templateId: templateId } : undefined);
+                }
+            }
         });
-        if (cascade) {
-            this.clearSteps();
-        }
+
     }
 
     /**
      * Clear all the steps tracked by the options.
-     * @param cascade Whether to cascade this down to dependent options.
+     * @param owningHierarchy Whether to clear an item if it belongs to a stream & template. If not provided, will remove any step.
      */
-    clearSteps() {
+    clearSteps(owningHierarchy?: { streamId: string; templateId: string }) {
         Object.keys(this.state.enabledSteps).forEach(key => {
-            delete this.state.enabledSteps[key];
+            let { streamId, templateId, stepName } = decodeStepKey(key);
+            if ((!owningHierarchy) || (streamId === owningHierarchy.streamId && templateId === owningHierarchy.templateId)) {
+                delete this.state.enabledSteps[key];
+            }
         });
     }
 
diff --git a/Engine/Source/Programs/Horde/HordeDashboard/src/components/buildhealth/BuildHealthUtilities.ts b/Engine/Source/Programs/Horde/HordeDashboard/src/components/buildhealth/BuildHealthUtilities.ts
--- a/Engine/Source/Programs/Horde/HordeDashboard/src/components/buildhealth/BuildHealthUtilities.ts
+++ b/Engine/Source/Programs/Horde/HordeDashboard/src/components/buildhealth/BuildHealthUtilities.ts
@@ -21,6 +21,18 @@ export function decodeTemplateKey(templateKey: string): { streamId: string; temp
     return { streamId, templateId };
 }
 
+/**
+ * Produces an encoded step key for unique identification
+ * @param streamId The streamId to encode.
+ * @param templateId The templateId to encode.
+ * @param stepName The stepName to encode.
+ * @returns The key representing the step.
+ * @remark We don't want toLocalLowerCase the templateId until UE-315953 is completed in order to bump the url encoding query.
+ */
+export function encodeStepNameFromStrings(streamId: string, templateId: string, stepName: string) {
+    return `${streamId.toLocaleLowerCase()}${KEY_SEPARATOR}${templateId}${KEY_SEPARATOR}${stepName.toLocaleLowerCase()}`;
+}
+
 /**
  * Produces an encoded step key for unique identification.
  * @param step The step to produce the key for.
diff --git a/Engine/Source/Programs/Horde/HordeDashboard/src/components/buildhealth/stepoutcome/StepOutcomeDataHandler.ts b/Engine/Source/Programs/Horde/HordeDashboard/src/components/buildhealth/stepoutcome/StepOutcomeDataHandler.ts
--- a/Engine/Source/Programs/Horde/HordeDashboard/src/components/buildhealth/stepoutcome/StepOutcomeDataHandler.ts
+++ b/Engine/Source/Programs/Horde/HordeDashboard/src/components/buildhealth/stepoutcome/StepOutcomeDataHandler.ts
@@ -4,7 +4,7 @@ import { GetJobResponse, IssueQuery, JobStreamQuery } from "horde/backend/Api";
 import { PollBase } from "horde/backend/PollBase";
 import { StepOutcomeTable, StepOutcomeTableEntry } from "./StepOutcomeDataTypes";
 import backend from "horde/backend";
-import { decodeStepKey, decodeTemplateKey } from "../BuildHealthUtilities";
+import { decodeStepKey, decodeTemplateKey, encodeStepNameFromStrings } from "../BuildHealthUtilities";
 import { action, computed, makeObservable, observable, runInAction } from "mobx";
 
 // #region -- API Data Types & Utilities --
@@ -141,7 +141,24 @@ export class StepOutcomeDataHandler extends PollBase {
         this.stepOutcomeTableData.reset();
 
         // @todo UE-309726 - The poll should not need to understand internal data structure; the filter object should be already decoded.
-        const stepSet = new Set(this.activeFilter.steps?.map(s => decodeStepKey(s).stepName.toLowerCase()));
+        const stepMap: Map<string, Map<string, Set<string>>> = new Map();
+
+        (this.activeFilter.steps ?? []).forEach(step => {
+            const { streamId, templateId, stepName } = decodeStepKey(step);
+
+            if (!stepMap.has(streamId)) {
+                stepMap.set(streamId, new Map());
+            }
+
+            const templateMap = stepMap.get(streamId)!;
+
+            if (!templateMap.has(templateId)) {
+                templateMap.set(templateId, new Set());
+            }
+
+            templateMap.get(templateId)!.add(step.toLocaleLowerCase());
+        });
+
         const streamList = this.activeFilter.streams ?? [];
         const modifiedAfter = this.activeFilter.jobHistorySpan.start.toISOString();
         const jobCount = this.activeFilter.maxJobCount ?? this.MAX_JOB_COUNT_DEFAULT;
@@ -186,18 +203,22 @@ export class StepOutcomeDataHandler extends PollBase {
             // Process all job responses
             for (const jobResponse of streamResult) {
                 // Filter preflights as appropriate
-                if(!this.activeFilter.includePreflights && (jobResponse.preflightChange || jobResponse.preflightDescription ))
-                {
+                if (!this.activeFilter.includePreflights && (jobResponse.preflightChange || jobResponse.preflightDescription)) {
                     continue;
                 }
 
                 const batchResponses = jobResponse.batches ?? [];
                 for (const batch of batchResponses) {
                     for (const step of batch.steps) {
+                        let fullyQualifiedStepName = encodeStepNameFromStrings(jobResponse.streamId, jobResponse.templateId!, step.name).toLocaleLowerCase();
+                        let streamStepSet = stepMap.get(jobResponse.streamId);
+                        let templateStepSet = streamStepSet ? streamStepSet.get(jobResponse.templateId!) : undefined;
                         if (
-                            stepSet.size === 0 ||
-                            stepSet.has(step.name.toLowerCase())
+                            !templateStepSet ||
+                            templateStepSet.size === 0 ||
+                            templateStepSet.has(fullyQualifiedStepName)
                         ) {
+
                             this.stepOutcomeTableData.addEntry(
                                 new StepOutcomeTableEntry(step, jobResponse.id, jobResponse.name, jobResponse.streamId, jobResponse.createTime.toString(), jobResponse.change)
                             );
@@ -210,10 +231,10 @@ export class StepOutcomeDataHandler extends PollBase {
         console.info(`Retrieved ${this.jobData.length} job instance results for StepOutcome query.`);
 
         this.stepOutcomeTableData.orderTableDataByChange(true);
-        
+
         // @todo - UE-315881 will remove the need for this constraint.
         this.StepOutcomeTableData.supportsDateAnchoredChanges = !this.activeFilter.includePreflights;
-        
+
         this.setUpdated();
 
         runInAction(() => {
