[Horde][BuildHealth][Dashboard] UseEffect in order to calculate proper height.

diff --git a/Engine/Source/Programs/Horde/HordeDashboard/src/components/buildhealth/stepoutcome/StepOutcomeD3Component.tsx b/Engine/Source/Programs/Horde/HordeDashboard/src/components/buildhealth/stepoutcome/StepOutcomeD3Component.tsx
--- a/Engine/Source/Programs/Horde/HordeDashboard/src/components/buildhealth/stepoutcome/StepOutcomeD3Component.tsx
+++ b/Engine/Source/Programs/Horde/HordeDashboard/src/components/buildhealth/stepoutcome/StepOutcomeD3Component.tsx
@@ -496,13 +496,31 @@ function GenerateD3Table(stepOutcomeTable: StepOutcomeTable, handler: StepOutcom
         }
     }, [handler.StepOutcomeTableData]);
 
+    const updateMaxHeight = () => {
+        if (containerRef.current) {
+            const containerTop = containerRef.current.getBoundingClientRect().top;
+            const availableHeight = window.innerHeight - containerTop - 10;
+            containerRef.current.style.maxHeight = `${availableHeight}px`;
+        }
+    };
+
+    // Maximum height of container, reactive to zoom events.
+    useEffect(() => {
+        updateMaxHeight();
+
+        window.addEventListener("resize", updateMaxHeight);
+
+        return () => {
+            window.removeEventListener("resize", updateMaxHeight);
+        };
+    }, []);
+
     return (
         <div style={gridStyle}>
             {/* Left side: table */}
             <div id="d3-container"
                 ref={containerRef}
                 style={{
-                    maxHeight: '83vh',
                     overflow: 'auto'
                 }}
             />
