[Horde][Server][Tests] Fix the DatabaseIntegrationTest class to properly forward dispose calls to the mongoDBRunner.

diff --git a/Engine/Source/Programs/Horde/HordeServer.Tests.Shared/DatabaseIntegrationTest.cs b/Engine/Source/Programs/Horde/HordeServer.Tests.Shared/DatabaseIntegrationTest.cs
--- a/Engine/Source/Programs/Horde/HordeServer.Tests.Shared/DatabaseIntegrationTest.cs
+++ b/Engine/Source/Programs/Horde/HordeServer.Tests.Shared/DatabaseIntegrationTest.cs
@@ -95,7 +95,9 @@ namespace HordeServer.Tests
 
 		public void Dispose()
 		{
+
 			IMongoClient strictClient = Client.WithWriteConcern(new WriteConcern(journal: true));
+
 			for (int i = 0; i < 5; i++)
 			{
 				strictClient.DropDatabase(DatabaseName);
@@ -109,6 +111,14 @@ namespace HordeServer.Tests
 
 			throw new Exception($"Unable to drop MongoDB database {DatabaseName}");
 		}
+
+		public static void FullDispose()
+		{
+			lock (s_lockObject)
+			{
+				s_mongoDbRunner?.Dispose();
+			}
+		}
 	}
 
 	public sealed class RedisInstance : IDisposable
@@ -181,6 +191,12 @@ namespace HordeServer.Tests
 		{
 		}
 
+		[AssemblyCleanup]
+		public static void Cleanup()
+		{
+			MongoInstance.FullDispose();
+		}
+
 		protected override void ConfigureServices(IServiceCollection services)
 		{
 			base.ConfigureServices(services);
