Horde: Fix compiler warnings

diff --git a/Engine/Source/Programs/Horde/Plugins/Experimental/HordeServer.Experimental.Tests/Notifications/ExperimentalNotificationTests.cs b/Engine/Source/Programs/Horde/Plugins/Experimental/HordeServer.Experimental.Tests/Notifications/ExperimentalNotificationTests.cs
--- a/Engine/Source/Programs/Horde/Plugins/Experimental/HordeServer.Experimental.Tests/Notifications/ExperimentalNotificationTests.cs
+++ b/Engine/Source/Programs/Horde/Plugins/Experimental/HordeServer.Experimental.Tests/Notifications/ExperimentalNotificationTests.cs
@@ -1290,7 +1290,7 @@ namespace HordeServer.Experimental.Tests.Notifications
 		}
 
 		[TestMethod]
-		public async Task NotificationConfigurationsByStreamTag()
+		public async Task NotificationConfigurationsByStreamTagAsync()
 		{
 			// Scenario: Job matched with the configuration using associated StreamTags enabled in the StreamConfig
 			// Expected: Single record should exist in both the JobNotificationCollection and JobStepNotificationCollection
