Horde: Add new PerforceMaterializer

Uses plain workspaces in Perforce, currently optimized for single stream use.

diff --git a/Engine/Source/Programs/Horde/Drivers/JobDriver.Tests/JobDriver.Tests.csproj b/Engine/Source/Programs/Horde/Drivers/JobDriver.Tests/JobDriver.Tests.csproj
--- a/Engine/Source/Programs/Horde/Drivers/JobDriver.Tests/JobDriver.Tests.csproj
+++ b/Engine/Source/Programs/Horde/Drivers/JobDriver.Tests/JobDriver.Tests.csproj
@@ -32,6 +32,7 @@
 
   <ItemGroup>
     <ProjectReference Include="..\JobDriver\JobDriver.csproj" />
+    <ProjectReference Include="..\..\..\Shared\EpicGames.Perforce.Managed.Tests\EpicGames.Perforce.Managed.Tests.csproj" />
   </ItemGroup>
 
   <ItemGroup>
diff --git a/Engine/Source/Programs/Horde/Drivers/JobDriver.Tests/PerforceMaterializerTests.cs b/Engine/Source/Programs/Horde/Drivers/JobDriver.Tests/PerforceMaterializerTests.cs
new file mode 100644
--- /dev/null
+++ b/Engine/Source/Programs/Horde/Drivers/JobDriver.Tests/PerforceMaterializerTests.cs
@@ -0,0 +1,147 @@
+// Copyright Epic Games, Inc. All Rights Reserved.
+
+using System;
+using System.Collections.Generic;
+using System.IO;
+using System.Linq;
+using System.Threading;
+using System.Threading.Tasks;
+using EpicGames.Perforce.Fixture;
+using HordeCommon.Rpc.Messages;
+using JobDriver.Execution;
+using Microsoft.Extensions.Logging;
+using Microsoft.VisualStudio.TestTools.UnitTesting;
+using OpenTelemetry.Trace;
+
+namespace JobDriver.Tests;
+
+[TestClass]
+public class PerforceMaterializerTests : BasePerforceFixtureTest
+{
+	private static Tracer NoOpTracer { get; } = TracerProvider.Default.GetTracer("NoOp");
+
+	[TestMethod]
+	public async Task WalkChangeListsWithReusedMaterializerAsync()
+	{
+		CancellationToken ct = CancellationToken.None;
+		using PerforceMaterializer pm = CreatePm(Fixture.StreamFooMain);
+		List<StreamState> states = [Main(2), Main(3), Main(4), Main(5), Main(6), Main(7), Main(7, 8)];
+		states.AddRange(states.SkipLast(1).Reverse());
+		foreach (StreamState ss in states)
+		{
+			await pm.SyncAsync(ss.Cl, ss.ShelvedCl, new SyncOptions(), ct);
+			Fixture.StreamFooMain.GetChangelist(ss.ShelvedCl > 0 ? ss.ShelvedCl : ss.Cl).AssertDepotFiles(pm.DirectoryPath.FullName);
+			await pm.FinalizeAsync(ct);
+		}
+	}
+	
+	[TestMethod]
+	public async Task TestAllStatesAsync()
+	{
+		StreamState[] mainState = [Main(2), Main(4), Main(5), Main(7), Main(7, 8)];
+		StreamState[] releaseStates = [Release(10), Release(11)];
+		
+		StreamState?[] localStates = [null, ..mainState, ..releaseStates];
+		StreamState?[] serverStates = [..localStates];
+		StreamState[] targetStates = [..mainState, ..releaseStates];
+
+		foreach (StreamState? localState in localStates)
+		{
+			foreach (StreamState? serverState in serverStates)
+			{
+				foreach (StreamState targetState in targetStates)
+				{
+					await SyncWithScenarioAsync(localState, serverState, targetState, CancellationToken.None);
+				}
+			}
+		}
+	}
+
+	private StreamState Main(int cl, int shelvedCl = 0) => new (Fixture.StreamFooMain, cl, shelvedCl);
+	private StreamState Release(int cl, int shelvedCl = 0) => new (Fixture.StreamFooRelease, cl, shelvedCl);
+	private record StreamState(StreamFixture Stream, int Cl, int ShelvedCl = 0)
+	{
+		public override string ToString()
+		{
+			return $"Stream={Stream.Root} CL={Cl} Shelve={ShelvedCl}";
+		}
+	}
+	private record SyncScenario(StreamState? LocalState, StreamState? ServerState, StreamState TargetState);
+
+	private Task SyncWithScenarioAsync(StreamState? localState, StreamState? serverState, StreamState targetState, CancellationToken cancellationToken)
+	{
+		return SyncWithScenarioAsync(new SyncScenario(localState, serverState, targetState), cancellationToken);
+	}
+
+	private async Task SyncWithScenarioAsync(SyncScenario scenario, CancellationToken cancellationToken)
+	{
+		RemoveTempDir();
+		CreateTempDir();
+		
+		Console.WriteLine("Temp dir is " + TempDir.FullName);
+        
+		if (scenario.LocalState != null)
+		{
+			// Initialize a local state with files corresponding to given stream state, but delete the client
+			Console.WriteLine($"Preparing local state {scenario.LocalState} ...");
+			using PerforceMaterializer pm = CreatePm(scenario.LocalState.Stream);
+			await SyncToStateAsync(pm, scenario.LocalState, cancellationToken);
+			await pm.DeleteClientForTestAsync(cancellationToken);
+		}
+		
+		if (scenario.ServerState != null)
+		{
+			Console.WriteLine($"Preparing server state {scenario.ServerState} ...");
+			
+			string stateDir = TempDir.FullName;
+			string backupDir = TempDir.FullName + "-backup";
+			if (Directory.Exists(stateDir))
+			{
+				// Save the original local state so the sync below cannot overwrite
+				Directory.Move(stateDir, backupDir);
+			}
+			
+            using PerforceMaterializer pm = CreatePm(scenario.ServerState.Stream);
+			await SyncToStateAsync(pm, scenario.ServerState, cancellationToken);
+
+			// Now that the workspace has been synced, reset the local files back
+			DeletePerforceDir(stateDir);
+			if (Directory.Exists(backupDir))
+			{
+				Directory.Move(backupDir, stateDir);
+			}
+		}
+
+		{
+			Console.WriteLine($" Local state: {scenario.LocalState}");
+			Console.WriteLine($"Server state: {scenario.ServerState}");
+			Console.WriteLine($"Target state: {scenario.TargetState}");
+			Console.WriteLine("Syncing ...");
+			using PerforceMaterializer pm = CreatePm(scenario.TargetState.Stream);
+			await SyncToStateAsync(pm, scenario.TargetState, cancellationToken);
+		}
+	}
+
+	private static async Task SyncToStateAsync(PerforceMaterializer pm, StreamState state, CancellationToken cancellationToken)
+	{
+		await pm.SyncAsync(state.Cl, state.ShelvedCl, new SyncOptions(), cancellationToken);
+		state.Stream.GetChangelist(state.ShelvedCl > 0 ? state.ShelvedCl : state.Cl).AssertDepotFiles(pm.DirectoryPath.FullName);
+		await pm.FinalizeAsync(cancellationToken);
+	}
+
+	private PerforceMaterializer CreatePm(StreamFixture stream, string? dataDir = null, string? identifier = "pmTest")
+	{
+		ILogger logger = LoggerFactory.CreateLogger("PM");
+		RpcAgentWorkspace raw = new ()
+		{
+			ServerAndPort = PerforceConnection.ServerAndPort,
+			UserName = PerforceConnection.UserName,
+			Password = PerforceConnection.Settings.Password,
+			Identifier = $"{TestGuid}-{identifier}",
+			Stream = stream.Root,
+		};
+		PerforceMaterializerOptions options = new(dataDir ?? TempDir.FullName, raw);
+		return new PerforceMaterializer(options, NoOpTracer, logger);
+	}
+}
+
diff --git a/Engine/Source/Programs/Horde/Drivers/JobDriver/Execution/PerforceMaterializer.cs b/Engine/Source/Programs/Horde/Drivers/JobDriver/Execution/PerforceMaterializer.cs
new file mode 100644
--- /dev/null
+++ b/Engine/Source/Programs/Horde/Drivers/JobDriver/Execution/PerforceMaterializer.cs
@@ -0,0 +1,344 @@
+﻿// Copyright Epic Games, Inc. All Rights Reserved.
+
+using System.Diagnostics;
+using System.Text.Json;
+using EpicGames.Core;
+using EpicGames.Perforce;
+using HordeCommon.Rpc.Messages;
+using JobDriver.Utility;
+using Microsoft.Extensions.Logging;
+using OpenTelemetry.Trace;
+
+namespace JobDriver.Execution;
+
+/// <summary>
+/// Options for <see cref="PerforceMaterializer" />
+/// </summary>
+/// <param name="DirPath">Base directory for this workspace</param>
+/// <param name="AgentWorkspace">Workspace options</param>
+/// <param name="HostName">Optional hostname override</param>
+/// <param name="EnvVars">Optional environment variables to specify</param>
+/// <param name="SyncThreads">Number of threads to use when syncing a Perforce workspace</param>
+/// <param name="SyncBatch">How many files in each sync batch</param>
+/// <param name="SyncBatchSize">Size of each sync batch</param>
+public record PerforceMaterializerOptions(
+	string DirPath,
+	RpcAgentWorkspace AgentWorkspace,
+	string? HostName = null,
+	IReadOnlyDictionary<string, string>? EnvVars = null,
+	int SyncThreads = -1,
+	int SyncBatch = -1,
+	int SyncBatchSize = -1);
+
+/// <summary>
+/// Materializer using standard Perforce client for syncing files
+/// Main use-case are for agents using only a single stream.
+/// For example, incremental build agents, as stream switching has not been optimized.
+/// </summary>
+public sealed class PerforceMaterializer : IWorkspaceMaterializer
+{
+	private enum TransactionStatus { Clean = 0, Dirty = 1 }
+	private record State(
+		int Version,
+		TransactionStatus Status,
+		string Client,
+		string Identifier,
+		string Stream,
+		int Changelist,
+		int ShelvedChangelist);
+	
+    /// <inheritdoc/>
+	public DirectoryReference DirectoryPath { get; }
+	
+	/// <inheritdoc/>
+	public string Identifier { get; }
+	
+	/// <inheritdoc/>
+	public IReadOnlyDictionary<string, string> EnvironmentVariables { get; }
+	
+	/// <inheritdoc/>
+	public bool IsPerforceWorkspace { get; } = true;
+	
+	private static readonly JsonSerializerOptions s_jsonOptions = new() { WriteIndented = true };
+	private readonly PerforceMaterializerOptions _options;
+	private readonly Tracer _tracer;
+	private readonly ILogger _logger;
+	private readonly DirectoryReference _metadataDir;
+	private readonly DirectoryReference _workspaceDir;
+	private PerforceSettings? _perforceSettings;
+	private IPerforceConnection? _perforceWithoutClient;
+	private IPerforceConnection? _perforceWithClient;
+	private const string StateFilename = "State.json";
+	private readonly string _stateFile;
+	private readonly string _stateTempFile;
+
+	/// <summary>
+	/// Constructor
+	/// </summary>
+	public PerforceMaterializer(PerforceMaterializerOptions options, Tracer tracer, ILogger logger)
+	{
+		_options = options.SyncThreads == -1 ? options with { SyncThreads = GetDefaultThreadCount(6) } : options;
+		Identifier = options.AgentWorkspace.Identifier;
+		EnvironmentVariables = options.EnvVars ?? new Dictionary<string, string>();
+		_tracer = tracer;
+		_logger = logger;
+		_metadataDir = DirectoryReference.Combine(new DirectoryReference(_options.DirPath), Identifier);
+		_workspaceDir = DirectoryReference.Combine(_metadataDir, "Sync");
+		_stateFile = Path.Join(_metadataDir.FullName, StateFilename);
+		_stateTempFile = Path.Join(_metadataDir.FullName, StateFilename + ".tmp");
+		DirectoryPath = _workspaceDir;
+	}
+	
+	/// <inheritdoc/>
+	public void Dispose()
+	{
+		_perforceWithClient?.Dispose();
+		_perforceWithoutClient?.Dispose();
+	}
+	
+	private async Task InitializeAsync(CancellationToken cancellationToken)
+	{
+		using TelemetrySpan span = _tracer.StartActiveSpan($"{nameof(PerforceMaterializer)}.{nameof(InitializeAsync)}");
+		using ILoggerProgress status = _logger.BeginProgressScope("Initializing...");
+		Stopwatch timer = Stopwatch.StartNew();
+		
+		const bool PreferNative = true;
+		RpcAgentWorkspace raw = _options.AgentWorkspace;
+		PerforceSettings tempSettings = new(raw.ServerAndPort, raw.UserName) { PreferNativeClient = PreferNative, Password = raw.Ticket };
+		_perforceWithoutClient = await PerforceConnection.CreateAsync(tempSettings, _logger);
+		await _perforceWithoutClient.LoginAsync(raw.Password, cancellationToken);
+		
+		// Get the host name, and fill in any missing metadata about the connection
+		InfoRecord info = await _perforceWithoutClient.GetInfoAsync(InfoOptions.ShortOutput, cancellationToken);
+		
+		string? hostName = info.ClientHost;
+		if (hostName == null)
+		{
+			throw new Exception("Unable to determine Perforce host name");
+		}
+		
+		// Replace invalid characters in the workspace identifier with a '+' character
+		// Append the slot index, if it's non-zero
+		//			my $slot_idx = $optional_arguments->{'slot_idx'} || 0;
+		//			$workspace_identifier .= sprintf("+%02d", $slot_idx) if $slot_idx;
+		// If running on an edge server, append the server ID to the client name
+		string edgeSuffix = String.Empty;
+		if (info is { Services: not null, ServerId: not null })
+		{
+			string[] services = info.Services.Split(' ', StringSplitOptions.RemoveEmptyEntries);
+			if (services.Any(x => x.Equals("edge-server", StringComparison.OrdinalIgnoreCase)))
+			{
+				edgeSuffix = $"+{info.ServerId}";
+			}
+		}
+		
+		string clientName = $"Horde+PM+{WorkspaceInfo.GetNormalizedHostName(hostName)}+{raw.Identifier}{edgeSuffix}";
+		_perforceSettings = new PerforceSettings(_perforceWithoutClient.Settings) { ClientName = clientName, PreferNativeClient = PreferNative };
+		status.Progress = $"({timer.Elapsed.TotalSeconds:0.0}s)";
+	}
+	
+	/// <inheritdoc/>
+	public async Task SyncAsync(int changeNum, int shelveChangeNum, SyncOptions options, CancellationToken cancellationToken)
+	{
+		using TelemetrySpan span = _tracer.StartActiveSpan($"{nameof(PerforceMaterializer)}.{nameof(SyncAsync)}");
+		
+		await InitializeAsync(cancellationToken);
+		if (_perforceSettings == null || _perforceWithoutClient == null)
+		{
+			throw new WorkspaceMaterializationException("Perforce settings/client not initialized!");
+		}
+		
+		State? state = await LoadStateAsync(cancellationToken);
+		
+		_logger.LogInformation("Status: {Status}\nClient: {PrevClient} -> {Client}\nIdentifier: {PrevId} -> {Id}\nStream: {PrevStream} -> {Stream}\nCL: {PrevCl} -> {Cl}\nShelve: {PrevShelve} -> {Shelve}",
+			state?.Status,
+			state?.Client, _perforceSettings.ClientName,
+			state?.Identifier, Identifier,
+			state?.Stream, _options.AgentWorkspace.Stream,
+			state?.Changelist, changeNum,
+			state?.ShelvedChangelist, shelveChangeNum);
+		
+		bool mustRecreateClient = state == null || state.Status == TransactionStatus.Dirty || state.Identifier != Identifier || state.Stream != _options.AgentWorkspace.Stream;
+		if (mustRecreateClient)
+		{
+			_logger.LogInformation("Re-creating Perforce client...");
+			using TelemetrySpan deleteSpan = _tracer.StartActiveSpan($"{nameof(PerforceMaterializer)}.{nameof(SyncAsync)}.Delete");
+			using ILoggerProgress status = _logger.BeginProgressScope("Deleting local workspace files...");
+			Stopwatch deleteTimer = Stopwatch.StartNew();
+			FileUtils.ForceDeleteDirectoryContents(_workspaceDir);
+			status.Progress = $"({deleteTimer.Elapsed.TotalSeconds:0.0}s)";
+		}
+
+		ClientRecord client = await CreateOrUpdateClientAsync(mustRecreateClient, cancellationToken);
+		
+		_logger.LogInformation("Using client {ClientName} (Host: {HostName}, Stream: {StreamName}, Type: {Type}, Root: {Path})",
+			client.Name, client.Host, client.Stream, client.Type, client.Root);
+
+		_perforceWithClient = await PerforceConnection.CreateAsync(_perforceSettings, _logger);
+		await RevertInternalAsync(_perforceWithClient, cancellationToken);
+
+		if (state != null && !mustRecreateClient)
+		{
+			using TelemetrySpan deleteSpan = _tracer.StartActiveSpan($"{nameof(PerforceMaterializer)}.{nameof(SyncAsync)}.Flush");
+			using ILoggerProgress status = _logger.BeginProgressScope($"Flushing client to CL {state.Changelist} (p4 sync -k)");
+			Stopwatch flushTimer = Stopwatch.StartNew();
+			await _perforceWithClient.SyncQuietAsync(EpicGames.Perforce.SyncOptions.KeepWorkspaceFiles, -1, $"@{state.Changelist}", cancellationToken);
+			status.Progress = $"({flushTimer.Elapsed.TotalSeconds:0.0}s)";
+		}
+		
+		await SaveStateAsync(TransactionStatus.Dirty, client.Name, changeNum, shelveChangeNum, cancellationToken);
+
+		{
+			using TelemetrySpan syncSpan = _tracer.StartActiveSpan($"{nameof(PerforceMaterializer)}.{nameof(SyncAsync)}.Sync");
+			using ILoggerProgress status = _logger.BeginProgressScope("Syncing files...");
+			Stopwatch syncTimer = Stopwatch.StartNew();
+			
+			string shelveStr = shelveChangeNum > 0 ? $" with shelved CL {shelveChangeNum}" : "";
+			_logger.LogInformation("Syncing CL {ChangeNum}{Shelve} ... (Threads: {Threads} Batch: {Batch} BatchSize: {BatchSize})",
+				changeNum, shelveStr, _options.SyncThreads, _options.SyncBatch, _options.SyncBatchSize);
+			
+			List<SyncSummaryRecord> syncSummaries = await _perforceWithClient.SyncQuietAsync(
+				EpicGames.Perforce.SyncOptions.None, -1, _options.SyncThreads, _options.SyncBatch, _options.SyncBatchSize,
+				-1, -1, $"@{changeNum}", cancellationToken);
+
+			long totalFileCount = syncSummaries.Sum(x => x.TotalFileCount);
+			long totalFileSize = syncSummaries.Sum(x => x.TotalFileSize);
+			double totalFileSizeMb = totalFileSize / (1024.0 * 1024.0);
+			_logger.LogInformation("Synced {FileCount} files. {FileSizeMb:F2} MB ({FileSize} bytes)", totalFileCount, totalFileSizeMb, totalFileSize);
+			syncSpan.SetAttribute("horde.pm.file_count", totalFileCount);
+			syncSpan.SetAttribute("horde.pm.file_size", totalFileSize);
+			status.Progress = $"({syncTimer.Elapsed.TotalSeconds:0.0}s)";
+		}
+
+		if (shelveChangeNum > 0)
+		{
+			using TelemetrySpan unshelveSpan = _tracer.StartActiveSpan($"{nameof(PerforceMaterializer)}.{nameof(SyncAsync)}.Unshelve");
+			using ILoggerProgress status = _logger.BeginProgressScope("Unshelving files...");
+			Stopwatch unshelveTimer = Stopwatch.StartNew();
+			
+			List<UnshelveRecord> unshelveRecords = await _perforceWithClient.UnshelveAsync(shelveChangeNum, -1, null, null, null, UnshelveOptions.ForceOverwrite, Array.Empty<string>(), cancellationToken);
+
+			unshelveSpan.SetAttribute("horde.pm.file_count", unshelveRecords.Count);
+			_logger.LogInformation("Unshelved {FileCount} files.", unshelveRecords.Count);
+			status.Progress = $"({unshelveTimer.Elapsed.TotalSeconds:0.0}s)";
+		}
+
+		await SaveStateAsync(TransactionStatus.Clean, client.Name, changeNum, shelveChangeNum, cancellationToken);
+	}
+	
+	/// <inheritdoc/>
+	public async Task FinalizeAsync(CancellationToken cancellationToken)
+	{
+		if (_perforceWithClient != null)
+		{
+			await RevertInternalAsync(_perforceWithClient, cancellationToken);
+		}
+	}
+
+	private async Task RevertInternalAsync(IPerforceConnection perforce, CancellationToken cancellationToken)
+	{
+		using TelemetrySpan span = _tracer.StartActiveSpan($"{nameof(PerforceMaterializer)}.{nameof(RevertInternalAsync)}");
+		using ILoggerProgress status = _logger.BeginProgressScope("Reverting files...");
+		Stopwatch timer = Stopwatch.StartNew();
+
+		List<RevertRecord> revertRecords = await perforce.RevertAsync(-1, null, RevertOptions.DeleteAddedFiles, "//...", cancellationToken);
+		_logger.LogInformation("Reverted {FileCount} files.", revertRecords.Count);
+		
+		span.SetAttribute("horde.pm.num_files", revertRecords.Count);
+		status.Progress = $"({timer.Elapsed.TotalSeconds:0.0}s)";
+	}
+	
+	private async Task<ClientRecord> CreateOrUpdateClientAsync(bool forceCreate, CancellationToken cancellationToken)
+	{
+		using TelemetrySpan span = _tracer.StartActiveSpan($"{nameof(PerforceMaterializer)}.{nameof(CreateOrUpdateClientAsync)}");
+		if (_perforceSettings == null || _perforceWithoutClient == null)
+		{
+			throw new WorkspaceMaterializationException("Perforce settings/client not initialized!");
+		}
+		
+		string clientName = _perforceSettings.ClientName!;
+		string streamName = _options.AgentWorkspace.Stream;
+		
+		ClientRecord newClient = new (clientName, _perforceSettings.UserName, _workspaceDir.FullName)
+		{
+			Host = _options.HostName ?? Environment.MachineName,
+			Stream = streamName,
+			Type = "partitioned",
+			Root = _workspaceDir.FullName,
+		};
+		Directory.CreateDirectory(newClient.Root);
+		
+		using ILoggerProgress status = _logger.BeginProgressScope(forceCreate ? "Force creating client ..." : "Updating client ...");
+		Stopwatch timer = Stopwatch.StartNew();
+		
+		if (forceCreate)
+		{
+			await _perforceWithoutClient.TryDeleteClientAsync(DeleteClientOptions.None, clientName, cancellationToken);
+			await _perforceWithoutClient.CreateClientAsync(newClient, cancellationToken);
+		}
+		else
+		{
+			// Create or update the client
+			PerforceResponse createRes = await _perforceWithoutClient.TryCreateClientAsync(newClient, cancellationToken);
+			if (!createRes.Succeeded)
+			{
+				_logger.LogInformation("Deleting and creating client...");
+				await _perforceWithoutClient.TryDeleteClientAsync(DeleteClientOptions.None, clientName, cancellationToken);
+				await _perforceWithoutClient.CreateClientAsync(newClient, cancellationToken);
+			}
+		}
+		
+		status.Progress = $"({timer.Elapsed.TotalSeconds:0.0}s)";
+		return newClient;
+	}
+
+	private Task SaveStateAsync(TransactionStatus status, string client ,int changelist, int shelvedChangelist, CancellationToken cancellationToken)
+	{
+		State pms = new (1, status, client, Identifier, _options.AgentWorkspace.Stream, changelist, shelvedChangelist);
+		return SaveStateAsync(pms, cancellationToken);
+	}
+	
+	private async Task SaveStateAsync(State state, CancellationToken cancellationToken)
+	{
+		try
+		{
+			string json = JsonSerializer.Serialize(state, s_jsonOptions);
+			await File.WriteAllTextAsync(_stateTempFile, json, cancellationToken);
+			File.Move(_stateTempFile, _stateFile, overwrite: true);
+		}
+		catch (Exception e)
+		{
+			throw new WorkspaceMaterializationException("Failed to save local workspace state", e);
+		}
+	}
+	
+	private async Task<State?> LoadStateAsync(CancellationToken cancellationToken)
+	{
+		try
+		{
+			if (!File.Exists(_stateFile))
+			{
+				return null;
+			}
+
+			string json = await File.ReadAllTextAsync(_stateFile, cancellationToken);
+			return JsonSerializer.Deserialize<State>(json, s_jsonOptions);
+		}
+		catch (Exception e)
+		{
+			_logger.LogWarning(e, "Failed to load local workspace state");
+			return null;
+		}
+	}
+	
+	private static int GetDefaultThreadCount(int defaultValue) => Math.Min(defaultValue, Math.Max(Environment.ProcessorCount - 1, 1));
+	
+	internal async Task DeleteClientForTestAsync(CancellationToken cancellationToken)
+	{
+		if (_perforceSettings == null || _perforceWithoutClient == null)
+		{
+			throw new WorkspaceMaterializationException("Perforce settings/client not initialized!");
+		}
+		await _perforceWithoutClient.DeleteClientAsync(DeleteClientOptions.None, _perforceSettings.ClientName!, cancellationToken);
+	}
+}
\ No newline at end of file
diff --git a/Engine/Source/Programs/Horde/Drivers/JobDriver/Execution/WorkspaceMaterializer.cs b/Engine/Source/Programs/Horde/Drivers/JobDriver/Execution/WorkspaceMaterializer.cs
--- a/Engine/Source/Programs/Horde/Drivers/JobDriver/Execution/WorkspaceMaterializer.cs
+++ b/Engine/Source/Programs/Horde/Drivers/JobDriver/Execution/WorkspaceMaterializer.cs
@@ -81,12 +81,12 @@ public interface IWorkspaceMaterializer : IDisposable
 	/// Once method has completed, file tree is available on disk.
 	/// </summary>
 	/// <param name="changeNum">Change number to materialize</param>
-	/// <param name="preflightChangeNum">Preflight change number to add</param>
+	/// <param name="shelveChangeNum">Preflight change number to add</param>
 	/// <param name="options">Additional options</param>
 	/// <param name="cancellationToken">Cancellation token for the call</param>
 	/// <exception cref="WorkspaceMaterializationException">Thrown if syncing fails</exception>
 	/// <returns>Async task</returns>
-	Task SyncAsync(int changeNum, int preflightChangeNum, SyncOptions options, CancellationToken cancellationToken);
+	Task SyncAsync(int changeNum, int shelveChangeNum, SyncOptions options, CancellationToken cancellationToken);
 
 	/// <summary>
 	/// Finalize and clean file system
diff --git a/Engine/Source/Programs/Horde/Drivers/JobDriver/Utility/WorkspaceInfo.cs b/Engine/Source/Programs/Horde/Drivers/JobDriver/Utility/WorkspaceInfo.cs
--- a/Engine/Source/Programs/Horde/Drivers/JobDriver/Utility/WorkspaceInfo.cs
+++ b/Engine/Source/Programs/Horde/Drivers/JobDriver/Utility/WorkspaceInfo.cs
@@ -224,7 +224,7 @@ namespace JobDriver.Utility
 		public async Task SetupWorkspaceAsync(IPerforceConnection perforceClient, CancellationToken cancellationToken)
 		{
 			// Create the workspace info
-			perforceClient.Logger.LogInformation("Syncing {ClientName} to {BaseDir} from {Server}, using stream {Stream} and view:{View}", PerforceSettings.ClientName, WorkspaceDir, PerforceSettings.ServerAndPort, StreamName, String.Join("", View.Select(x => $"\n  {x}")));
+			perforceClient.Logger.LogInformation("Syncing {ClientName} to {BaseDir} from {Server}, using stream {Stream} and view: {View}", PerforceSettings.ClientName, WorkspaceDir, PerforceSettings.ServerAndPort, StreamName, String.Join("", View.Select(x => $"\n  {x}")));
 
 			// Create the repository
 			if (RemoveUntrackedFiles)
@@ -462,7 +462,7 @@ namespace JobDriver.Utility
 		/// </summary>
 		/// <param name="hostName">Hostname with optional domain</param>
 		/// <returns>Normalized host name</returns>
-		static string GetNormalizedHostName(string hostName)
+		public static string GetNormalizedHostName(string hostName)
 		{
 			return Regex.Replace(hostName, @"\..*$", "").ToUpperInvariant();
 		}
