Horde: Add support for LatestChangeNumber constant in PerforceMaterializer

diff --git a/Engine/Source/Programs/Horde/Drivers/JobDriver.Tests/PerforceMaterializerTests.cs b/Engine/Source/Programs/Horde/Drivers/JobDriver.Tests/PerforceMaterializerTests.cs
--- a/Engine/Source/Programs/Horde/Drivers/JobDriver.Tests/PerforceMaterializerTests.cs
+++ b/Engine/Source/Programs/Horde/Drivers/JobDriver.Tests/PerforceMaterializerTests.cs
@@ -34,7 +34,20 @@ public class PerforceMaterializerTests : BasePerforceFixtureTest
 			await pm.FinalizeAsync(ct);
 		}
 	}
-	
+
+	[TestMethod]
+	public async Task SyncLatestAsync()
+	{
+		using PerforceMaterializer pm = CreatePm(Fixture.StreamFooMain);
+		
+		await pm.SyncAsync(IWorkspaceMaterializer.LatestChangeNumber, -1, new SyncOptions(), CancellationToken.None);
+		Fixture.StreamFooMain.GetChangelist(7).AssertDepotFiles(pm.DirectoryPath.FullName);
+		await pm.FinalizeAsync(CancellationToken.None);
+
+		PerforceMaterializer.State? state = await pm.LoadStateForTestAsync(CancellationToken.None);
+		Assert.AreEqual(7, state!.Changelist);
+	}
+
 	[TestMethod]
 	public async Task TestAllStatesAsync()
 	{
@@ -77,8 +90,6 @@ public class PerforceMaterializerTests : BasePerforceFixtureTest
 	{
 		RemoveTempDir();
 		CreateTempDir();
-		
-		Console.WriteLine("Temp dir is " + TempDir.FullName);
         
 		if (scenario.LocalState != null)
 		{
@@ -127,6 +138,12 @@ public class PerforceMaterializerTests : BasePerforceFixtureTest
 		await pm.SyncAsync(state.Cl, state.ShelvedCl, new SyncOptions(), cancellationToken);
 		state.Stream.GetChangelist(state.ShelvedCl > 0 ? state.ShelvedCl : state.Cl).AssertDepotFiles(pm.DirectoryPath.FullName);
 		await pm.FinalizeAsync(cancellationToken);
+		
+		PerforceMaterializer.State? pmState = await pm.LoadStateForTestAsync(cancellationToken);
+		Assert.AreEqual(state.Cl, pmState!.Changelist);
+		Assert.AreEqual(state.ShelvedCl, pmState!.ShelvedChangelist);
+		Assert.AreEqual(state.Stream.Root, pmState.Stream);
+		Assert.AreEqual(PerforceMaterializer.TransactionStatus.Clean, pmState.Status);
 	}
 
 	private PerforceMaterializer CreatePm(StreamFixture stream, string? dataDir = null, string? identifier = "pmTest")
diff --git a/Engine/Source/Programs/Horde/Drivers/JobDriver/Execution/PerforceMaterializer.cs b/Engine/Source/Programs/Horde/Drivers/JobDriver/Execution/PerforceMaterializer.cs
--- a/Engine/Source/Programs/Horde/Drivers/JobDriver/Execution/PerforceMaterializer.cs
+++ b/Engine/Source/Programs/Horde/Drivers/JobDriver/Execution/PerforceMaterializer.cs
@@ -37,8 +37,13 @@ public record PerforceMaterializerOptions(
 /// </summary>
 public sealed class PerforceMaterializer : IWorkspaceMaterializer
 {
-	private enum TransactionStatus { Clean = 0, Dirty = 1 }
-	private record State(
+	/// <summary>
+	/// Name of this materializer
+	/// </summary>
+	public const string Name = "Perforce";
+	
+	internal enum TransactionStatus { Clean = 0, Dirty = 1 }
+	internal record State(
 		int Version,
 		TransactionStatus Status,
 		string Client,
@@ -187,19 +192,25 @@ public sealed class PerforceMaterializer : IWorkspaceMaterializer
 		}
 		
 		await SaveStateAsync(TransactionStatus.Dirty, client.Name, changeNum, shelveChangeNum, cancellationToken);
-
+		int reportedSyncChangeNum;
+		
 		{
 			using TelemetrySpan syncSpan = _tracer.StartActiveSpan($"{nameof(PerforceMaterializer)}.{nameof(SyncAsync)}.Sync");
 			using ILoggerProgress status = _logger.BeginProgressScope("Syncing files...");
 			Stopwatch syncTimer = Stopwatch.StartNew();
 			
+			string changeNumStr = changeNum == IWorkspaceMaterializer.LatestChangeNumber ? "#head" : "@" + changeNum;
 			string shelveStr = shelveChangeNum > 0 ? $" with shelved CL {shelveChangeNum}" : "";
 			_logger.LogInformation("Syncing CL {ChangeNum}{Shelve} ... (Threads: {Threads} Batch: {Batch} BatchSize: {BatchSize})",
-				changeNum, shelveStr, _options.SyncThreads, _options.SyncBatch, _options.SyncBatchSize);
-			
+				changeNumStr, shelveStr, _options.SyncThreads, _options.SyncBatch, _options.SyncBatchSize);
+
 			List<SyncSummaryRecord> syncSummaries = await _perforceWithClient.SyncQuietAsync(
 				EpicGames.Perforce.SyncOptions.None, -1, _options.SyncThreads, _options.SyncBatch, _options.SyncBatchSize,
-				-1, -1, $"@{changeNum}", cancellationToken);
+				-1, -1, changeNumStr, cancellationToken);
+
+			// The sync summary will only report the correct CL if #head was used, and we only need it in that case.
+			// Otherwise, just use the explicit change number passed as the parameter.
+			reportedSyncChangeNum = changeNum == IWorkspaceMaterializer.LatestChangeNumber ? syncSummaries.Max(x => x.Change) : changeNum;
 
 			long totalFileCount = syncSummaries.Sum(x => x.TotalFileCount);
 			long totalFileSize = syncSummaries.Sum(x => x.TotalFileSize);
@@ -223,7 +234,7 @@ public sealed class PerforceMaterializer : IWorkspaceMaterializer
 			status.Progress = $"({unshelveTimer.Elapsed.TotalSeconds:0.0}s)";
 		}
 
-		await SaveStateAsync(TransactionStatus.Clean, client.Name, changeNum, shelveChangeNum, cancellationToken);
+		await SaveStateAsync(TransactionStatus.Clean, client.Name, reportedSyncChangeNum, shelveChangeNum, cancellationToken);
 	}
 	
 	/// <inheritdoc/>
@@ -341,4 +352,9 @@ public sealed class PerforceMaterializer : IWorkspaceMaterializer
 		}
 		await _perforceWithoutClient.DeleteClientAsync(DeleteClientOptions.None, _perforceSettings.ClientName!, cancellationToken);
 	}
+	
+	internal Task<State?> LoadStateForTestAsync(CancellationToken cancellationToken)
+	{
+		return LoadStateAsync(cancellationToken);
+	}
 }
\ No newline at end of file
