Horde: Minor fixes to comments in HashedIssueFileMetadataTestAsync
#rnx

diff --git a/Engine/Source/Programs/Horde/Plugins/Build/HordeServer.Build.Tests/Issues/IssueServiceTests.cs b/Engine/Source/Programs/Horde/Plugins/Build/HordeServer.Build.Tests/Issues/IssueServiceTests.cs
--- a/Engine/Source/Programs/Horde/Plugins/Build/HordeServer.Build.Tests/Issues/IssueServiceTests.cs
+++ b/Engine/Source/Programs/Horde/Plugins/Build/HordeServer.Build.Tests/Issues/IssueServiceTests.cs
@@ -1832,12 +1832,9 @@ namespace HordeServer.Tests.Issues
 					"{ \"time\":\"2025-01-01T01:20:30.990Z\",\"level\":\"Error\",\"message\":\"Assertion failed: 1 == 2\",\"properties\":{ \"file\":\"a/d.uasset\"} }"
 				};
 
-				// need to edit lines to evoid localization issue handler handling the error
 				await ParseEventsAsync(job2, 0, 0, lines);
 				await UpdateCompleteStepAsync(job2, 0, 0, JobStepOutcome.Failure);
 
-				// TODO: test multiple events in the same step with lines
-
 				IReadOnlyList<IIssue> issues = await IssueCollection.FindIssuesAsync();
 				Assert.AreEqual(1, issues.Count);
 
