Adding logging to specify both the stream id and stream tags of a stream

#rnx

diff --git a/Engine/Source/Programs/Horde/Plugins/Experimental/HordeServer.Experimental/Notifications/SlackNotificationProcessor.cs b/Engine/Source/Programs/Horde/Plugins/Experimental/HordeServer.Experimental/Notifications/SlackNotificationProcessor.cs
--- a/Engine/Source/Programs/Horde/Plugins/Experimental/HordeServer.Experimental/Notifications/SlackNotificationProcessor.cs
+++ b/Engine/Source/Programs/Horde/Plugins/Experimental/HordeServer.Experimental/Notifications/SlackNotificationProcessor.cs
@@ -361,12 +361,19 @@ namespace HordeServer.Experimental.Notifications
 			// Fetch our notification configurations for the parent job to append our current job's steps
 			ExperimentalConfig currentExperimentalConfig = _experimentalConfig.CurrentValue;
 
+			string streamInformation = $"'{parentJob.StreamId}'";
+			if (streamConfig is not null && streamConfig.StreamTags.Count > 0)
+			{
+				streamInformation += $" or with stream tags {String.Join(',', streamConfig.StreamTags)}";
+			}
+
 			HashSet<JobNotificationConfig>? jobNotificationConfigs;
 			if (!currentExperimentalConfig.TryGetJobNofications(parentJob.StreamId, parentJob.TemplateId, out jobNotificationConfigs, streamConfig?.StreamTags ?? null))
 			{
-				_logger.LogInformation("Parent configurations not found for parent step '{ParentJobStep}' in parent job '{ParentJob}' for '{JobId}' in stream '{StreamId}' and for template '{TemplateId}'", job.ParentJobStepId, job.ParentJobId, job.Id, parentJob.StreamId, parentJob.TemplateId);
+				_logger.LogInformation("Parent configurations not found for parent step '{ParentJobStep}' in parent job '{ParentJob}' for '{JobId}' in stream {StreamInformation} and for template '{TemplateId}'", job.ParentJobStepId, job.ParentJobId, job.Id, streamInformation, parentJob.TemplateId);
 				return true;
 			}
+			_logger.LogInformation("{NumConfigurations} parent configurations found for parent step '{ParentJobStep}' in parent job '{ParentJob}' for '{JobId}' in stream {StreamInformation} and for template '{TemplateId}'", jobNotificationConfigs.Count, job.ParentJobStepId, job.ParentJobId, job.Id, streamInformation, parentJob.TemplateId);
 
 			foreach (JobNotificationConfig config in jobNotificationConfigs)
 			{
@@ -399,13 +406,21 @@ namespace HordeServer.Experimental.Notifications
 			// Fetch our notification configurations for the parent job to append our current job's steps
 			ExperimentalConfig currentExperimentalConfig = _experimentalConfig.CurrentValue;
 
+			string streamInformation = $"'{job.StreamId}'";
+			if (streamConfig is not null && streamConfig.StreamTags.Count > 0)
+			{
+				streamInformation += $" or with stream tags {String.Join(',', streamConfig.StreamTags)}";
+			}
+
 			HashSet<JobNotificationConfig>? jobNotificationConfigs;
 			if (!currentExperimentalConfig.TryGetJobNofications(job.StreamId, job.TemplateId, out jobNotificationConfigs, streamConfig?.StreamTags ?? null))
 			{
-				_logger.LogInformation("No configurations found for '{JobId}' in stream '{StreamId}' and for template '{TemplateId}'", job.Id, job.StreamId, job.TemplateId);
+				_logger.LogInformation("No configurations found for '{JobId}' in stream '{StreamInformation}' and for template '{TemplateId}'", job.Id, streamInformation, job.TemplateId);
 				return;
 			}
 
+			_logger.LogInformation("{NumConfigurations} configurations found for '{JobId}' in stream '{StreamInformation}' and for template '{TemplateId}'", jobNotificationConfigs.Count, job.Id, streamInformation, job.TemplateId);
+
 			string branch = GetBranchFromStream(_buildConfig.CurrentValue, job.StreamId);
 
 			// Check if the Setup step failed to notify that the job can't be executed
