Filter in Interrupted in Failed outcome

#rnx

diff --git a/Engine/Source/Programs/Horde/HordeDashboard/plugins/test-automation/testAutomationCommon.tsx b/Engine/Source/Programs/Horde/HordeDashboard/plugins/test-automation/testAutomationCommon.tsx
--- a/Engine/Source/Programs/Horde/HordeDashboard/plugins/test-automation/testAutomationCommon.tsx
+++ b/Engine/Source/Programs/Horde/HordeDashboard/plugins/test-automation/testAutomationCommon.tsx
@@ -522,6 +522,7 @@ class PhaseFiltersContext {
                     break;
                 case TestOutcome.Failure:
                     outcome.add(TestPhaseOutcome.Failed);
+                    outcome.add(TestPhaseOutcome.Interrupted);
                     break;
                 case TestOutcome.Success:
                     outcome.add(TestPhaseOutcome.Success);
