Horde: Resolve workspace materializer through workspace config/type

Also add a factory for the PerforceMaterializer.

diff --git a/Engine/Source/Programs/Horde/Drivers/JobDriver.Tests/Execution/FakeWorkspaceMaterializer.cs b/Engine/Source/Programs/Horde/Drivers/JobDriver.Tests/Execution/FakeWorkspaceMaterializer.cs
--- a/Engine/Source/Programs/Horde/Drivers/JobDriver.Tests/Execution/FakeWorkspaceMaterializer.cs
+++ b/Engine/Source/Programs/Horde/Drivers/JobDriver.Tests/Execution/FakeWorkspaceMaterializer.cs
@@ -19,7 +19,7 @@ class FakeWorkspaceMaterializer : IWorkspaceMaterializer
 	public DirectoryReference DirectoryPath => _rootDir;
 
 	public string Identifier => "fakeWorkspaceIdentifier";
-
+	public string Name => "FakeWorkspaceMaterializer";
 	public IReadOnlyDictionary<string, string> EnvironmentVariables { get; } = new Dictionary<string, string>();
 
 	public bool IsPerforceWorkspace => false;
diff --git a/Engine/Source/Programs/Horde/Drivers/JobDriver.Tests/WorkspaceExecutorTests.cs b/Engine/Source/Programs/Horde/Drivers/JobDriver.Tests/WorkspaceExecutorTests.cs
new file mode 100644
--- /dev/null
+++ b/Engine/Source/Programs/Horde/Drivers/JobDriver.Tests/WorkspaceExecutorTests.cs
@@ -0,0 +1,87 @@
+// Copyright Epic Games, Inc. All Rights Reserved.
+
+using System;
+using System.Collections.Generic;
+using System.Linq;
+using System.Threading;
+using System.Threading.Tasks;
+using EpicGames.Core;
+using HordeCommon.Rpc.Messages;
+using HordeCommon.Rpc.Tasks;
+using JobDriver.Execution;
+using Microsoft.Extensions.Logging;
+using Microsoft.VisualStudio.TestTools.UnitTesting;
+using OpenTelemetry.Trace;
+
+namespace JobDriver.Tests;
+
+internal class MaterializerStub(string name) : IWorkspaceMaterializer
+{
+	public void Dispose() { }
+	public string Name { get; } = name;
+	public DirectoryReference DirectoryPath { get; } = new("not-set");
+	public string Identifier { get; } = "not-set";
+	public IReadOnlyDictionary<string, string> EnvironmentVariables { get; } = new Dictionary<string, string>();
+	public bool IsPerforceWorkspace { get; } = false;
+	public Task SyncAsync(int changeNum, int shelveChangeNum, SyncOptions options, CancellationToken cancellationToken)
+	{
+		throw new NotImplementedException();
+	}
+	public Task FinalizeAsync(CancellationToken cancellationToken)
+	{
+		throw new NotImplementedException();
+	}
+}
+
+internal class MaterializerFactoryStub(IWorkspaceMaterializer wm) : IWorkspaceMaterializerFactory
+{
+	public Task<IWorkspaceMaterializer?> CreateMaterializerAsync(string name, RpcAgentWorkspace raw, DirectoryReference dir, bool forAutoSdk = false, CancellationToken ct = default)
+	{
+		return Task.FromResult(name.Equals(wm.Name, StringComparison.OrdinalIgnoreCase) ? wm : null);
+	}
+}
+
+[TestClass]
+public class WorkspaceExecutorTests
+{
+	private static Tracer NoOpTracer { get; } = TracerProvider.Default.GetTracer("NoOp");
+	private readonly DirectoryReference _tempDir = new ("unused-temp-dir");
+	private const string ManagedWorkspace = ManagedWorkspaceMaterializer.TypeName;
+	private const string Perforce = PerforceMaterializer.TypeName;
+	private const string FooBar = "FooBar";
+	
+	[TestMethod]
+	public void Factory()
+	{
+		Assert.ThrowsException<ArgumentException>(() => GetMaterializer(new RpcJobOptions(), new RpcAgentWorkspace(), []));
+		
+		AssertWm(ManagedWorkspace, "", "");
+		AssertWm(ManagedWorkspace, "mAnaGedWorkSpace", "");
+		AssertWm(FooBar, "fOObAr", "");
+		AssertWm(ManagedWorkspace, "", "fOObAr");
+		AssertWm(FooBar, "", "name=fOObAr");
+		AssertWm(ManagedWorkspace, "", "name=mAnaGedWorkSpace");
+		AssertWm(Perforce, "perforce", "name=fOObAr");
+	}
+	
+	private void AssertWm(string expected, string jobMaterializer, string method)
+	{
+		RpcJobOptions rjo = new () { WorkspaceMaterializer = jobMaterializer };
+		RpcAgentWorkspace raw = new () { Method = method };
+		IWorkspaceMaterializer wm = GetMaterializer(rjo, raw, [FooBar, ManagedWorkspace, Perforce]);
+		Assert.AreEqual(expected, wm.Name);
+	}
+	
+	private IWorkspaceMaterializer GetMaterializer(RpcJobOptions rjo, RpcAgentWorkspace raw, string[] names)
+	{
+		using ILoggerFactory loggerFactory = LoggerFactory.Create(builder =>
+		{
+			builder.SetMinimumLevel(LogLevel.Debug);
+			builder.AddSimpleConsole(options => { options.SingleLine = true; });
+		});
+		
+		List<IWorkspaceMaterializerFactory> factories = [..names.Select(name => new MaterializerFactoryStub(new MaterializerStub(name)))];
+		WorkspaceExecutorFactory wef = new (factories, NoOpTracer, loggerFactory);
+		return wef.CreateMaterializerAsync(rjo, raw, _tempDir, false, CancellationToken.None).GetAwaiter().GetResult();
+	}
+}
diff --git a/Engine/Source/Programs/Horde/Drivers/JobDriver/DriverApp.cs b/Engine/Source/Programs/Horde/Drivers/JobDriver/DriverApp.cs
--- a/Engine/Source/Programs/Horde/Drivers/JobDriver/DriverApp.cs
+++ b/Engine/Source/Programs/Horde/Drivers/JobDriver/DriverApp.cs
@@ -81,6 +81,7 @@ namespace JobDriver
 			services.AddSingleton<IJobExecutorFactory, TestExecutorFactory>();
 
 			services.AddSingleton<IWorkspaceMaterializerFactory, ManagedWorkspaceMaterializerFactory>();
+			services.AddSingleton<IWorkspaceMaterializerFactory, PerforceMaterializerFactory>();
 
 			services.AddCommandsFromAssembly(Assembly.GetExecutingAssembly());
 		}
diff --git a/Engine/Source/Programs/Horde/Drivers/JobDriver/Execution/ManagedWorkspaceMaterializer.cs b/Engine/Source/Programs/Horde/Drivers/JobDriver/Execution/ManagedWorkspaceMaterializer.cs
--- a/Engine/Source/Programs/Horde/Drivers/JobDriver/Execution/ManagedWorkspaceMaterializer.cs
+++ b/Engine/Source/Programs/Horde/Drivers/JobDriver/Execution/ManagedWorkspaceMaterializer.cs
@@ -19,7 +19,7 @@ public sealed class ManagedWorkspaceMaterializer : IWorkspaceMaterializer
 	/// <summary>
 	/// Name of this materializer
 	/// </summary>
-	public const string Name = "ManagedWorkspace";
+	public const string TypeName = "ManagedWorkspace";
 
 	private readonly RpcAgentWorkspace _agentWorkspace;
 	private readonly bool _useCacheFile;
@@ -30,6 +30,9 @@ public sealed class ManagedWorkspaceMaterializer : IWorkspaceMaterializer
 
 	/// <inheritdoc/>
 	public DirectoryReference DirectoryPath => _workspace.WorkspaceDir;
+	
+	/// <inheritdoc/>
+	public string Name => TypeName;
 
 	/// <inheritdoc/>
 	public string Identifier => _agentWorkspace.Identifier;
@@ -170,7 +173,7 @@ class ManagedWorkspaceMaterializerFactory : IWorkspaceMaterializerFactory
 	/// <inheritdoc/>
 	public async Task<IWorkspaceMaterializer?> CreateMaterializerAsync(string name, RpcAgentWorkspace workspaceInfo, DirectoryReference workspaceDir, bool forAutoSdk, CancellationToken cancellationToken)
 	{
-		if (name.Equals(ManagedWorkspaceMaterializer.Name, StringComparison.OrdinalIgnoreCase))
+		if (name.Equals(ManagedWorkspaceMaterializer.TypeName, StringComparison.OrdinalIgnoreCase))
 		{
 			Tracer tracer = _serviceProvider.GetRequiredService<Tracer>();
 			ILogger<ManagedWorkspaceMaterializer> logger = _serviceProvider.GetRequiredService<ILogger<ManagedWorkspaceMaterializer>>();
diff --git a/Engine/Source/Programs/Horde/Drivers/JobDriver/Execution/PerforceMaterializer.cs b/Engine/Source/Programs/Horde/Drivers/JobDriver/Execution/PerforceMaterializer.cs
--- a/Engine/Source/Programs/Horde/Drivers/JobDriver/Execution/PerforceMaterializer.cs
+++ b/Engine/Source/Programs/Horde/Drivers/JobDriver/Execution/PerforceMaterializer.cs
@@ -6,6 +6,7 @@ using EpicGames.Core;
 using EpicGames.Perforce;
 using HordeCommon.Rpc.Messages;
 using JobDriver.Utility;
+using Microsoft.Extensions.DependencyInjection;
 using Microsoft.Extensions.Logging;
 using OpenTelemetry.Trace;
 
@@ -40,7 +41,7 @@ public sealed class PerforceMaterializer : IWorkspaceMaterializer
 	/// <summary>
 	/// Name of this materializer
 	/// </summary>
-	public const string Name = "Perforce";
+	public const string TypeName = "Perforce";
 	
 	internal enum TransactionStatus { Clean = 0, Dirty = 1 }
 	internal record State(
@@ -54,7 +55,10 @@ public sealed class PerforceMaterializer : IWorkspaceMaterializer
 	
     /// <inheritdoc/>
 	public DirectoryReference DirectoryPath { get; }
-	
+
+    /// <inheritdoc/>
+	public string Name => TypeName;
+
 	/// <inheritdoc/>
 	public string Identifier { get; }
 	
@@ -357,4 +361,20 @@ public sealed class PerforceMaterializer : IWorkspaceMaterializer
 	{
 		return LoadStateAsync(cancellationToken);
 	}
+}
+
+class PerforceMaterializerFactory(IServiceProvider serviceProvider) : IWorkspaceMaterializerFactory
+{
+	/// <inheritdoc/>
+	public Task<IWorkspaceMaterializer?> CreateMaterializerAsync(string name, RpcAgentWorkspace workspaceInfo, DirectoryReference workspaceDir, bool forAutoSdk, CancellationToken cancellationToken)
+	{
+		if (name.Equals(PerforceMaterializer.TypeName, StringComparison.OrdinalIgnoreCase))
+		{
+			Tracer tracer = serviceProvider.GetRequiredService<Tracer>();
+			ILogger<PerforceMaterializer> logger = serviceProvider.GetRequiredService<ILogger<PerforceMaterializer>>();
+			PerforceMaterializerOptions options = new (workspaceDir.FullName, workspaceInfo);
+			return Task.FromResult<IWorkspaceMaterializer?>(new PerforceMaterializer(options, tracer, logger));
+		}
+		return Task.FromResult<IWorkspaceMaterializer?>(null);
+	}
 }
\ No newline at end of file
diff --git a/Engine/Source/Programs/Horde/Drivers/JobDriver/Execution/WorkspaceExecutor.cs b/Engine/Source/Programs/Horde/Drivers/JobDriver/Execution/WorkspaceExecutor.cs
--- a/Engine/Source/Programs/Horde/Drivers/JobDriver/Execution/WorkspaceExecutor.cs
+++ b/Engine/Source/Programs/Horde/Drivers/JobDriver/Execution/WorkspaceExecutor.cs
@@ -1,10 +1,13 @@
 // Copyright Epic Games, Inc. All Rights Reserved.
 
+using System.Collections.Specialized;
 using System.Text.RegularExpressions;
+using System.Web;
 using EpicGames.Core;
 using JobDriver.Utility;
 using Horde.Common.Rpc;
 using HordeCommon.Rpc.Messages;
+using HordeCommon.Rpc.Tasks;
 using Microsoft.Extensions.Logging;
 using OpenTelemetry.Trace;
 
@@ -163,12 +166,10 @@ namespace JobDriver.Execution
 			IWorkspaceMaterializer? autoSdkMaterializer = null;
 			try
 			{
-				string name = String.IsNullOrEmpty(options.JobOptions.WorkspaceMaterializer) ? ManagedWorkspaceMaterializer.Name : options.JobOptions.WorkspaceMaterializer;
-
-				workspaceMaterializer = await CreateMaterializerAsync(name, workspaceInfo, options.WorkingDir, forAutoSdk: false, cancellationToken);
+				workspaceMaterializer = await CreateMaterializerAsync(options.JobOptions, workspaceInfo, options.WorkingDir, forAutoSdk: false, cancellationToken);
 				if (autoSdkWorkspaceInfo != null)
 				{
-					autoSdkMaterializer = await CreateMaterializerAsync(name, autoSdkWorkspaceInfo, options.WorkingDir, forAutoSdk: true, cancellationToken);
+					autoSdkMaterializer = await CreateMaterializerAsync(options.JobOptions, autoSdkWorkspaceInfo, options.WorkingDir, forAutoSdk: true, cancellationToken);
 				}
 
 				return new WorkspaceExecutor(options, workspaceMaterializer, autoSdkMaterializer, _tracer, _loggerFactory.CreateLogger<WorkspaceExecutor>());
@@ -181,8 +182,22 @@ namespace JobDriver.Execution
 			}
 		}
 
-		async Task<IWorkspaceMaterializer> CreateMaterializerAsync(string name, RpcAgentWorkspace workspaceInfo, DirectoryReference workspaceDir, bool forAutoSdk, CancellationToken cancellationToken)
+		internal async Task<IWorkspaceMaterializer> CreateMaterializerAsync(RpcJobOptions jobOptions, RpcAgentWorkspace workspaceInfo, DirectoryReference workspaceDir, bool forAutoSdk, CancellationToken cancellationToken)
 		{
+			const string DefaultMaterializer = ManagedWorkspaceMaterializer.TypeName;
+			string name = DefaultMaterializer;
+			
+			if (!String.IsNullOrEmpty(workspaceInfo.Method))
+			{
+				NameValueCollection nameValues = HttpUtility.ParseQueryString(workspaceInfo.Method);
+				name = nameValues["name"] ?? name;
+			}
+			
+			if (!String.IsNullOrEmpty(jobOptions.WorkspaceMaterializer))
+			{
+				name = jobOptions.WorkspaceMaterializer;
+			}
+			
 			foreach (IWorkspaceMaterializerFactory materializerFactory in _materializerFactories)
 			{
 				IWorkspaceMaterializer? materializer = await materializerFactory.CreateMaterializerAsync(name, workspaceInfo, workspaceDir, forAutoSdk, cancellationToken);
diff --git a/Engine/Source/Programs/Horde/Drivers/JobDriver/Execution/WorkspaceMaterializer.cs b/Engine/Source/Programs/Horde/Drivers/JobDriver/Execution/WorkspaceMaterializer.cs
--- a/Engine/Source/Programs/Horde/Drivers/JobDriver/Execution/WorkspaceMaterializer.cs
+++ b/Engine/Source/Programs/Horde/Drivers/JobDriver/Execution/WorkspaceMaterializer.cs
@@ -53,6 +53,11 @@ public interface IWorkspaceMaterializer : IDisposable
 	/// </summary>
 	public const int LatestChangeNumber = -2;
 
+	/// <summary>
+	/// Name of this materializer
+	/// </summary>
+	string Name { get; }
+	
 	/// <summary>
 	/// Path to local file system directory where files from changelist are materialized
 	/// </summary>
