[Horde][BuildHealth][Server] Expand streams/{streamId}/changes api to support batch semantics. It can take a manually specified set of changes to obtain.

diff --git a/Engine/Source/Programs/Horde/HordeServer.Tests.Shared/DatabaseIntegrationTest.cs b/Engine/Source/Programs/Horde/HordeServer.Tests.Shared/DatabaseIntegrationTest.cs
--- a/Engine/Source/Programs/Horde/HordeServer.Tests.Shared/DatabaseIntegrationTest.cs
+++ b/Engine/Source/Programs/Horde/HordeServer.Tests.Shared/DatabaseIntegrationTest.cs
@@ -117,6 +117,7 @@ namespace HordeServer.Tests
 			lock (s_lockObject)
 			{
 				s_mongoDbRunner?.Dispose();
+				s_mongoDbRunner = null;
 			}
 		}
 	}
diff --git a/Engine/Source/Programs/Horde/Plugins/Build/HordeServer.Build.Tests/BuildTestSetup.cs b/Engine/Source/Programs/Horde/Plugins/Build/HordeServer.Build.Tests/BuildTestSetup.cs
--- a/Engine/Source/Programs/Horde/Plugins/Build/HordeServer.Build.Tests/BuildTestSetup.cs
+++ b/Engine/Source/Programs/Horde/Plugins/Build/HordeServer.Build.Tests/BuildTestSetup.cs
@@ -99,6 +99,7 @@ namespace HordeServer.Tests
 		public TestDataService TestDataService => ServiceProvider.GetRequiredService<TestDataService>();
 		public ComputeService ComputeService => ServiceProvider.GetRequiredService<ComputeService>();
 
+		public StreamsController StreamsController => GetStreamsController();
 		public JobsController JobsController => GetJobsController();
 		public AgentsController AgentsController => GetAgentsController();
 		public PoolsController PoolsController => GetPoolsController();
@@ -198,6 +199,13 @@ namespace HordeServer.Tests
 			return Fixture.CreateAsync(ConfigService, GraphCollection, TemplateCollection, JobService, AgentService, PluginCollection, ServerSettings);
 		}
 
+		private StreamsController GetStreamsController()
+		{
+			StreamsController streamCtrl = ActivatorUtilities.CreateInstance<StreamsController>(ServiceProvider);
+			streamCtrl.ControllerContext = GetControllerContext();
+			return streamCtrl;
+		}
+
 		private JobsController GetJobsController()
 		{
 			JobsController jobsCtrl = ActivatorUtilities.CreateInstance<JobsController>(ServiceProvider);
diff --git a/Engine/Source/Programs/Horde/Plugins/Build/HordeServer.Build.Tests/Streams/StreamCollectionTests.cs b/Engine/Source/Programs/Horde/Plugins/Build/HordeServer.Build.Tests/Streams/StreamCollectionTests.cs
--- a/Engine/Source/Programs/Horde/Plugins/Build/HordeServer.Build.Tests/Streams/StreamCollectionTests.cs
+++ b/Engine/Source/Programs/Horde/Plugins/Build/HordeServer.Build.Tests/Streams/StreamCollectionTests.cs
@@ -1,8 +1,12 @@
 // Copyright Epic Games, Inc. All Rights Reserved.
 
+using EpicGames.Horde.Commits;
 using EpicGames.Horde.Jobs.Templates;
 using EpicGames.Horde.Streams;
 using HordeServer.Streams;
+using HordeServer.Users;
+using Microsoft.AspNetCore.Mvc;
+using Microsoft.Extensions.Primitives;
 
 namespace HordeServer.Tests.Streams
 {
@@ -29,5 +33,86 @@ namespace HordeServer.Tests.Streams
 			config.Templates.Add(new TemplateRefConfig { Id = new TemplateId("bar") });
 			HordeServer.Streams.StreamCollection.Validate(_streamId, config);
 		}
+
+		[TestMethod]
+		public async Task ValidateGetChangesAsync()
+		{
+			Fixture fixture = await CreateFixtureAsync();
+			IUser jerry = await UserCollection.FindOrAddUserByLoginAsync("Jerry");
+
+			int firstChangelist = 112233;
+			int secondChangelist = 112234;
+			int thirdChangelist = 112235;
+			int invalidChangelist = 112236;
+
+			PerforceService.AddChange(fixture.StreamId, firstChangelist, jerry, "Edited file1.cpp.", new[] { "file.cpp" });
+			PerforceService.AddChange(fixture.StreamId, secondChangelist, jerry, "Added file2.cpp.", new[] { "file2.cpp" });
+			PerforceService.AddChange(fixture.StreamId, thirdChangelist, jerry, "Edited file1.cpp.", new[] { "file1.cpp" });
+
+			CommitId firstCommitId = new CommitId(firstChangelist.ToString());
+			CommitId secondCommitId = new CommitId(secondChangelist.ToString());
+			CommitId thirdCommitId = new CommitId(thirdChangelist.ToString());
+			CommitId invalidCommitId = new CommitId(invalidChangelist.ToString());
+
+			// Single element edge case
+			{
+				ActionResult<List<object>> results = await StreamsController.GetChangesAsync(fixture.StreamId, null, null, [firstCommitId]);
+
+				Assert.IsNotNull(results.Value);
+
+				List<GetCommitResponse> castedValue = [.. results.Value.Cast<GetCommitResponse>()];
+				Assert.IsTrue(castedValue.Count == 1);
+				Assert.IsNotNull(castedValue.SingleOrDefault(c => c.Id.Name == firstChangelist.ToString()));
+			}
+
+			// Multi element nominal case
+			{
+				ActionResult<List<object>> results = await StreamsController.GetChangesAsync(fixture.StreamId, null, null, [firstCommitId, secondCommitId]);
+
+				Assert.IsNotNull(results.Value);
+
+				List<GetCommitResponse> castedValue = [.. results.Value.Cast<GetCommitResponse>()];
+				Assert.IsTrue(castedValue.Count == 2);
+				Assert.IsNotNull(castedValue.SingleOrDefault(c => c.Id.Name == firstChangelist.ToString()));
+				Assert.IsNotNull(castedValue.SingleOrDefault(c => c.Id.Name == secondChangelist.ToString()));
+			}
+
+			// Partial success case
+			{
+				StreamsController httpContextPreservedStreamsController = StreamsController;
+				ActionResult<List<object>> results = await httpContextPreservedStreamsController.GetChangesAsync(fixture.StreamId, null, null, [firstCommitId, invalidCommitId]);
+
+				Assert.IsNotNull(results.Value);
+
+				List<GetCommitResponse> castedValue = [.. results.Value.Cast<GetCommitResponse>()];
+				Assert.IsTrue(castedValue.Count == 1);
+				Assert.IsNotNull(castedValue.SingleOrDefault(c => c.Id.Name == firstChangelist.ToString()));
+				Assert.IsNull(castedValue.SingleOrDefault(c => c.Id.Name == invalidChangelist.ToString()));
+				Assert.IsTrue(httpContextPreservedStreamsController.Response.Headers.ContainsKey(StreamsController.MissingCommitsHeader));
+
+				// Assert contents
+				StringValues stringValues = httpContextPreservedStreamsController.Response.Headers[StreamsController.MissingCommitsHeader];
+				Assert.IsTrue(stringValues.Contains(invalidChangelist.ToString()));
+			}
+
+			// Empty edge case
+			{
+				ActionResult<List<object>> results = await StreamsController.GetChangesAsync(fixture.StreamId, null, null, []);
+
+				Assert.IsNotNull(results.Value);
+
+				List<GetCommitResponse> castedValue = [.. results.Value.Cast<GetCommitResponse>()];
+				Assert.IsTrue(castedValue.Count == 0);
+			}
+
+			// Bad request test case when providing too many arugments.
+			{
+				ActionResult<List<object>> results = await StreamsController.GetChangesAsync(fixture.StreamId, firstCommitId, secondCommitId, [thirdCommitId]);
+				BadRequestObjectResult? result = results.Result as BadRequestObjectResult;
+
+				Assert.IsNotNull(result);
+				Assert.AreEqual(StreamsController.BadRequest().StatusCode, result.StatusCode);
+			}
+		}
 	}
 }
diff --git a/Engine/Source/Programs/Horde/Plugins/Build/HordeServer.Build.Tests/Stubs/Services/PerforceServiceStub.cs b/Engine/Source/Programs/Horde/Plugins/Build/HordeServer.Build.Tests/Stubs/Services/PerforceServiceStub.cs
--- a/Engine/Source/Programs/Horde/Plugins/Build/HordeServer.Build.Tests/Stubs/Services/PerforceServiceStub.cs
+++ b/Engine/Source/Programs/Horde/Plugins/Build/HordeServer.Build.Tests/Stubs/Services/PerforceServiceStub.cs
@@ -9,6 +9,7 @@ using HordeServer.VersionControl.Perforce;
 using HordeServer.Streams;
 using HordeServer.Users;
 using HordeServer.Utilities;
+using EpicGames.Perforce;
 
 namespace HordeServer.Tests.Stubs.Services
 {
@@ -44,7 +45,7 @@ namespace HordeServer.Tests.Stubs.Services
 			public UserId OwnerId { get; }
 			public string Description { get; }
 			public string BasePath => throw new NotImplementedException();
-			public DateTime DateUtc => throw new NotImplementedException();
+			public DateTime DateUtc => DateTime.UtcNow;
 			public List<string> Files { get; }
 
 			public Commit(StreamId streamId, int number, int originalChange, UserId authorId, UserId ownerId, string description, List<string> files)
@@ -177,6 +178,11 @@ namespace HordeServer.Tests.Stubs.Services
 			List<ICommit> results = new List<ICommit>();
 			foreach (int changeNumber in changeNumbers)
 			{
+				if(!Changes.ContainsKey(stream.Id) || !Changes[stream.Id].ContainsKey(changeNumber))
+				{
+					throw new PerforceException($"PerforceServiceStub is simulating a changelist not found exception {changeNumber}");
+				}
+
 				results.Add(Changes[stream.Id][changeNumber]);
 			}
 			return Task.FromResult(results);
diff --git a/Engine/Source/Programs/Horde/Plugins/Build/HordeServer.Build/Streams/StreamsController.cs b/Engine/Source/Programs/Horde/Plugins/Build/HordeServer.Build/Streams/StreamsController.cs
--- a/Engine/Source/Programs/Horde/Plugins/Build/HordeServer.Build/Streams/StreamsController.cs
+++ b/Engine/Source/Programs/Horde/Plugins/Build/HordeServer.Build/Streams/StreamsController.cs
@@ -1,5 +1,6 @@
 // Copyright Epic Games, Inc. All Rights Reserved.
 
+using System.Collections.Concurrent;
 using EpicGames.Core;
 using EpicGames.Horde.Commits;
 using EpicGames.Horde.Jobs;
@@ -7,6 +8,7 @@ using EpicGames.Horde.Jobs.Templates;
 using EpicGames.Horde.Projects;
 using EpicGames.Horde.Streams;
 using EpicGames.Horde.Users;
+using EpicGames.Perforce;
 using HordeServer.Auditing;
 using HordeServer.Commits;
 using HordeServer.Jobs;
@@ -17,6 +19,7 @@ using HordeServer.Server;
 using HordeServer.Users;
 using HordeServer.Utilities;
 using Microsoft.AspNetCore.Authorization;
+using Microsoft.AspNetCore.Http;
 using Microsoft.AspNetCore.Mvc;
 using Microsoft.Extensions.Options;
 using OpenTelemetry.Trace;
@@ -122,7 +125,7 @@ namespace HordeServer.Streams
 			}
 
 			IStream? stream = await _streamCollection.GetAsync(streamConfig.Id, cancellationToken);
-			if(stream == null)
+			if (stream == null)
 			{
 				return NotFound(streamId);
 			}
@@ -181,21 +184,35 @@ namespace HordeServer.Streams
 			return new GetStreamResponse(stream, apiTemplateRefs);
 		}
 
+		/// <summary>
+		/// Headers key utilized for expressing missing commits in <see cref="StreamsController.GetChangesAsync"/>.
+		/// </summary>
+		public const string MissingCommitsHeader = "X-Missing-Commits";
+
+		private const int MaxConcurrentSCMChangeRequests = 32;
+
 		/// <summary>
 		/// Gets a list of changes for a stream
 		/// </summary>
 		/// <param name="streamId">The stream id</param>
 		/// <param name="min">The starting changelist number</param>
 		/// <param name="max">The ending changelist number</param>
+		/// <param name="changeNumbers">
+		///		Array of specific changelist IDs to retrieve for the stream.
+		///		If provided, this parameter is mutually exclusive with <paramref name="min"/>, <paramref name="max"/>, <paramref name="tags"/>,
+		///		A <see cref="BadRequestObjectResult"/> will return if any of the above are provided alongside changeNumbers.
+		///		<paramref name="results"/> will be ignored.
+		/// </param>
 		/// <param name="results">Number of results to return</param>
 		/// <param name="tags">Tags to filter the changes returned</param>
 		/// <param name="filter">The filter to apply to the results</param>
 		/// <param name="cancellationToken">Cancellation token for the operation</param>
-		/// <returns>Http result code</returns>
+		/// Returns 200 OK with list of filtered commits. Partial failures for batch retrieval are indicated via the <see cref="MissingCommitsHeader"/> header.
+		/// <remarks>Batch request mode executed with <paramref name="changeNumbers"/> will store any commits that were not found or that threw exceptions in the <see cref="MissingCommitsHeader"/> response header.</remarks>
 		[HttpGet]
 		[Route("/api/v1/streams/{streamId}/changes")]
 		[ProducesResponseType(typeof(List<GetCommitResponse>), 200)]
-		public async Task<ActionResult<List<object>>> GetChangesAsync(StreamId streamId, [FromQuery] CommitId? min = null, [FromQuery] CommitId? max = null, [FromQuery] int results = 50, [FromQuery] string? tags = null, PropertyFilter? filter = null, CancellationToken cancellationToken = default)
+		public async Task<ActionResult<List<object>>> GetChangesAsync(StreamId streamId, [FromQuery] CommitId? min = null, [FromQuery] CommitId? max = null, [FromQuery] CommitId[]? changeNumbers = null, [FromQuery] int results = 50, [FromQuery] string? tags = null, PropertyFilter? filter = null, CancellationToken cancellationToken = default)
 		{
 			StreamConfig? streamConfig;
 			if (!_buildConfig.Value.TryGetStream(streamId, out streamConfig))
@@ -206,6 +223,10 @@ namespace HordeServer.Streams
 			{
 				return Forbid(StreamAclAction.ViewChanges, streamId);
 			}
+			if (changeNumbers != null && (min != null || max != null || tags != null))
+			{
+				return BadRequest("The changeNumbers option is mutually exclusive with min/max/tag options.");
+			}
 
 			List<CommitTag>? commitTags = null;
 			if (tags != null)
@@ -213,7 +234,49 @@ namespace HordeServer.Streams
 				commitTags = tags.Split(';', StringSplitOptions.RemoveEmptyEntries).Select(x => new CommitTag(x)).ToList();
 			}
 
-			List<ICommit> commits = await _commitService.GetCollection(streamConfig).FindAsync(min, max, results, commitTags, cancellationToken).ToListAsync(cancellationToken);
+			List<ICommit>? commits = null;
+
+			if (changeNumbers != null) // Batch commit request path.
+			{
+				commits = new List<ICommit>(changeNumbers.Length);
+				ConcurrentBag<string> failedCommits = [];
+
+				// This is a pre-emptive measure as to not inundate our SCM with too many concurrent requests.
+				using SemaphoreSlim maxSCMActiveRequests = new SemaphoreSlim(MaxConcurrentSCMChangeRequests);
+				Task<ICommit?>[] commitServiceTasks = [.. changeNumbers.Select(async id =>
+				{
+					await maxSCMActiveRequests.WaitAsync(cancellationToken);
+					try
+					{
+						return await _commitService.GetCollection(streamConfig).GetAsync(id, cancellationToken);
+					}
+					catch (PerforceException) // Swallow PerforceExceptions, and store failed commits.
+					{
+						failedCommits.Add(id.ToString());
+
+						return null;
+					}
+					finally
+					{
+						maxSCMActiveRequests.Release();
+					}
+				})];
+
+				ICommit?[] commitServiceResults = await Task.WhenAll(commitServiceTasks);
+				commits = [.. commitServiceResults.Where(c => c != null).Select(c => c!)];
+
+				// Store any failed commits into the response header.
+				if (!failedCommits.IsEmpty)
+				{
+					Response.Headers.Append(MissingCommitsHeader, String.Join(",", failedCommits));
+				}
+
+				commits = [.. commits.OrderByDescending(x => x.Id)];
+			}
+			else // Change range & tag request path.
+			{
+				commits = await _commitService.GetCollection(streamConfig).FindAsync(min, max, results, commitTags, cancellationToken).ToListAsync(cancellationToken);
+			}
 
 			List<GetCommitResponse> responses = new List<GetCommitResponse>();
 			foreach (ICommit commit in commits)
@@ -423,7 +486,7 @@ namespace HordeServer.Streams
 			}
 
 			IStream? stream = await _streamCollection.GetAsync(streamId);
-			if(stream == null)
+			if (stream == null)
 			{
 				return NotFound(streamId);
 			}
