[Horde][BuildHealth][Dashboard] Change Headers anchored by submit times

diff --git a/Engine/Source/Programs/Horde/HordeDashboard/src/backend/Api.ts b/Engine/Source/Programs/Horde/HordeDashboard/src/backend/Api.ts
--- a/Engine/Source/Programs/Horde/HordeDashboard/src/backend/Api.ts
+++ b/Engine/Source/Programs/Horde/HordeDashboard/src/backend/Api.ts
@@ -2582,10 +2582,12 @@ export type GetChangeSummaryResponse = {
 	/**  The description text */
 	description: string;
 
+	/** The submit date of the changelist */
+	dateUtc: Date | string;
+
 	/**  Information about the change author */
 	authorInfo: GetThinUserInfoResponse;
 
-
 }
 
 /**Information about a device attached to this agent */
diff --git a/Engine/Source/Programs/Horde/HordeDashboard/src/backend/Backend.ts b/Engine/Source/Programs/Horde/HordeDashboard/src/backend/Backend.ts
--- a/Engine/Source/Programs/Horde/HordeDashboard/src/backend/Backend.ts
+++ b/Engine/Source/Programs/Horde/HordeDashboard/src/backend/Backend.ts
@@ -585,7 +585,7 @@ export class Backend {
         // the disadvantage of the POST method using fetch is the lack of progression indicator for large download size
         // this method wait for the whole data to be downloaded before passing the data blob to the new window
         return new Promise<object>(() => {
-            this._fetch.post(url, request, {responseBlob: true}).then((response) => {
+            this._fetch.post(url, request, { responseBlob: true }).then((response) => {
                 // Create a temp link and activate it to save the incoming data
                 const blobURL = URL.createObjectURL(response.data);
                 const link = document.createElement('a');
@@ -829,16 +829,23 @@ export class Backend {
 
     }
 
-    getCommit(streamId?: string, changelist?: number): Promise<ChangeSummaryData | undefined> {
+    getCommit(streamId?: string, changelist?: number, filter?: string): Promise<ChangeSummaryData | undefined> {
 
         return new Promise<ChangeSummaryData | undefined>((resolve) => {
+            const params: any = {};
 
             if (!streamId || !changelist) {
                 resolve(undefined);
                 return;
             }
 
-            this._fetch.get(`/api/v1/streams/${streamId}/changes/${changelist}`).then((value) => {
+            if (filter) {
+                params.filter = filter
+            }
+
+            this._fetch.get(`/api/v1/streams/${streamId}/changes/${changelist}`, {
+                params: params
+            }).then((value) => {
                 resolve(value.data as ChangeSummaryData);
             }).catch(reason => {
                 console.error(reason);
@@ -847,7 +854,7 @@ export class Backend {
         });
     }
 
-    getChangeSummaries(streamId: string, minChange?: number, maxChange?: number, maxResults = 100): Promise<ChangeSummaryData[]> {
+    getChangeSummaries(streamId: string, minChange?: number, maxChange?: number, changeNumbers? :number[], filter?: string, maxResults = 100): Promise<ChangeSummaryData[]> {
 
         return new Promise<ChangeSummaryData[]>((resolve, reject) => {
 
@@ -863,6 +870,15 @@ export class Backend {
                 params.max = maxChange;
             }
 
+            if (filter) {
+                params.filter = filter;
+            }
+
+            if(changeNumbers)
+            {
+                params.changeNumbers = changeNumbers;
+            }
+
             this._fetch.get(`/api/v1/streams/${streamId}/changes`, {
                 params: params
             }).then((value) => {
@@ -1705,7 +1721,7 @@ export class Backend {
         });
     }
 
-    checkPreflightConfig(shelvedChange: number, cluster?: string ): Promise<PreflightConfigResponse> {
+    checkPreflightConfig(shelvedChange: number, cluster?: string): Promise<PreflightConfigResponse> {
 
         return new Promise<PreflightConfigResponse>((resolve, reject) => {
             this._fetch.post(`/api/v1/server/preflightconfig`, { shelvedChange: shelvedChange, cluster: cluster }).then((value) => {
diff --git a/Engine/Source/Programs/Horde/HordeDashboard/src/backend/CommitCache.ts b/Engine/Source/Programs/Horde/HordeDashboard/src/backend/CommitCache.ts
--- a/Engine/Source/Programs/Horde/HordeDashboard/src/backend/CommitCache.ts
+++ b/Engine/Source/Programs/Horde/HordeDashboard/src/backend/CommitCache.ts
@@ -81,14 +81,14 @@ export class CommitCache {
             }
 
             map!.set(c, {
-                number: 0, description: "", authorInfo : { id : "", email:"", name: ""}
+                number: 0, description: "", dateUtc: "", authorInfo : { id : "", email:"", name: ""}
             })
 
         })
 
         // new query
         try {
-            const commits = await backend.getChangeSummaries(streamId, minCL, maxCL, maxCL - minCL + 1);
+            const commits = await backend.getChangeSummaries(streamId, minCL, maxCL, undefined, undefined, maxCL - minCL + 1);
             commits.forEach(c => map!.set(c.number, c));
             this.setUpdated();
         } catch (reason) {
diff --git a/Engine/Source/Programs/Horde/HordeDashboard/src/components/buildhealth/stepoutcome/StepOutcomeDataHandler.ts b/Engine/Source/Programs/Horde/HordeDashboard/src/components/buildhealth/stepoutcome/StepOutcomeDataHandler.ts
--- a/Engine/Source/Programs/Horde/HordeDashboard/src/components/buildhealth/stepoutcome/StepOutcomeDataHandler.ts
+++ b/Engine/Source/Programs/Horde/HordeDashboard/src/components/buildhealth/stepoutcome/StepOutcomeDataHandler.ts
@@ -1,8 +1,8 @@
 // Copyright Epic Games, Inc. All Rights Reserved.
 
-import { GetJobResponse, IssueQuery, JobStreamQuery } from "horde/backend/Api";
+import { ChangeSummaryData, GetJobResponse, IssueQuery, JobStreamQuery } from "horde/backend/Api";
 import { PollBase } from "horde/backend/PollBase";
-import { StepOutcomeTable, StepOutcomeTableEntry } from "./StepOutcomeDataTypes";
+import { StepOutcomeTable, StepOutcomeTableEntry, UpdatedChangeSummaryData } from "./StepOutcomeDataTypes";
 import backend from "horde/backend";
 import { decodeStepKey, decodeTemplateKey, encodeStepNameFromStrings } from "../BuildHealthUtilities";
 import { action, computed, makeObservable, observable, runInAction } from "mobx";
@@ -90,6 +90,7 @@ export interface StepOutcomeFilters {
  */
 export class StepOutcomeDataHandler extends PollBase {
     private readonly STREAM_JOBS_FILTER = "id,streamId,name,change,templateId,state,createTime,updateTime,batches,preflightCommitId,preflightDescription,cancellationReason";
+    private readonly STREAM_CHANGES_FILTER = "id,dateUtc";
     private readonly MAX_JOB_COUNT_DEFAULT = 150;
 
     // #region -- Private Members -- 
@@ -104,6 +105,7 @@ export class StepOutcomeDataHandler extends PollBase {
 
     private activeFilter: StepOutcomeFilters;
     private jobData: GetJobResponse[] = [];
+    private changeData: Map<number, UpdatedChangeSummaryData> = new Map<number, UpdatedChangeSummaryData>();
     private stepOutcomeTableData: StepOutcomeTable = new StepOutcomeTable();
 
     //#endregion -- Private Members
@@ -176,6 +178,7 @@ export class StepOutcomeDataHandler extends PollBase {
         });
 
         const queries: { streamId: string; query: JobStreamQuery }[] = [];
+        const changeQueries: { streamId: string; query: any }[] = [];
 
         for (const streamId of streamList) {
             const jobsForStream = jobsByStream[streamId] ?? [];
@@ -187,17 +190,49 @@ export class StepOutcomeDataHandler extends PollBase {
                     modifiedAfter: modifiedAfter
                 };
                 queries.push({ streamId, query });
+                changeQueries.push({ streamId: streamId, query: { filter: this.STREAM_CHANGES_FILTER } });
             }
         }
 
-        // Parallelize backend calls
+        // Parallelize backend calls for stream jobs
         const results = await Promise.all(
             queries.map(({ streamId, query }) => backend.getStreamJobs(streamId, query))
         );
 
+        // Annotate the changes that we need to request in getChangeSummaries.
+        let jobHistoryCount: Map<string, number[]> = new Map<string, number[]>();
+
+        results.forEach((jobs, i) => {
+            const streamId = queries[i].streamId;
+
+            const newChanges = jobs
+                .map(job => job.change!)
+                .filter(change => !this.changeData.has(change));
+
+            if (newChanges.length > 0) {
+                if (!jobHistoryCount.has(streamId)) {
+                    jobHistoryCount.set(streamId, []);
+                } 
+                
+                jobHistoryCount.get(streamId)!.push(...newChanges);
+            }
+        });
 
         // Flatten all jobData
         this.jobData = results.flat();
+        
+        if (jobHistoryCount.size > 0) {
+            // Parallelize backend calls for change summaries
+            const changeResults = await Promise.all(
+                changeQueries.map(({ streamId, query }) => backend.getChangeSummaries(streamId, undefined, undefined, jobHistoryCount.get(streamId), query.filter))
+            );
+
+            changeResults.flat().map((value: UpdatedChangeSummaryData) => {
+                if (!this.changeData.has(value.id.order)) {
+                    this.changeData.set(value.id.order, value);
+                }
+            });
+        }
 
         for (const streamResult of results) {
             // Process all job responses
@@ -220,7 +255,7 @@ export class StepOutcomeDataHandler extends PollBase {
                         ) {
 
                             this.stepOutcomeTableData.addEntry(
-                                new StepOutcomeTableEntry(step, jobResponse.id, jobResponse.name, jobResponse.streamId, jobResponse.createTime.toString(), jobResponse.change)
+                                new StepOutcomeTableEntry(step, jobResponse.id, jobResponse.name, jobResponse.streamId, jobResponse.createTime.toString(), jobResponse.change), this.changeData.get(jobResponse.change!)
                             );
                         }
                     }
@@ -228,7 +263,7 @@ export class StepOutcomeDataHandler extends PollBase {
             }
         }
 
-        console.info(`Retrieved ${this.jobData.length} job instance results for StepOutcome query.`);
+        console.info(`Retrieved: (${this.jobData.length} job instance) results & (${this.changeData.size} change summary) results for StepOutcome query.`);
 
         this.stepOutcomeTableData.orderTableDataByChange(true);
 
diff --git a/Engine/Source/Programs/Horde/HordeDashboard/src/components/buildhealth/stepoutcome/StepOutcomeDataTypes.ts b/Engine/Source/Programs/Horde/HordeDashboard/src/components/buildhealth/stepoutcome/StepOutcomeDataTypes.ts
--- a/Engine/Source/Programs/Horde/HordeDashboard/src/components/buildhealth/stepoutcome/StepOutcomeDataTypes.ts
+++ b/Engine/Source/Programs/Horde/HordeDashboard/src/components/buildhealth/stepoutcome/StepOutcomeDataTypes.ts
@@ -1,6 +1,6 @@
 // Copyright Epic Games, Inc. All Rights Reserved.
 
-import { GetStepResponse, IssueData } from "horde/backend/Api";
+import { ChangeSummaryData, GetStepResponse, IssueData } from "horde/backend/Api";
 import { computeDuration, DURATION_NOT_SET, formatDurationFromMs, toDate } from "./StepOutcomeUtilities";
 
 /**
@@ -32,6 +32,21 @@ export const INVALID_CHANGELIST: number = -1;
 
 // #region -- Data Types --
 
+/**
+ * ChangeId type that has the string & numerical representation of a commit id.
+ */
+export type ChangeId = {
+    name: string
+    order: number
+}
+
+/**
+ * Convenience type to fold in change id into the existing ChangeSummaryData.
+ */
+export type UpdatedChangeSummaryData = ChangeSummaryData & {
+    id: ChangeId
+}
+
 /*
  * Convenience type to fold in extra metadata relevant for the StepResponse.
  */
@@ -187,9 +202,10 @@ export class StepOutcomeTable {
     /**
      * Adds a new entry to the table data.
      * @param tableEntry The table entry to add to the data table.
+     * @param changeSummaryData The change summary data corresponding to the @see tableEntry.
      * @returns If tableEntry parameter is invalid, returns.
      */
-    addEntry(tableEntry: StepOutcomeTableEntry): void {
+    addEntry(tableEntry: StepOutcomeTableEntry, changeSummaryData :  UpdatedChangeSummaryData | undefined): void {
         if (tableEntry.change === undefined) {
             return;
         }
@@ -204,13 +220,12 @@ export class StepOutcomeTable {
             this.stepRowLookup.set(normalizedStepName, insertionRow);
             this.stepNameRowHeaders.push({ streamId: tableEntry.streamId, stepName: tableEntry.stepResponse.name });
         }
-
-        // Resolve or create column index
+        
         let insertionCol = this.changeColLookup.get(tableEntry.change);
         if (insertionCol === undefined) {
             insertionCol = this.changeColHeaders.length;
             this.changeColLookup.set(tableEntry.change, insertionCol);
-            this.changeColHeaders.push({ change: tableEntry.change, date: tableEntry.jobCreateTime });
+            this.changeColHeaders.push({ change: tableEntry.change, date: toDate(changeSummaryData?.dateUtc)?.toISOString() });
 
             for (const row of this.tableEntries) {
                 if (row) {
