[Horde][BuildHealth][Dashboard] Increase title size to match changelist/date anchors & checkbox to hide date anchors

diff --git a/Engine/Source/Programs/Horde/HordeDashboard/src/components/buildhealth/BuildHealthOptions.ts b/Engine/Source/Programs/Horde/HordeDashboard/src/components/buildhealth/BuildHealthOptions.ts
--- a/Engine/Source/Programs/Horde/HordeDashboard/src/components/buildhealth/BuildHealthOptions.ts
+++ b/Engine/Source/Programs/Horde/HordeDashboard/src/components/buildhealth/BuildHealthOptions.ts
@@ -212,6 +212,7 @@ export class BuildHealthOptionsState {
     @observable enabledSteps: Record<string, string> = {};
     @observable jobHistoryTimeSpan: TimeSpan = DEFAULT_JOB_HISTORY_TIME_SPAN;
     @observable includePreflight: boolean = false;
+    @observable includeDateAnchors: boolean = true;
 
     // #endregion -- User Selected Options --
 
@@ -333,7 +334,7 @@ export class BuildHealthOptionsController {
         } else {
             delete this.state.enabledTemplates[id];
             let { streamId, templateId } = decodeTemplateKey(id);
-            this.clearSteps({streamId, templateId});
+            this.clearSteps({ streamId, templateId });
         }
         this.setOptionsChanged();
     }
@@ -374,6 +375,16 @@ export class BuildHealthOptionsController {
         this.setOptionsChanged();
     }
 
+    /**
+     * Sets whether the date anchors should be used when visualizing the table.
+     * @param includeDateAnchors 
+     */
+    @action
+    setHideDateAnchors(includeDateAnchors: boolean) {
+        this.state.includeDateAnchors = includeDateAnchors;
+        this.setOptionsChanged();
+    }
+
     /**
       * Clear all the projects tracked by the options.
       * @param cascade Whether to cascade this down to dependent options.
diff --git a/Engine/Source/Programs/Horde/HordeDashboard/src/components/buildhealth/BuildhealthView.tsx b/Engine/Source/Programs/Horde/HordeDashboard/src/components/buildhealth/BuildhealthView.tsx
--- a/Engine/Source/Programs/Horde/HordeDashboard/src/components/buildhealth/BuildhealthView.tsx
+++ b/Engine/Source/Programs/Horde/HordeDashboard/src/components/buildhealth/BuildhealthView.tsx
@@ -474,6 +474,9 @@ const BuildHealthSidebar: React.FC<{ handler: BuildHealthDataHandler, buildHealt
     const { hordeClasses } = getHordeStyling();
     const [timeSpanKey, setTimeSpanKey] = useState(() => buildHealthOptions.state.jobHistoryTimeSpan.key);
     const [includePreflight, setIncludePreflight] = useState(() => buildHealthOptions.state.includePreflight)
+    const [hideDateAnchors, setHideDateAnchors] = useState(() => buildHealthOptions.state.includeDateAnchors)
+
+    // #region -- Property Callbacks --
 
     const updateJobHistoryTimeSpan = (_ev, option) => {
         const select = option as TimeSpan;
@@ -486,6 +489,15 @@ const BuildHealthSidebar: React.FC<{ handler: BuildHealthDataHandler, buildHealt
         buildHealthOptions.setIncludePreflight(isChecked);
     }
 
+    const updateIncludeDateAnchors = (isChecked) => {
+        setHideDateAnchors(isChecked);
+        buildHealthOptions.setHideDateAnchors(isChecked);
+    }
+
+    // #endregion -- Property Callbacks --
+
+    // #region -- Property Use Effects --
+
     useEffect(() => {
         setTimeSpanKey(buildHealthOptions.state.jobHistoryTimeSpan.key);
     }, [buildHealthOptions.state.jobHistoryTimeSpan.key]);
@@ -494,6 +506,12 @@ const BuildHealthSidebar: React.FC<{ handler: BuildHealthDataHandler, buildHealt
         setIncludePreflight(buildHealthOptions.state.includePreflight);
     }, [buildHealthOptions.state.includePreflight]);
 
+    useEffect(() => {
+        setHideDateAnchors(buildHealthOptions.state.includeDateAnchors);
+    }, [buildHealthOptions.state.includeDateAnchors]);
+
+    // #endregion -- Property Use Effects --
+
     let timeComboWidth = 180;
 
     return (
@@ -525,9 +543,6 @@ const BuildHealthSidebar: React.FC<{ handler: BuildHealthDataHandler, buildHealt
                             <div style={{ maxWidth: 300 }}>
                                 <div>
                                     Include preflights in the job history results, or not.
-                                    <p>
-                                        <i>Note: Including preflights disables date anchors in header.</i>
-                                    </p>
                                 </div>
                             </div>
                         }>
@@ -537,6 +552,19 @@ const BuildHealthSidebar: React.FC<{ handler: BuildHealthDataHandler, buildHealt
                                 onChange={(_, isChecked) => updateIncludePreflight(!!isChecked)}
                             />
                         </TooltipHost>
+                        <TooltipHost content={
+                            <div style={{ maxWidth: 300 }}>
+                                <div>
+                                    Include change submit date anchors in the header to separate changes.
+                                </div>
+                            </div>
+                        }>
+                            <Checkbox
+                                label="Use Date Anchors"
+                                checked={hideDateAnchors}
+                                onChange={(_, isChecked) => updateIncludeDateAnchors(!!isChecked)}
+                            />
+                        </TooltipHost>
                     </Stack>
                 </Stack>
             </Stack>
@@ -557,13 +585,15 @@ export const BuildHealthView: React.FC = observer(function ConstructBuildHealthV
         jobs: buildHealthState.stepOutcomeEnabledJobKeys,
         steps: buildHealthState.stepOutcomeEnabledStepKeys,
         jobHistorySpan: { start: buildHealthState.startDate, end: buildHealthState.endDate },
-        includePreflights: buildHealthState.includePreflight
+        includePreflights: buildHealthState.includePreflight,
+        includeDateAnchors: buildHealthState.includeDateAnchors
     }), [
         buildHealthState.stepOutcomeEnabledStreamKeys,
         buildHealthState.stepOutcomeEnabledJobKeys,
         buildHealthState.stepOutcomeEnabledStepKeys,
         { start: buildHealthState.startDate, end: buildHealthState.endDate },
-        buildHealthState.includePreflight
+        buildHealthState.includePreflight,
+        buildHealthState.includeDateAnchors
     ]);
 
     return (
diff --git a/Engine/Source/Programs/Horde/HordeDashboard/src/components/buildhealth/stepoutcome/StepOutcome.tsx b/Engine/Source/Programs/Horde/HordeDashboard/src/components/buildhealth/stepoutcome/StepOutcome.tsx
--- a/Engine/Source/Programs/Horde/HordeDashboard/src/components/buildhealth/stepoutcome/StepOutcome.tsx
+++ b/Engine/Source/Programs/Horde/HordeDashboard/src/components/buildhealth/stepoutcome/StepOutcome.tsx
@@ -161,8 +161,8 @@ const StepOutcomeTableComponent: React.FC<StepOutcomeTableProps> = observer(({ h
                     {dataCount > 0 ? (
                         <Stack>
                             <Stack horizontalAlign="center" styles={{ root: { paddingTop: 8, paddingBottom: 8 } }}>
-                                <Text styles={{ root: { fontWeight: 'bold' } }}>{header}</Text>
-                                <Text >{subHeader} - {subSubHeader}</Text>
+                                <Text styles={{ root: { fontSize: "large", fontWeight: 'bold' } }}>{header}</Text>
+                                <Text styles={{ root: { fontSize: "small", fontWeight: 'bold' } }} >{subHeader} - {subSubHeader}</Text>
                             </Stack>
                             <StepOutcomeD3Table
                                 stepOutcomeTable={handler.StepOutcomeTableData}
diff --git a/Engine/Source/Programs/Horde/HordeDashboard/src/components/buildhealth/stepoutcome/StepOutcomeDataHandler.ts b/Engine/Source/Programs/Horde/HordeDashboard/src/components/buildhealth/stepoutcome/StepOutcomeDataHandler.ts
--- a/Engine/Source/Programs/Horde/HordeDashboard/src/components/buildhealth/stepoutcome/StepOutcomeDataHandler.ts
+++ b/Engine/Source/Programs/Horde/HordeDashboard/src/components/buildhealth/stepoutcome/StepOutcomeDataHandler.ts
@@ -80,6 +80,7 @@ export interface StepOutcomeFilters {
     steps?: string[];
     maxJobCount?: number;
     includePreflights: boolean;
+    includeDateAnchors: boolean;
 }
 
 // #endregion -- API Data Types & Utilities --
@@ -267,8 +268,7 @@ export class StepOutcomeDataHandler extends PollBase {
 
         this.stepOutcomeTableData.orderTableDataByChange(true);
 
-        // @todo - UE-315881 will remove the need for this constraint.
-        this.StepOutcomeTableData.supportsDateAnchoredChanges = !this.activeFilter.includePreflights;
+        this.StepOutcomeTableData.supportsDateAnchoredChanges = this.activeFilter.includeDateAnchors;
 
         this.setUpdated();
 
