Horde Dashboard: Put error handling in own component to keep root view from needing to be an observer

diff --git a/Engine/Source/Programs/Horde/HordeDashboard/src/components/jobDetailsV2/JobDetailViewV2.tsx b/Engine/Source/Programs/Horde/HordeDashboard/src/components/jobDetailsV2/JobDetailViewV2.tsx
--- a/Engine/Source/Programs/Horde/HordeDashboard/src/components/jobDetailsV2/JobDetailViewV2.tsx
+++ b/Engine/Source/Programs/Horde/HordeDashboard/src/components/jobDetailsV2/JobDetailViewV2.tsx
@@ -342,8 +342,18 @@ const ScrollRestore: React.FC<{ jobDetails: JobDetailsV2, scrollRef: React.RefOb
    return null;
 }
 
+const DetailsViewErrorView: React.FC<{ jobDetails: JobDetailsV2 }> = observer(({ jobDetails }) => {
 
-const DetailsView: React.FC<{ jobDetails: JobDetailsV2 }> = observer(({ jobDetails }) => {
+   jobDetails.subscribe();
+
+   if (jobDetails.jobError) {
+      return <ErrorView title="Unable to load job" description="Verify job and stream exists and potentially your internet connection." reason={jobDetails.jobError} />;
+   }
+})
+
+
+/** This is a root view and should not be made an observer other than for a specific observable, ie. not JobDetails updates as these are very coarse and cause issues */
+const DetailsView: React.FC<{ jobDetails: JobDetailsV2 }> = ({ jobDetails }) => {
 
    const windowSize = useWindowSize();
    const scrollRef = useRef<HTMLDivElement>(null);
@@ -353,15 +363,11 @@ const DetailsView: React.FC<{ jobDetails: JobDetailsV2 }> = observer(({ jobDetai
 
    const details = jobDetails;
 
-   if (jobDetails.updated) { } // Reference an observable value to ensure component is reloaded
-   if (jobDetails?.jobError) {
-      return <ErrorView title="Unable to load job" description="Verify job and stream exists and potentially your internet connection." reason={jobDetails.jobError} />;
-   }
-
    const centerAlign = vw / 2 - 720 /*890*/;
    const key = `windowsize_jobdetail_view_${windowSize.width}_${windowSize.height}`;
 
    return <Stack id="rail_job_overview">
+      <DetailsViewErrorView jobDetails={jobDetails}/>
       <JobBreadCrumbs jobDetails={details} />
       <ScrollRestore jobDetails={details} scrollRef={scrollRef} />
       <Stack horizontal styles={{ root: { backgroundColor: modeColors.background } }}>
@@ -396,7 +402,7 @@ const DetailsView: React.FC<{ jobDetails: JobDetailsV2 }> = observer(({ jobDetai
          </Stack>
       </Stack>
    </Stack>
-});
+};
 
 export const JobDetailViewV2: React.FC = () => {
 
