[Horde][Dashboard] Remove left-click to add highlighted text to search box in log view

diff --git a/Engine/Source/Programs/Horde/HordeDashboard/src/components/logs/LogView.tsx b/Engine/Source/Programs/Horde/HordeDashboard/src/components/logs/LogView.tsx
--- a/Engine/Source/Programs/Horde/HordeDashboard/src/components/logs/LogView.tsx
+++ b/Engine/Source/Programs/Horde/HordeDashboard/src/components/logs/LogView.tsx
@@ -452,12 +452,6 @@ const LogList: React.FC<{ logId: string }> = observer(({ logId }) => {
 
    }, []);
 
-   const handleSearchBoxClick = (e: React.MouseEvent<HTMLElement, MouseEvent>) => {
-      if (e.button !== 0) return; // only capture left clicks
-
-      searchSelectedText();
-   };
-
    if (handler && handler.logSource && handler.logSource.logId !== logId) {
       selection.setItems([], true);
       LogHandler.clear();
@@ -1312,7 +1306,6 @@ const LogList: React.FC<{ logId: string }> = observer(({ logId }) => {
                                     componentRef={searchBox}
                                     value={searchState.text ? searchState.text : ''}
                                     onChange={(ev, value) => { setSearchState({...searchState, text: value}); }}
-                                    onMouseDown={handleSearchBoxClick}
                                     
                                     onRenderSuffix={
                                        !!searchState.text ? () => (
