[Horde][BuildHealth][Server][Build] JobSummary telemetry

diff --git a/Engine/Source/Programs/Horde/HordeServer.Tests.Shared/DatabaseIntegrationTest.cs b/Engine/Source/Programs/Horde/HordeServer.Tests.Shared/DatabaseIntegrationTest.cs
--- a/Engine/Source/Programs/Horde/HordeServer.Tests.Shared/DatabaseIntegrationTest.cs
+++ b/Engine/Source/Programs/Horde/HordeServer.Tests.Shared/DatabaseIntegrationTest.cs
@@ -192,7 +192,10 @@ namespace HordeServer.Tests
 		{
 		}
 
-		[AssemblyCleanup]
+		/// <summary>
+		/// Explicitly calls a full dispose of the underyling <see cref="MongoInstance"/>.
+		/// </summary>
+		/// <remarks>This should be invoked in the <see cref="AssemblyCleanupAttribute"/> of test suite's assembly.</remarks>
 		public static void Cleanup()
 		{
 			MongoInstance.FullDispose();
diff --git a/Engine/Source/Programs/Horde/Plugins/Analytics/HordeServer.Analytics.Integration.Tests/Telemetry/AnalyticsPluginTestCleanup.cs b/Engine/Source/Programs/Horde/Plugins/Analytics/HordeServer.Analytics.Integration.Tests/Telemetry/AnalyticsPluginTestCleanup.cs
new file mode 100644
--- /dev/null
+++ b/Engine/Source/Programs/Horde/Plugins/Analytics/HordeServer.Analytics.Integration.Tests/Telemetry/AnalyticsPluginTestCleanup.cs
@@ -0,0 +1,16 @@
+// Copyright Epic Games, Inc. All Rights Reserved.
+
+using HordeServer.Tests;
+
+namespace HordeServer.Analytics.Integration.Tests.Telemetry
+{
+	[TestClass]
+	public static class AssemblyHooks
+	{
+		[AssemblyCleanup]
+		public static void Cleanup()
+		{
+			DatabaseIntegrationTest.Cleanup();
+		}
+	}
+}
\ No newline at end of file
diff --git a/Engine/Source/Programs/Horde/Plugins/Analytics/HordeServer.Analytics.Integration.Tests/Telemetry/TelemetryIntegrationTests.cs b/Engine/Source/Programs/Horde/Plugins/Analytics/HordeServer.Analytics.Integration.Tests/Telemetry/TelemetryIntegrationTests.cs
--- a/Engine/Source/Programs/Horde/Plugins/Analytics/HordeServer.Analytics.Integration.Tests/Telemetry/TelemetryIntegrationTests.cs
+++ b/Engine/Source/Programs/Horde/Plugins/Analytics/HordeServer.Analytics.Integration.Tests/Telemetry/TelemetryIntegrationTests.cs
@@ -3,13 +3,25 @@
 using System.Collections;
 using System.Reflection;
 using EpicGames.Core;
+using EpicGames.Horde.Commits;
 using EpicGames.Horde.Issues;
 using EpicGames.Horde.Jobs;
+using EpicGames.Horde.Jobs.Graphs;
+using EpicGames.Horde.Jobs.Templates;
+using EpicGames.Horde.Projects;
+using EpicGames.Horde.Streams;
 using EpicGames.Horde.Telemetry;
+using HordeCommon;
 using HordeServer.Analytics.Tests;
 using HordeServer.Issues;
+using HordeServer.Jobs;
+using HordeServer.Logs;
+using HordeServer.Projects;
+using HordeServer.Server;
+using HordeServer.Streams;
 using HordeServer.Telemetry;
 using HordeServer.Telemetry.Metrics;
+using HordeServer.Tests;
 using HordeServer.Tests.Issues;
 using Microsoft.Extensions.DependencyInjection;
 
@@ -91,11 +103,10 @@ namespace HordeServer.Analytics.Integration.Tests.Telemetry
 	}
 
 	/// <summary>
-	/// Test class used to test the integration of <see cref="AnalyticsPlugin"/> with various other upstream systems.
+	/// Test class used to test the integration of <see cref="AnalyticsPlugin"/> with the Issue system.
 	/// </summary>
-	/// <remarks><see cref="BuildPlugin"/> represents the current upstream integration.</remarks>
 	[TestClass]
-	public class TelemetryIntegrationTests : AbstractIssueServiceTests
+	public class IssueTelemetryIntegrationTests : AbstractIssueServiceTests
 	{
 		#region -- Test Constants --
 
@@ -211,4 +222,99 @@ namespace HordeServer.Analytics.Integration.Tests.Telemetry
 
 		#endregion -- Test Cases --
 	}
+
+	/// <summary>
+	/// Test class used to test the integration of <see cref="AnalyticsPlugin"/> with the Job system.
+	/// </summary>
+	[TestClass]
+	public class JobTelemetryIntegrationTests : BuildTestSetup
+	{
+		#region -- Test Constants --
+
+		const string EventNameKeyStr = "EventName";
+
+		#endregion -- Test Constants -- 
+
+		protected override void ConfigureServices(IServiceCollection services)
+		{
+			base.ConfigureServices(services);
+
+			services.AddSingleton<TestTelemetryWriter>();
+			services.AddSingleton<ITelemetryWriter>(sp => sp.GetRequiredService<TestTelemetryWriter>());
+		}
+
+		[TestMethod]
+		public async Task TestJobSummaryEmissionAsync()
+		{
+			// Configure the metric
+			TelemetryStoreConfig telemetryStoreConfig = new TelemetryStoreConfig();
+			telemetryStoreConfig.Id = TelemetryStoreId.Default;
+			AnalyticsConfig analyticsConfig = new AnalyticsConfig();
+			analyticsConfig.Stores.Add(telemetryStoreConfig);
+
+			// Setup Project, Stream, Template
+			ProjectId projectId = new ProjectId("ue5");
+			StreamId streamId = new StreamId("ue5-main");
+			TemplateId templateRefId1 = new TemplateId("template1");
+			TemplateId templateRefId2 = new TemplateId("template2");
+
+			StreamConfig streamConfig = new StreamConfig { Id = streamId };
+			streamConfig.Templates.Add(new TemplateRefConfig { Id = templateRefId1, Name = "Test Template" });
+			streamConfig.Tabs.Add(new TabConfig { Title = "foo", Templates = new List<TemplateId> { templateRefId1, templateRefId2 } });
+
+			ProjectConfig projectConfig = new ProjectConfig { Id = projectId };
+			projectConfig.Streams.Add(streamConfig);
+
+			BuildConfig buildConfig = new BuildConfig();
+			buildConfig.Projects.Add(projectConfig);
+
+
+			GlobalConfig globalConfig = new GlobalConfig();
+			globalConfig.Plugins.AddBuildConfig(buildConfig);
+			globalConfig.Plugins.AddAnalyticsTestConfig(analyticsConfig);
+
+			await SetConfigAsync(globalConfig);
+
+			buildConfig.GetStream(streamId).TelemetryStoreId = TelemetryStoreId.Default;
+
+			// Initialize the test telemetry writer
+			TestTelemetryWriter concreteTelemetryWriter = ServiceProvider.GetRequiredService<TestTelemetryWriter>();
+
+			// Validate our preconditions
+			Assert.IsNotNull(concreteTelemetryWriter);
+			Assert.IsTrue(!concreteTelemetryWriter.RecordedTelemetry.ContainsKey(TelemetryStoreId.Default));
+
+			{
+				CreateJobOptions options = new CreateJobOptions();
+				options.PreflightCommitId = CommitId.FromPerforceChange(999);
+
+				ITemplate template = await TemplateCollection.GetOrAddAsync(streamConfig.Templates[0]);
+
+				IGraph graph = await GraphCollection.AddAsync(template);
+
+				IJob job = await JobService.CreateJobAsync(null, streamConfig, templateRefId1, template.Hash, graph, "Hello", CommitIdWithOrder.FromPerforceChange(1234), CommitIdWithOrder.FromPerforceChange(1233), options);
+
+				job = Deref(await JobService.UpdateBatchAsync(job, job.Batches[0].Id, streamConfig, LogIdUtils.GenerateNewId(), JobStepBatchState.Running));
+				job = Deref(await JobService.UpdateStepAsync(job, job.Batches[0].Id, job.Batches[0].Steps[0].Id, streamConfig, JobStepState.Running));
+
+				FakeClock clock = ServiceProvider.GetRequiredService<FakeClock>();
+
+				await clock.AdvanceAsync(TimeSpan.FromSeconds(60.0));
+
+				job = Deref(await JobService.UpdateStepAsync(job, job.Batches[0].Id, job.Batches[0].Steps[0].Id, streamConfig, JobStepState.Completed, JobStepOutcome.Success));
+			}
+
+			// Validate we have received a telemetry event.
+			Assert.IsTrue(concreteTelemetryWriter.RecordedTelemetry.ContainsKey(TelemetryStoreId.Default));
+
+			// Verify presence of job summary telemetry
+			IList<Tuple<TelemetryRecordMeta?, Hashtable>> actualTelemetryPayloads = [.. concreteTelemetryWriter.RecordedTelemetry[TelemetryStoreId.Default]];
+			Assert.IsTrue(actualTelemetryPayloads.Any(x => (x.Item2[EventNameKeyStr] as string) == JobSummaryTelemetry.DefaultEventName));
+
+			// Assert we actually have the expected telemetry type.
+			Hashtable actualStateJobSummaryPayload = actualTelemetryPayloads.Where(x => x.Item2[EventNameKeyStr] as string == JobSummaryTelemetry.DefaultEventName).First().Item2;
+			JobSummaryTelemetry? actualJobSummaryTelemetry = actualStateJobSummaryPayload.MapToObject<JobSummaryTelemetry>();
+			Assert.IsNotNull(actualJobSummaryTelemetry);
+		}
+	}
 }
\ No newline at end of file
diff --git a/Engine/Source/Programs/Horde/Plugins/Build/HordeServer.Build/Jobs/JobService.cs b/Engine/Source/Programs/Horde/Plugins/Build/HordeServer.Build/Jobs/JobService.cs
--- a/Engine/Source/Programs/Horde/Plugins/Build/HordeServer.Build/Jobs/JobService.cs
+++ b/Engine/Source/Programs/Horde/Plugins/Build/HordeServer.Build/Jobs/JobService.cs
@@ -24,10 +24,11 @@ using HordeServer.Jobs.Templates;
 using HordeServer.Jobs.Timing;
 using HordeServer.Logs;
 using HordeServer.Notifications;
-using HordeServer.VersionControl.Perforce;
 using HordeServer.Streams;
+using HordeServer.Telemetry;
 using HordeServer.Users;
 using HordeServer.Utilities;
+using HordeServer.VersionControl.Perforce;
 using Microsoft.Extensions.Logging;
 using Microsoft.Extensions.Options;
 using OpenTelemetry.Trace;
@@ -86,6 +87,7 @@ namespace HordeServer.Jobs
 		readonly Tracer _tracer;
 		readonly ILogger _logger;
 		readonly IAuditLog<JobId> _auditLog;
+		readonly ITelemetryWriter _telemetryWriter;
 
 		/// <summary>
 		/// Delegate for label update events
@@ -127,7 +129,7 @@ namespace HordeServer.Jobs
 		/// <summary>
 		/// Constructor
 		/// </summary>
-		public JobService(IJobCollection jobs, IGraphCollection graphs, IAgentCollection agents, IJobStepRefCollection jobStepRefs, IBisectTaskCollection bisectTasks, IJobTimingCollection jobTimings, IUserCollection userCollection, INotificationTriggerCollection triggerCollection, JobTaskSource jobTaskSource, ITemplateCollectionInternal templateCollection, IPerforceService perforceService, Tracer tracer, IOptionsMonitor<BuildConfig> buildConfig, ILogger<JobService> logger, IAuditLog<JobId> auditLog)
+		public JobService(IJobCollection jobs, IGraphCollection graphs, IAgentCollection agents, IJobStepRefCollection jobStepRefs, IBisectTaskCollection bisectTasks, IJobTimingCollection jobTimings, IUserCollection userCollection, INotificationTriggerCollection triggerCollection, JobTaskSource jobTaskSource, ITemplateCollectionInternal templateCollection, IPerforceService perforceService, Tracer tracer, IOptionsMonitor<BuildConfig> buildConfig, ILogger<JobService> logger, IAuditLog<JobId> auditLog, ITelemetryWriter telemetryWriter)
 		{
 			_jobs = jobs;
 			_graphs = graphs;
@@ -137,13 +139,14 @@ namespace HordeServer.Jobs
 			_jobTimings = jobTimings;
 			_userCollection = userCollection;
 			_triggerCollection = triggerCollection;
-			_jobTaskSource = jobTaskSource;			
+			_jobTaskSource = jobTaskSource;
 			_templateCollection = templateCollection;
 			_perforceService = perforceService;
 			_buildConfig = buildConfig;
 			_tracer = tracer;
 			_logger = logger;
 			_auditLog = auditLog;
+			_telemetryWriter = telemetryWriter;
 
 			_jobTaskSource.OnJobScheduled += (pool, numAgentsOnline, job, graph, batchId) =>
 			{
@@ -790,17 +793,17 @@ namespace HordeServer.Jobs
 							else
 							{
 								_logger.LogError("Batch {BatchId} marked with error {JobStepBatchError}", batchId, error.Value.ToString());
-							}								
+							}
 							_jobs.JobAudit(job.Id, "Batch {BatchId} marked with error {JobStepBatchError}", batchId, error.Value.ToString());
 						}
 						else if (!batch.Steps.Any(x => x.StartTimeUtc != null))
-						{							
+						{
 							newError = JobStepBatchError.SyncingFailed;
 							_logger.LogError("Batch {BatchId} marked with error SyncingFailed", batchId);
 							_jobs.JobAudit(job.Id, "Batch {BatchId} marked with error SyncingFailed", batchId);
 						}
 						else
-						{							
+						{
 							newError = JobStepBatchError.Incomplete;
 							_logger.LogError("Batch {BatchId} marked with error Incomplete", batchId);
 							_jobs.JobAudit(job.Id, "Batch {BatchId} marked with error Incomplete", batchId);
@@ -1043,6 +1046,30 @@ namespace HordeServer.Jobs
 						OnJobStepComplete?.Invoke(job, graph, batchId, stepId);
 					}
 
+					// Emit telemetry for job completion.
+					if (updatedState == JobStepState.Completed || updatedState == JobStepState.Aborted)
+					{
+						if (job.GetState() == JobState.Complete)
+						{
+							IJobStep? completingJobStep;
+
+							if (job.TryGetStep(batchId, stepId, out completingJobStep) && completingJobStep.FinishTimeUtc != null)
+							{
+								if (!streamConfig.TelemetryStoreId.IsEmpty)
+								{
+									try
+									{
+										_telemetryWriter.WriteEvent(streamConfig.TelemetryStoreId, new JobSummaryTelemetry(job, completingJobStep.FinishTimeUtc.Value));
+									}
+									catch(Exception ex)
+									{
+										_logger.LogError(ex, "Failed to publish JobSummary telemetry.");
+									}
+								}
+							}
+						}
+					}
+
 					// Create any downstream jobs
 					if (updatedState == JobStepState.Completed && updatedOutcome != JobStepOutcome.Failure)
 					{
diff --git a/Engine/Source/Programs/Horde/Plugins/Build/HordeServer.Build/Jobs/JobTelemetry.cs b/Engine/Source/Programs/Horde/Plugins/Build/HordeServer.Build/Jobs/JobTelemetry.cs
new file mode 100644
--- /dev/null
+++ b/Engine/Source/Programs/Horde/Plugins/Build/HordeServer.Build/Jobs/JobTelemetry.cs
@@ -0,0 +1,145 @@
+// Copyright Epic Games, Inc. All Rights Reserved.
+
+using EpicGames.Horde.Commits;
+using EpicGames.Horde.Jobs;
+using EpicGames.Horde.Jobs.Templates;
+using EpicGames.Horde.Streams;
+
+namespace HordeServer.Jobs
+{
+	/// <summary>
+	/// Record used for Job Summary telemetry. Models a <see cref="IJob"/>.
+	/// </summary>
+	public record JobSummaryTelemetry
+	{
+		/// <summary>
+		/// The telemetry event name.
+		/// </summary>
+		public string EventName { get; init; }
+
+		/// <summary>
+		/// The job Id.
+		/// </summary>
+		public JobId JobId { get; init; }
+
+		/// <summary>
+		/// The template of the job.
+		/// </summary>
+		public TemplateId TemplateId { get; init; }
+
+		/// <summary>
+		/// The stream of the job.
+		/// </summary>
+		public StreamId StreamId { get; init; }
+
+		/// <summary>
+		/// The commit of the job. <see cref="IJob.CommitId"/>.
+		/// </summary>
+		public CommitIdWithOrder CommitId { get; init; }
+
+		/// <summary>
+		/// The code commit of the job. <see cref="IJob.CodeCommitId"/>.
+		/// </summary>
+		public CommitIdWithOrder CodeCommitId { get; init; }
+
+		/// <summary>
+		/// The created time of the job.
+		/// </summary>
+		public DateTime CreateTimeUtc { get; init; }
+
+		/// <summary>
+		/// The finish time of the job.
+		/// </summary>
+		public DateTime FinishTimeUtc { get; init; }
+
+		/// <summary>
+		/// 
+		/// </summary>
+		public float JobStepsTotalTime { get; init; }
+
+		/// <summary>
+		/// The wall time of the job, in seconds.
+		/// </summary>
+		public float JobWallTime { get; init; }
+
+		/// <summary>
+		/// <see cref="JobStepsCompletionInfo.PassRatio"/>.
+		/// </summary>
+		public float PassRatio { get; init; }
+
+		/// <summary>
+		/// <see cref="JobStepsCompletionInfo.PassWithWarningRatio"/>.
+		/// </summary>
+		public float PassWithWarningRatio { get; init; }
+
+		/// <summary>
+		/// <see cref="JobStepsCompletionInfo.WarningRatio"/>.
+		/// </summary>
+		public float WarningRatio { get; init; }
+
+		/// <summary>
+		/// <see cref="JobStepsCompletionInfo.FailureRatio"/>.
+		/// </summary>
+		public float FailureRatio { get; init; }
+
+		/// <summary>
+		/// <see cref="JobStepsCompletionInfo.StepPassCount"/>.
+		/// </summary>
+		public int StepPassCount { get; init; }
+
+		/// <summary>
+		/// <see cref="JobStepsCompletionInfo.StepWarningCount"/>.
+		/// </summary>
+		public int StepWarningCount { get; init; }
+
+		/// <summary>
+		/// <see cref="JobStepsCompletionInfo.StepFailureCount"/>.
+		/// </summary>
+		public int StepFailureCount { get; init; }
+
+		/// <summary>
+		/// <see cref="JobStepsCompletionInfo.StepTotalCount"/>.
+		/// </summary>
+		public int StepTotalCount { get; init; }
+
+		/// <summary>
+		/// Default event name for the JobSummaryTelemetry.
+		/// </summary>
+		public const string DefaultEventName = "State.JobSummary";
+
+		/// <summary>
+		/// Current schema version.
+		/// </summary>
+		public static readonly int CurrentSchemaVersion = 1;
+
+		/// <summary>
+		/// Constructs a JobSummary Telemetry event.
+		/// </summary>
+		/// <param name="job">The job to base this telemetry event on.</param>
+		/// <param name="finishTimeUtc">The completion time of this job.</param>
+		public JobSummaryTelemetry(IJob job, DateTime finishTimeUtc)
+		{
+			EventName = DefaultEventName;
+			JobId = job.Id;
+			TemplateId = job.TemplateId;
+			StreamId = job.StreamId;
+			CommitId = job.CommitId;
+			CodeCommitId = job.CodeCommitId ?? CommitIdWithOrder.Empty;
+			CreateTimeUtc = job.CreateTimeUtc;
+			FinishTimeUtc = finishTimeUtc;
+			JobWallTime = (float)(FinishTimeUtc - CreateTimeUtc).TotalSeconds;
+
+			JobStepsCompletionInfo jobStepsCompletionInfo = job.GetStepCompletionInfo();
+
+			StepTotalCount = jobStepsCompletionInfo.StepTotalCount;
+			StepPassCount = jobStepsCompletionInfo.StepPassCount;
+			StepWarningCount = jobStepsCompletionInfo.StepWarningCount;
+			StepFailureCount = jobStepsCompletionInfo.StepFailureCount;
+			JobStepsTotalTime = jobStepsCompletionInfo.JobStepsTotalTime;
+			PassRatio = jobStepsCompletionInfo.PassRatio;
+			PassWithWarningRatio = jobStepsCompletionInfo.PassWithWarningRatio;
+			WarningRatio = jobStepsCompletionInfo.WarningRatio;
+			FailureRatio = jobStepsCompletionInfo.FailureRatio;
+		}
+	}
+}
\ No newline at end of file
diff --git a/Engine/Source/Programs/Shared/EpicGames.Horde/Jobs/IJob.cs b/Engine/Source/Programs/Shared/EpicGames.Horde/Jobs/IJob.cs
--- a/Engine/Source/Programs/Shared/EpicGames.Horde/Jobs/IJob.cs
+++ b/Engine/Source/Programs/Shared/EpicGames.Horde/Jobs/IJob.cs
@@ -745,6 +745,54 @@ namespace EpicGames.Horde.Jobs
 #pragma warning restore IDE0054 // Use compound assignment
 		}
 
+		/// <summary>
+		/// Gets the step completion info of the job.
+		/// </summary>
+		/// <param name="job">The job document</param>
+		/// <returns>The job's step completion info.</returns>
+		/// <remarks>Any steps not in a terminal state (<see cref="JobStepState.Completed"/>, <see cref="JobStepState.Aborted"/>, <see cref="JobStepState.Skipped"/>) are excluded in total counts except <see cref="JobStepsCompletionInfo.StepTotalCount"/>.</remarks>
+		public static JobStepsCompletionInfo GetStepCompletionInfo(this IJob job)
+		{
+			JobStepsCompletionInfo jobStepsCompletionInfo = new JobStepsCompletionInfo();
+
+			foreach (IJobStepBatch batch in job.Batches)
+			{
+				foreach (IJobStep step in batch.Steps)
+				{
+					jobStepsCompletionInfo.StepTotalCount++;
+
+					if (step.State != JobStepState.Completed && step.State != JobStepState.Aborted && step.State != JobStepState.Skipped)
+					{
+						continue;
+					}
+
+					if (step.FinishTimeUtc != null && step.StartTimeUtc != null)
+					{
+						TimeSpan? stepDuration = step.FinishTimeUtc.Value - step.StartTimeUtc;
+
+						if (stepDuration != null)
+						{
+							jobStepsCompletionInfo.JobStepsTotalTime += (float)stepDuration.Value.TotalSeconds;
+						}
+					}
+
+					switch (step.Outcome)
+					{
+						case JobStepOutcome.Success:
+							jobStepsCompletionInfo.StepPassCount++;
+							break;
+						case JobStepOutcome.Warnings:
+							jobStepsCompletionInfo.StepWarningCount++;
+							break;
+						case JobStepOutcome.Failure:
+							jobStepsCompletionInfo.StepFailureCount++;
+							break;
+					}
+				}
+			}
+			return jobStepsCompletionInfo;
+		}
+
 		/// <summary>
 		/// Gets the average wait time for this batch
 		/// </summary>
@@ -1759,6 +1807,57 @@ namespace EpicGames.Horde.Jobs
 		}
 	}
 
+	/// <summary>
+	/// Step completion information
+	/// </summary>
+	public class JobStepsCompletionInfo
+	{
+		/// <summary>
+		/// The total count of steps passed within a job.
+		/// </summary>
+		public int StepPassCount { get; set; }
+
+		/// <summary>
+		/// The total count of steps with a warning result within a job.
+		/// </summary>
+		public int StepWarningCount { get; set; }
+
+		/// <summary>
+		/// The total count of steps with a failure result within a job.
+		/// </summary>
+		public int StepFailureCount { get; set; }
+
+		/// <summary>
+		/// The total number of steps within a job.
+		/// </summary>
+		public int StepTotalCount { get; set; }
+
+		/// <summary>
+		/// The total step time for the job, in seconds.
+		/// </summary>
+		public float JobStepsTotalTime { get; set; }
+
+		/// <summary>
+		/// The pass count as compared to the total step count.
+		/// </summary>
+		public float PassRatio => StepTotalCount == 0 ? 0 : StepPassCount / (float)StepTotalCount;
+
+		/// <summary>
+		/// The pass and warning count as compared to the total step count.
+		/// </summary>
+		public float PassWithWarningRatio => StepTotalCount == 0 ? 0 : (StepPassCount + StepWarningCount) / (float)StepTotalCount;
+
+		/// <summary>
+		/// The warning count as compared to the total step count.
+		/// </summary>
+		public float WarningRatio => StepTotalCount == 0 ? 0 : StepWarningCount / (float)StepTotalCount;
+
+		/// <summary>
+		/// The failure count as compared to the total step count.
+		/// </summary>
+		public float FailureRatio => StepTotalCount == 0 ? 0 : StepFailureCount / (float)StepTotalCount;
+	}
+
 	/// <summary>
 	/// Information about a chained job trigger
 	/// </summary>
