[Horde][BuildHealth][Dashboard] Summary column & hours label for title & improving the console.log for change summary request count so we can see the *actual* request

diff --git a/Engine/Source/Programs/Horde/HordeDashboard/src/components/buildhealth/stepoutcome/StepOutcomeD3Component.tsx b/Engine/Source/Programs/Horde/HordeDashboard/src/components/buildhealth/stepoutcome/StepOutcomeD3Component.tsx
--- a/Engine/Source/Programs/Horde/HordeDashboard/src/components/buildhealth/stepoutcome/StepOutcomeD3Component.tsx
+++ b/Engine/Source/Programs/Horde/HordeDashboard/src/components/buildhealth/stepoutcome/StepOutcomeD3Component.tsx
@@ -5,8 +5,8 @@ import * as d3 from "d3";
 import { getHordeStyling, preloadFonts } from "horde/styles/Styles";
 import { useEffect, useRef } from "react";
 import { StepOutcomeDataHandler } from "./StepOutcomeDataHandler";
-import { ChangeHeader, encodeStepNameFromStepNameHeader, INVALID_CHANGELIST, StepNameHeader, StepOutcomeTable, StepOutcomeTableEntry } from "./StepOutcomeDataTypes";
-import { getCellStyle, IncrementalRefreshIndicatorPanel, sharedBorderStyle } from "./StepOutcomeSharedUIComponents";
+import { ChangeHeader, ColumnHeader, encodeStepNameFromStepNameHeader, isChangeHeader, isDateHeader, isStepOutcome, isSummary, isSummaryHeader, StepNameHeader, StepOutcomeTable, TableEntry } from "./StepOutcomeDataTypes";
+import { getCellStyle, getColorRecords, IncrementalRefreshIndicatorPanel, sharedBorderStyle } from "./StepOutcomeSharedUIComponents";
 import { getHordeTheme } from "horde/styles/theme";
 
 // #region -- Styles --
@@ -17,6 +17,7 @@ const baseTooltipLine = {
 };
 
 const tooltipClasses = mergeStyleSets({
+    toolTipSection: sharedBorderStyle.bottomBorder,
     tooltipLine: baseTooltipLine,
     tooltipLineWithBorder: {
         ...baseTooltipLine,
@@ -101,6 +102,10 @@ function GenerateD3Table(stepOutcomeTable: StepOutcomeTable, handler: StepOutcom
     const containerRef = useRef<HTMLDivElement>(null);
     const { modeColors, _hordeClasses, _detailClasses } = getHordeStyling();
     const theme = getHordeTheme();
+    const colorRecords: Record<string, string> = getColorRecords();
+    const summaryColorScale = d3.scaleLinear<string>()
+        .domain([0, 0.5, 1])
+        .range([colorRecords["Failure"], colorRecords["Warnings"], colorRecords["Success"]]);
 
     const scaffold = useRef<{
         wrapper: d3.Selection<HTMLDivElement, unknown, null, undefined>;
@@ -229,29 +234,39 @@ function GenerateD3Table(stepOutcomeTable: StepOutcomeTable, handler: StepOutcom
                 handler.poll();
             });
 
-        // Add all columns to the header row; these will be the changelists.
+        // Add all columns to the header row; these will be the changelists; these control downstream widths
         scaffold.current.theaderrow.selectAll("th.col")
-            .data([...stepOutcomeTable.getDateAnchoredChanges()] as ChangeHeader[])
+            .data([...stepOutcomeTable.getColumnHeaders()] as ColumnHeader[])
             .join("th")
             .attr("class", "col")
             .style("position", "sticky")
             .style("top", "0")
             .style("z-index", "3")
             .style("background-color", modeColors.content)
-            .style("width", "30px")
-            .style("font-size", d => d.change === INVALID_CHANGELIST ? "large" : "small")
+            .style("width", d => isSummaryHeader(d) ? "70px" : "30px")
+            .style("font-size", d => (isDateHeader(d) || isSummaryHeader(d) ? "large" : "small"))
             .style("box-sizing", "border-box")
             .style("writing-mode", "vertical-rl")
             .style("transform", "rotate(180deg)")
             .style("padding-bottom", "5px")
             .style("text-align", "right")
             .style("border", `1px solid ${modeColors.background}`)
-            .text(d => d.change === INVALID_CHANGELIST ? (d.date ?? "") : d.change)
+            .text(entry => isSummaryHeader(entry) ? entry.name : (isChangeHeader(entry) ? entry.change : (entry.date ?? "")))
             .on("mouseover", function (_event, d) {
-                d3.select(this).style("background-color", theme.semanticColors.bodyBackgroundHovered);
                 const entry = d as unknown as ChangeHeader;
 
-                if (entry && entry.change !== INVALID_CHANGELIST) {
+                if (isDateHeader(entry) || isSummaryHeader(entry)) {
+                    return;
+                }
+
+                d3.select(this).style("background-color", theme.semanticColors.bodyBackgroundHovered);
+                if (!entry) {
+                    {
+                        return;
+                    }
+                }
+
+                if (isChangeHeader(entry)) {
                     const tooltipContent = document.createElement("div");
                     tooltipContent.style.paddingLeft = "4px";
 
@@ -368,8 +383,7 @@ function GenerateD3Table(stepOutcomeTable: StepOutcomeTable, handler: StepOutcom
             const cells = tr.selectAll("td.step-name")
                 .data([rowData])
                 .join(
-                    enter => enter.append("td")
-                        .attr("class", "step-name").text(d => d.stepName),
+                    enter => enter.append("td").attr("class", "step-name").text(d => d.stepName),
                     update => update.text(d => d.stepName),
                     exit => exit.remove()
                 );
@@ -411,59 +425,154 @@ function GenerateD3Table(stepOutcomeTable: StepOutcomeTable, handler: StepOutcom
 
         rows.selectAll("td.step-name").style("border", `1px solid ${modeColors.background}`);
 
-        const styleAndBindStepCell = (selection: d3.Selection<HTMLDivElement, StepOutcomeTableEntry, any, any>) => {
+        const styleAndBindStepCell = (
+            selection: d3.Selection<HTMLDivElement, TableEntry, any, any>
+        ) => {
             selection
                 .style("vertical-align", "middle")
-                .style("cursor", "pointer")
                 .style("width", "90%")
                 .style("height", "100%")
-                .style("background-color", d => getCellStyle(d).bgColor)
                 .style("display", "inline-block")
                 .style("text-align", "center")
-                .on("click", (_event, d) => {
-                    onCellClick?.(d);
-                })
-                .on("mouseover", function (_event, d) {
-                    if (!d) return;
-                    const entry = d as unknown as StepOutcomeTableEntry;
-                    const tooltipContent = document.createElement("div");
-                    tooltipContent.style.paddingLeft = "4px";
-
-                    const changeDiv = document.createElement("div");
-                    changeDiv.innerHTML = `<strong>Change:</strong><p class="${tooltipClasses.tooltipLineWithBorder}">${entry.change ?? "-"}</p>`;
-                    tooltipContent.appendChild(changeDiv);
+                .each(function (entry) {
+                    const div = d3.select(this);
 
-                    const durationDiv = document.createElement("div");
-                    durationDiv.innerHTML = `<strong>Duration:</strong><p class="${tooltipClasses.tooltipLine}">${entry.getDurationString?.() ?? "-"}</p>`;
-                    tooltipContent.appendChild(durationDiv);
-
-                    const stateDiv = document.createElement("div");
-                    stateDiv.innerHTML = `<strong>State:</strong><p class="${tooltipClasses.tooltipLineWithBorder}">${entry.stepResponse?.state ?? "-"}</p>`;
-                    tooltipContent.appendChild(stateDiv);
-
-                    const streamDiv = document.createElement("div");
-                    streamDiv.innerHTML = `<strong>Stream:</strong><p class="${tooltipClasses.tooltipLine}">${entry.streamId}</p>`;
-                    tooltipContent.appendChild(streamDiv);
-
-                    const jobDiv = document.createElement("div");
-                    jobDiv.innerHTML = `<strong>Job:</strong><p class="${tooltipClasses.tooltipLine}">${entry.jobName ?? "-"}</p>`;
-                    tooltipContent.appendChild(jobDiv);
-
-                    const stepDiv = document.createElement("div");
-                    stepDiv.innerHTML = `<strong>Step:</strong><p class="${tooltipClasses.tooltipLine}">${entry.stepResponse?.name ?? "-"}</p>`;
-                    tooltipContent.appendChild(stepDiv);
+                    if (!entry) {
+                        div.text("").style("background-color", "transparent").style("cursor", "default");
+                        return;
+                    }
 
-                    attachTargettedToolTip(this as HTMLElement, scaffold.current!.tooltip, tooltipContent);
-                })
-                .on("mouseout", function () {
-                    scaffold.current!.tooltip.style("opacity", 0);
+                    if (isStepOutcome(entry)) { // StepOutcomeTableEntry
+                        div.text(""),
+                            div.style("cursor", "pointer")
+                                .style("background-color", getCellStyle(entry).bgColor)
+                                .on("click", () => onCellClick?.(entry))
+                                .on("mouseover", function () {
+                                    const tooltipContent = document.createElement("div");
+                                    tooltipContent.style.paddingLeft = "4px";
+
+                                    const changeDiv = document.createElement("div");
+                                    changeDiv.innerHTML = `<strong>Change:</strong><p class="${tooltipClasses.tooltipLineWithBorder}">${entry.change ?? "-"}</p>`;
+                                    tooltipContent.appendChild(changeDiv);
+
+                                    const durationDiv = document.createElement("div");
+                                    durationDiv.innerHTML = `<strong>Duration:</strong><p class="${tooltipClasses.tooltipLine}">${entry.getDurationString?.() ?? "-"}</p>`;
+                                    tooltipContent.appendChild(durationDiv);
+
+                                    const stateDiv = document.createElement("div");
+                                    stateDiv.innerHTML = `<strong>State:</strong><p class="${tooltipClasses.tooltipLineWithBorder}">${entry.stepResponse?.state ?? "-"}</p>`;
+                                    tooltipContent.appendChild(stateDiv);
+
+                                    const streamDiv = document.createElement("div");
+                                    streamDiv.innerHTML = `<strong>Stream:</strong><p class="${tooltipClasses.tooltipLine}">${entry.streamId}</p>`;
+                                    tooltipContent.appendChild(streamDiv);
+
+                                    const jobDiv = document.createElement("div");
+                                    jobDiv.innerHTML = `<strong>Job:</strong><p class="${tooltipClasses.tooltipLine}">${entry.jobName ?? "-"}</p>`;
+                                    tooltipContent.appendChild(jobDiv);
+
+                                    const stepDiv = document.createElement("div");
+                                    stepDiv.innerHTML = `<strong>Step:</strong><p class="${tooltipClasses.tooltipLine}">${entry.stepResponse?.name ?? "-"}</p>`;
+                                    tooltipContent.appendChild(stepDiv);
+
+                                    attachTargettedToolTip(this as HTMLElement, scaffold.current!.tooltip, tooltipContent);
+                                })
+                                .on("mouseout", () => {
+                                    scaffold.current!.tooltip.style("opacity", 0);
+                                });
+                    } else if (isSummary(entry)) { // SummaryTableEntry
+                        const ratio = entry.steps > 0 ? entry.stepsPass / entry.stepsCompleted : 0;
+                        let ratioStr: string = `${(ratio * 100).toFixed(1)}%`;
+
+                        div.style("cursor", "default")
+                            .style("background-color", summaryColorScale(ratio))
+                            .style("width", "100%")
+                            .style("display", "flex")
+                            .style("align-items", "center")
+                            .style("justify-content", "center")
+                            .style("font-weight", "900")
+                            .style("color", "black")
+                            .text(entry.steps ? ratioStr : "-")
+                            .on("click", null)
+                            .on("mouseover", function () {
+                                const tooltipContent = document.createElement("div");
+                                tooltipContent.style.paddingLeft = "4px";
+
+                                // Show ratio in proper format
+                                const ratio = document.createElement("div");
+                                ratio.style.display = "flex";
+                                ratio.style.justifyContent = "space-between";
+                                ratio.style.alignItems = "flex-start";
+                                ratio.className = tooltipClasses.toolTipSection;
+
+                                const leftDiv = document.createElement("div");
+                                leftDiv.style.display = "flex";
+                                leftDiv.style.flexDirection = "column";
+
+                                const label = document.createElement("strong");
+                                label.textContent = "Success Ratio:";
+                                leftDiv.appendChild(label);
+
+                                const desc = document.createElement("span");
+                                desc.style.fontSize = "smaller";
+                                desc.textContent = "(Success / Completed Steps)";
+                                leftDiv.appendChild(desc);
+
+                                // Right side: ratio value
+                                const rightDiv = document.createElement("span");
+                                rightDiv.textContent = ratioStr;
+
+                                // Append both sides to the main container
+                                ratio.appendChild(leftDiv);
+                                ratio.appendChild(rightDiv);
+
+                                tooltipContent.appendChild(ratio);
+
+                                // Show totals
+                                const runningSteps = document.createElement("div");
+                                runningSteps.style.display = "flex";
+                                runningSteps.style.justifyContent = "space-between";
+                                runningSteps.innerHTML = `<strong>Pending Steps:</strong><span>${entry.stepsPending ?? "-"}</span>`;
+                                tooltipContent.appendChild(runningSteps);
+
+                                const succeededSteps = document.createElement("div");
+                                succeededSteps.style.display = "flex";
+                                succeededSteps.style.justifyContent = "space-between";
+                                succeededSteps.innerHTML = `<strong>Succeeded Steps:</strong><span>${entry.stepsPass ?? "-"}</span>`;
+                                tooltipContent.appendChild(succeededSteps);
+
+                                const warningSteps = document.createElement("div");
+                                warningSteps.style.display = "flex";
+                                warningSteps.style.justifyContent = "space-between";
+                                warningSteps.innerHTML = `<strong>Warning Steps:</strong><span>${entry.stepsWarning ?? "-"}</span>`;
+                                tooltipContent.appendChild(warningSteps);
+
+                                const failedSteps = document.createElement("div");
+                                failedSteps.style.display = "flex";
+                                failedSteps.style.justifyContent = "space-between";
+                                failedSteps.innerHTML = `<strong>Failed Steps:</strong><span>${entry.stepsFail ?? "-"}</span>`;
+                                tooltipContent.appendChild(failedSteps);
+
+                                const completedStepsDiv = document.createElement("div");
+                                completedStepsDiv.style.display = "flex";
+                                completedStepsDiv.style.justifyContent = "space-between";
+                                completedStepsDiv.innerHTML = `<strong>Completed Steps:</strong><span>${entry.stepsCompleted ?? "-"}</span>`;
+                                tooltipContent.appendChild(completedStepsDiv);
+
+                                attachTargettedToolTip(this as HTMLElement, scaffold.current!.tooltip, tooltipContent);
+                            })
+                            .on("mouseout", () => {
+                                scaffold.current!.tooltip.style("opacity", 0);
+                            });
+                    }
                 });
-        }
+        };
 
-        // Render all of the step outcome data
+        // Render step outcome data
         rows.each(function (rowData: StepNameHeader) {
             const tr = d3.select(this);
-            const cellsRaw = stepOutcomeTable.getStepOutputRow(encodeStepNameFromStepNameHeader(rowData));
+
+            const cellsRaw = stepOutcomeTable.getStepOutputTableEntries(encodeStepNameFromStepNameHeader(rowData));
             const cells = cellsRaw ? [...cellsRaw] : [];
 
             const dataCells = tr.selectAll("td.step-outcome-data")
@@ -476,7 +585,7 @@ function GenerateD3Table(stepOutcomeTable: StepOutcomeTable, handler: StepOutcom
             dataCells.selectAll("div")
                 .data(d => [d])
                 .join("div")
-                .call(div => styleAndBindStepCell(div as d3.Selection<HTMLDivElement, StepOutcomeTableEntry, any, any>))
+                .call(div => styleAndBindStepCell(div as d3.Selection<HTMLDivElement, TableEntry, any, any>))
                 .style("height", "25px");
         });
 
@@ -492,7 +601,12 @@ function GenerateD3Table(stepOutcomeTable: StepOutcomeTable, handler: StepOutcom
     // After the user regains control and can horizontally scroll, we do *not* want to snap to right anymore.
     useEffect(() => {
         if (containerRef.current) {
-            containerRef.current.scrollLeft = containerRef.current.scrollWidth;
+            const container = containerRef.current;
+
+            // We need to make sure the d3 table is fully drawn just so we get the padding included in the width.
+            requestAnimationFrame(() => {
+                container.scrollLeft = container.scrollWidth;
+            });
         }
     }, [handler.StepOutcomeTableData]);
 
diff --git a/Engine/Source/Programs/Horde/HordeDashboard/src/components/buildhealth/stepoutcome/StepOutcomeDataHandler.ts b/Engine/Source/Programs/Horde/HordeDashboard/src/components/buildhealth/stepoutcome/StepOutcomeDataHandler.ts
--- a/Engine/Source/Programs/Horde/HordeDashboard/src/components/buildhealth/stepoutcome/StepOutcomeDataHandler.ts
+++ b/Engine/Source/Programs/Horde/HordeDashboard/src/components/buildhealth/stepoutcome/StepOutcomeDataHandler.ts
@@ -1,6 +1,6 @@
 // Copyright Epic Games, Inc. All Rights Reserved.
 
-import { ChangeSummaryData, GetJobResponse, IssueQuery, JobStreamQuery } from "horde/backend/Api";
+import { GetJobResponse, IssueQuery, JobStreamQuery } from "horde/backend/Api";
 import { PollBase } from "horde/backend/PollBase";
 import { StepOutcomeTable, StepOutcomeTableEntry, UpdatedChangeSummaryData } from "./StepOutcomeDataTypes";
 import backend from "horde/backend";
@@ -59,12 +59,20 @@ export function prettyPrintStepOutcomeFiltersSummary(filters: StepOutcomeFilters
 
     const subHeader = `Jobs: ${jobs} - Steps: ${steps}`;
 
-    // Sub-subheader: Last N Days
+    // Sub-subheader: Last N Days/Hours
     const start = filters.jobHistorySpan.start;
     const end = filters.jobHistorySpan.end || new Date();
-    const diffTime = Math.abs(end.getTime() - start.getTime());
-    const diffDays = Math.ceil(diffTime / (1000 * 60 * 60 * 24));
-    const subSubHeader = `Last ${diffDays} day${diffDays > 1 ? "s" : ""}`;
+
+    const diffMs = end.getTime() - start.getTime();
+    const diffDays = Math.floor(diffMs / (1000 * 60 * 60 * 24));
+
+    let subSubHeader: string;
+    if (diffDays >= 1) {
+        subSubHeader = `Last ${diffDays} day${diffDays > 1 ? "s" : ""}`;
+    } else {
+        const diffHours = Math.ceil(diffMs / (1000 * 60 * 60));
+        subSubHeader = `Last ${diffHours} hour${diffHours > 1 ? "s" : ""}`;
+    }
 
     return [header, subHeader, subSubHeader];
 }
@@ -213,26 +221,29 @@ export class StepOutcomeDataHandler extends PollBase {
             if (newChanges.length > 0) {
                 if (!jobHistoryCount.has(streamId)) {
                     jobHistoryCount.set(streamId, []);
-                } 
-                
+                }
+
                 jobHistoryCount.get(streamId)!.push(...newChanges);
             }
         });
 
         // Flatten all jobData
         this.jobData = results.flat();
-        
+
         if (jobHistoryCount.size > 0) {
             // Parallelize backend calls for change summaries
             const changeResults = await Promise.all(
                 changeQueries.map(({ streamId, query }) => backend.getChangeSummaries(streamId, undefined, undefined, jobHistoryCount.get(streamId), query.filter))
             );
 
-            changeResults.flat().map((value: UpdatedChangeSummaryData) => {
+            let flattenedResults = changeResults.flat();
+            flattenedResults.map((value: UpdatedChangeSummaryData) => {
                 if (!this.changeData.has(value.id.order)) {
                     this.changeData.set(value.id.order, value);
                 }
             });
+
+            console.info(`Retrieved: ${flattenedResults.length} change summary results for StepOutcome query.`);
         }
 
         for (const streamResult of results) {
@@ -261,10 +272,14 @@ export class StepOutcomeDataHandler extends PollBase {
                         }
                     }
                 }
+
+                // Summarize each row
             }
         }
+        
+        this.StepOutcomeTableData.summarize();
 
-        console.info(`Retrieved: (${this.jobData.length} job instance) results & (${this.changeData.size} change summary) results for StepOutcome query.`);
+        console.info(`Retrieved: ${this.jobData.length} job instance results for StepOutcome query.`);
 
         this.stepOutcomeTableData.orderTableDataByChange(true);
 
diff --git a/Engine/Source/Programs/Horde/HordeDashboard/src/components/buildhealth/stepoutcome/StepOutcomeDataTypes.ts b/Engine/Source/Programs/Horde/HordeDashboard/src/components/buildhealth/stepoutcome/StepOutcomeDataTypes.ts
--- a/Engine/Source/Programs/Horde/HordeDashboard/src/components/buildhealth/stepoutcome/StepOutcomeDataTypes.ts
+++ b/Engine/Source/Programs/Horde/HordeDashboard/src/components/buildhealth/stepoutcome/StepOutcomeDataTypes.ts
@@ -3,6 +3,13 @@
 import { ChangeSummaryData, GetStepResponse, IssueData } from "horde/backend/Api";
 import { computeDuration, DURATION_NOT_SET, formatDurationFromMs, toDate } from "./StepOutcomeUtilities";
 
+export const WarningStr: string = "Warnings";
+export const SuccessStr: string = "Success";
+export const SkippedStr: string = "Skipped";
+export const FailureStr: string = "Failure";
+export const WaitingStr: string = "Waiting";
+export const RunningStr: string = "Running";
+
 /**
  * Encodes a step outcome table entry to a stream qualified step name.
  * Note: The template is not included in the hierarchy qualification since two separate templates within the same stream will collocate the same step row.
@@ -26,7 +33,6 @@ export function encodeStepNameFromStepNameHeader(stepNameHeader: StepNameHeader)
 // #region -- Constants --
 
 const KEY_SEPARATOR: string = "::";
-export const INVALID_CHANGELIST: number = -1;
 
 // #endregion -- Constants --
 
@@ -54,6 +60,43 @@ export type StepReponseData = GetStepResponse & {
     stepDurationMs: number;
 }
 
+// #region -- Table Entry Types --
+
+/**
+ * Discriminated union which represents a tabel entry. This can be a @see StepOutcomeTableEntry (the core data in the table) or a @see SummaryTableEntry (the summarization of an entire row).
+ */
+export type TableEntry = StepOutcomeTableEntry | SummaryTableEntry | undefined;
+
+/**
+ * Type predicate to test whether this is a @see StepOutcomeTableEntry .
+ * @param entry The entry to test.
+ * @returns True if it is a @see StepOutcomeTableEntry - false otherwise.
+ */
+export function isStepOutcome(entry: TableEntry): entry is StepOutcomeTableEntry {
+    return entry instanceof StepOutcomeTableEntry;
+}
+
+/**
+ * Type predicate to test whether this is a @see SummaryTableEntry .
+ * @param entry The entry to test.
+ * @returns True if it is a @see SummaryTableEntry - false otherwise.
+ */
+export function isSummary(entry: TableEntry): entry is SummaryTableEntry {
+    return entry instanceof SummaryTableEntry;
+}
+
+/**
+ * Data structure that represents a step summary for a collection of @see StepOutcomeTableEntry data.
+ */
+export class SummaryTableEntry {
+    steps: number = 0;
+    stepsCompleted: number = 0;
+    stepsPass: number = 0;
+    stepsFail: number = 0;
+    stepsWarning: number = 0;
+    stepsPending: number = 0;
+}
+
 /**
  * Model for a table entry.
  */
@@ -113,14 +156,67 @@ export class StepOutcomeTableEntry {
     // #endregion -- Public API --
 }
 
+// #endregion -- Table Entry Types --
+
+// #region -- Table Header Types --
+
+/**
+ * Discriminated union which represents a tabel entry. This can be a @see ChangeHeader (the fundamental change header in the table) or a @see SummaryHeader (the summarization of an entire row).
+ */
+export type ColumnHeader = ChangeHeader | SummaryHeader | DateHeader;
+
+/**
+ * Type predicate to test whether this is a @see SummaryHeader .
+ * @param header the header to apply the predicate to.
+ * @returns True if the header was a @see SummaryHeader ; false otherwise.
+ */
+export function isSummaryHeader(header: ColumnHeader): header is SummaryHeader {
+    return header.type === "summary";
+}
+
+/**
+ * Type predicate to test whether this is a @see DateHeader .
+ * @param header the header to apply the predicate to.
+ * @returns True if the header was a @see DateHeader ; false otherwise.
+ */
+export function isDateHeader(header: ColumnHeader): header is DateHeader {
+    return header.type === "date";
+}
+
+/**
+ * Type predicate to test whether this is a @see ChangeHeader .
+ * @param header the header to apply the predicate to.
+ * @returns True if the header was a @see ChangeHeader ; false otherwise.
+ */
+export function isChangeHeader(header: ColumnHeader): header is ChangeHeader {
+    return header.type === "change";
+}
+
+/**
+ * Data struct that represents a summary header.
+ */
+export type SummaryHeader = {
+    type: "summary";
+    name: string;
+}
+
 /**
  * Data struct that represents a change header. This header contains associated data related to rendering & ordering the header.
  */
 export type ChangeHeader = {
+    type: "change";
     change: number;
     date?: string;
 }
 
+/**
+ * Data struct that represents a date header. This header contains associated data related to rendering & ordering the header.
+ */
+export type DateHeader = {
+    type: "date";
+    date: string;
+}
+
 /**
  * Data struct that represents a step name header. This header contains associated data related to rendering & ordering the header.
  */
@@ -129,6 +225,8 @@ export type StepNameHeader = {
     streamId: string;
 }
 
+// #region -- Table Header Types --
+
 /**
  * Model for the StepOutcomeTable.
  */
@@ -154,8 +252,14 @@ export class StepOutcomeTable {
      * Step name row headers.
      */
     stepRowLookup: Map<string, number> = new Map<string, number>();
+    stepRowNumberLookup: Map<number, string> = new Map<number, string>();
     stepNameRowHeaders: StepNameHeader[] = [];
 
+    /**
+     * SummaryLookup maps a normalized step name to a summary entry.
+     */
+    summaryLookup: Map<string, SummaryTableEntry> = new Map<string, SummaryTableEntry>();
+
     tableStreamMetadata: Map<string, number> = new Map<string, number>();
 
     /**
@@ -194,6 +298,8 @@ export class StepOutcomeTable {
         this.tableEntries = [];
         this.stepNameRowHeaders = [];
         this.stepRowLookup.clear();
+        this.stepRowNumberLookup.clear();
+        this.summaryLookup.clear();
         this.changeColHeaders = [];
         this.changeColLookup.clear();
         this.tableStreamMetadata.clear();
@@ -205,7 +311,7 @@ export class StepOutcomeTable {
      * @param changeSummaryData The change summary data corresponding to the @see tableEntry.
      * @returns If tableEntry parameter is invalid, returns.
      */
-    addEntry(tableEntry: StepOutcomeTableEntry, changeSummaryData :  UpdatedChangeSummaryData | undefined): void {
+    addEntry(tableEntry: StepOutcomeTableEntry, changeSummaryData: UpdatedChangeSummaryData | undefined): void {
         if (tableEntry.change === undefined) {
             return;
         }
@@ -217,15 +323,16 @@ export class StepOutcomeTable {
         if (insertionRow === undefined) {
             this.annotateStreamRow(tableEntry.streamId);
             insertionRow = this.stepNameRowHeaders.length;
-            this.stepRowLookup.set(normalizedStepName, insertionRow);
+            this.stepRowLookup.set(normalizedStepName, insertionRow); // cache the normalized step name -> row number
+            this.stepRowNumberLookup.set(insertionRow, normalizedStepName); // cache the row number -> normalized step name
             this.stepNameRowHeaders.push({ streamId: tableEntry.streamId, stepName: tableEntry.stepResponse.name });
         }
-        
+
         let insertionCol = this.changeColLookup.get(tableEntry.change);
         if (insertionCol === undefined) {
             insertionCol = this.changeColHeaders.length;
             this.changeColLookup.set(tableEntry.change, insertionCol);
-            this.changeColHeaders.push({ change: tableEntry.change, date: toDate(changeSummaryData?.dateUtc)?.toISOString() });
+            this.changeColHeaders.push({ type: "change", change: tableEntry.change, date: toDate(changeSummaryData?.dateUtc)?.toISOString() });
 
             for (const row of this.tableEntries) {
                 if (row) {
@@ -249,18 +356,88 @@ export class StepOutcomeTable {
             });
             return;
         }
-        
+
         this.dataCount++;
         this.tableEntries[insertionRow][insertionCol] = tableEntry;
     }
 
     /**
-     * Gets the Table Entries for a given step name.
+     * Summarizes all of the data currently stored within the step outcome table.
+     * A @see SummaryTableEntry is the step completion (across failure & success states) for all rows in the table.
+     */
+    summarize() {
+        for (let rowIndex = 0; rowIndex < this.tableEntries.length; rowIndex++) {
+            const row = this.tableEntries[rowIndex];
+            if (!this.stepRowNumberLookup.has(rowIndex)) {
+                continue;
+            }
+
+            let summaryTableEntry: SummaryTableEntry = new SummaryTableEntry();
+            let normalizedTableEntryName: string = this.stepRowNumberLookup.get(rowIndex)!;
+            this.summaryLookup.set(normalizedTableEntryName, summaryTableEntry);
+
+            for (let colIndex = 0; colIndex < row.length; colIndex++) {
+                const cell = row[colIndex];
+                if (cell == null) {
+                    continue;
+                }
+                summaryTableEntry.steps++;
+
+                if (cell.stepResponse.state === WaitingStr || cell.stepResponse.state === RunningStr) {
+                    summaryTableEntry.stepsPending++;
+                    continue;
+                }
+
+                summaryTableEntry.stepsCompleted++;
+
+                if (cell.stepResponse.outcome == WarningStr) {
+                    summaryTableEntry.stepsWarning++;
+                }
+                else if (cell.stepResponse.outcome == FailureStr) {
+                    summaryTableEntry.stepsFail++;
+                }
+                else if (cell.stepResponse.outcome == SuccessStr) {
+                    summaryTableEntry.stepsPass++;
+                }
+            }
+        }
+    }
+
+    /**
+     * Retrieves a step row summary for a provided row.
+     * @param row The normalized row name, or number.
+     * @returns The step summary for that row if it exists, undefined otherwise.
+     */
+    getStepRowSummary(row: string | number): TableEntry {
+        let rowString: string | undefined = undefined;
+        if (typeof row === 'number') {
+            rowString = this.stepRowNumberLookup.get(row);
+        } else {
+            rowString = row;
+        }
+
+        let summaryTableEntry: SummaryTableEntry | undefined = this.summaryLookup.has(rowString!) ? this.summaryLookup.get(rowString!) : undefined;
+
+        return summaryTableEntry;
+    }
+
+    /**
+     * Gets all of the Table Entries in a row.
+     * @param * Generator of table entries.
+     */
+    * getStepOutputTableEntries(row: string | number): Generator<TableEntry> {
+        yield* this.getStepOutputRow(row);
+        yield this.getStepRowSummary(row);
+    }
+
+    /**
+     * Gets the Step Outcome Table Entries for a given step name.
      * @param rowName The step name to obtain.
      * @returns * Generator of TableEntries.
      */
     getStepOutputRow(rowName: string): Generator<StepOutcomeTableEntry>;
     getStepOutputRow(rowIndex: number): Generator<StepOutcomeTableEntry>;
+    getStepOutputRow(row: string | number): Generator<StepOutcomeTableEntry>;
     *getStepOutputRow(row: string | number): Generator<StepOutcomeTableEntry | null | undefined> {
         let rowIdx: number | undefined;
 
@@ -277,16 +454,21 @@ export class StepOutcomeTable {
         const rowEntries = this.tableEntries[rowIdx];
 
         for (const entry of this.getDateAnchoredChanges()) {
-            const colIdx = this.changeColLookup.get(entry.change);
-            const value = colIdx !== undefined && colIdx > INVALID_CHANGELIST ? rowEntries[colIdx] : undefined;
-            yield value;
+            if (isDateHeader(entry)) {
+                yield undefined;
+            }
+            else {
+                const colIdx = this.changeColLookup.get(entry.change);
+                yield colIdx !== undefined ? rowEntries[colIdx] : undefined;
+
+            }
         }
     }
 
     /**
      * Gets the Table Entries for the given change.
      * @param changeNumber The change number.
-     * @returns All Table Entries for the specific change.
+     * @returns * Generator of all Table Entries for the specific change.
      */
     *getChangeColumn(changeNumber: number): Generator<StepOutcomeTableEntry> {
 
@@ -329,11 +511,23 @@ export class StepOutcomeTable {
         });
     }
 
+    /**
+     * Gets all of the column headers for the table.
+     * @returns * Generator of all the column headers.
+     */
+    *getColumnHeaders(): Generator<ColumnHeader> {
+        yield* this.getDateAnchoredChanges();
+        yield {
+            type: "summary",
+            name: "Summary"
+        }
+    }
+
     /**
      * Gets the the date anchored changes. 
-     * @returns All of the @see ChangeHeader items separated by entries that represent dates. These will have @see INVALID_CHANGELIST as their change property.
+     * @returns * Generator of all of the @see ChangeHeader items separated by @see DateHeader that represent dates. 
      */
-    *getDateAnchoredChanges(): Generator<ChangeHeader> {
+    private *getDateAnchoredChanges(): Generator<ChangeHeader | DateHeader> {
         if (!this.supportsDateAnchoredChanges) {
             for (const header of this.changeColHeaders) {
                 yield header;
@@ -343,23 +537,20 @@ export class StepOutcomeTable {
         }
 
         let prevDay: string | null = null;
-        let prevChangelist: string | null = null;
         for (let i: number = 0; i < this.changeColHeaders.length; ++i) {
 
             // If the date associated with the changelist header has changed since the last, we are on a new day. Set that as the colimn.
             let changeHeader = this.changeColHeaders[i];
             if (changeHeader.date !== undefined) {
                 const headerDate = new Date(changeHeader.date!);
-                const headerChangelist = changeHeader.change.toString();
 
                 // Extract the day portion as YYYY-MM-DD
                 const dayStr = `${headerDate.getUTCFullYear()}-${headerDate.getUTCMonth() + 1}-${headerDate.getUTCDate()}`;
 
                 if (prevDay !== dayStr) {
                     prevDay = dayStr;
-                    prevChangelist = headerChangelist;
-                    let dateAnchorChangelistHeader: ChangeHeader = {
-                        change: INVALID_CHANGELIST,
+                    let dateAnchorChangelistHeader: DateHeader = {
+                        type: "date",
                         date: dayStr
                     }
                     yield dateAnchorChangelistHeader;
diff --git a/Engine/Source/Programs/Horde/HordeDashboard/src/components/buildhealth/stepoutcome/StepOutcomeSharedUIComponents.tsx b/Engine/Source/Programs/Horde/HordeDashboard/src/components/buildhealth/stepoutcome/StepOutcomeSharedUIComponents.tsx
--- a/Engine/Source/Programs/Horde/HordeDashboard/src/components/buildhealth/stepoutcome/StepOutcomeSharedUIComponents.tsx
+++ b/Engine/Source/Programs/Horde/HordeDashboard/src/components/buildhealth/stepoutcome/StepOutcomeSharedUIComponents.tsx
@@ -4,7 +4,7 @@ import { mergeStyleSets, Spinner, SpinnerSize } from "@fluentui/react";
 import dashboard, { StatusColor } from "horde/backend/Dashboard";
 import { observer } from "mobx-react-lite";
 import { StepOutcomeDataHandler } from "./StepOutcomeDataHandler";
-import { StepOutcomeTableEntry } from "./StepOutcomeDataTypes";
+import { FailureStr, SkippedStr, StepOutcomeTableEntry, SuccessStr, WarningStr } from "./StepOutcomeDataTypes";
 
 // #region -- Constants --
 
@@ -28,36 +28,42 @@ export const stepOutcomeTableClasses = mergeStyleSets({
 
 // #endregion -- Styles --
 
+/**
+ * Gets the color records based off of the string representations of @see StatusColor - which maps to the status string.
+ * @returns A record that maps sates to their corresponding colours. 
+ */
+export function getColorRecords(): Record<string, string> {
+    const scolors = dashboard.getStatusColors();
+    const colors: Record<string, string> = {
+        "Success": scolors.get(StatusColor.Success)!,
+        "Failure": scolors.get(StatusColor.Failure)!,
+        "Warnings": scolors.get(StatusColor.Warnings)!,
+        "Skipped": scolors.get(StatusColor.Skipped)!,
+        "Running": scolors.get(StatusColor.Running)!,
+        "Waiting": scolors.get(StatusColor.Waiting)!,
+        "Ready": scolors.get(StatusColor.Ready)!,
+        "Unspecified": scolors.get(StatusColor.Skipped)!,
+    };
+
+    return colors;
+}
+
 /**
  * Gets the cell style given the provided StepOutcomeTableEntry.
  * @param entry The entry to base the style off of.
  * @returns The matching style given the state of the entry. 
  */
 export function getCellStyle(entry: StepOutcomeTableEntry | null): { bgColor: string } {
-    const scolors = dashboard.getStatusColors();
+    const colors: Record<string, string> = getColorRecords();
+
     let resultColor: string = 'transparent';
 
     if (!entry) {
         return { bgColor: resultColor };
     }
 
-    const WarningStr: string = "Warnings";
-    const SuccessStr: string = "Success";
-    const SkippedStr: string = "Skipped";
-    const FailureStr: string = "Failure";
-
     const { stepResponse } = entry;
     const { outcome, state } = stepResponse;
-    const colors: Record<string, string> = {
-        "Success": scolors.get(StatusColor.Success)!,
-        "Failure": scolors.get(StatusColor.Failure)!,
-        "Warnings": scolors.get(StatusColor.Warnings)!,
-        "Skipped": scolors.get(StatusColor.Skipped)!,
-        "Running": scolors.get(StatusColor.Running)!,
-        "Waiting": scolors.get(StatusColor.Waiting)!,
-        "Ready": scolors.get(StatusColor.Ready)!,
-        "Unspecified": scolors.get(StatusColor.Skipped)!,
-    };
 
     // There are some interesting mappings going on internally with respect to state & outcome. Outcome & State have overlapping "result view" space, and as such we need to 
     // narrowly select the edge cases (e.g. outcome=success && state=completed uniquely signifies a completed run with no errors, otherwise outcome=success it could mean ready, waiting, or running).
