Horde: Add setting for overriding available RAM in agent

#13774

diff --git a/Engine/Source/Programs/Horde/HordeAgent/AgentSettings.cs b/Engine/Source/Programs/Horde/HordeAgent/AgentSettings.cs
--- a/Engine/Source/Programs/Horde/HordeAgent/AgentSettings.cs
+++ b/Engine/Source/Programs/Horde/HordeAgent/AgentSettings.cs
@@ -248,6 +248,11 @@ namespace HordeAgent
 		/// </summary>
 		public double CpuMultiplier { get; set; } = 1.0;
 
+		/// <summary>
+		/// Maximum available RAM in gigabytes. Defaults to available RAM.
+		/// </summary>
+		public int? RamGb { get; set; } = null;
+
 		/// <summary>
 		/// Key/value properties in addition to those set internally by the agent
 		/// </summary>
diff --git a/Engine/Source/Programs/Horde/HordeAgent/Services/CapabilitiesService.cs b/Engine/Source/Programs/Horde/HordeAgent/Services/CapabilitiesService.cs
--- a/Engine/Source/Programs/Horde/HordeAgent/Services/CapabilitiesService.cs
+++ b/Engine/Source/Programs/Horde/HordeAgent/Services/CapabilitiesService.cs
@@ -205,8 +205,7 @@ namespace HordeAgent.Services
 					{
 						long totalCapacity = Int64.Parse(match.Groups[1].Value) * 1024;
 						int ramGb = (int)(totalCapacity / (1024 * 1024 * 1024));
-						capabilities.Resources.Add(KnownResourceNames.Ram, ramGb);
-						capabilities.Properties.Add($"{KnownResourceNames.Ram}={ramGb}");
+						AddMemInfo(_settings.CurrentValue, capabilities, ramGb);
 					}
 				}
 
@@ -313,8 +312,7 @@ namespace HordeAgent.Services
 						string[] parts = memoryNode.InnerText.Split(new char[] { ' ' }, StringSplitOptions.RemoveEmptyEntries);
 						if (parts.Length == 2 && parts[1] == "GB" && Int32.TryParse(parts[0], out int ramGb))
 						{
-							capabilities.Resources.Add(KnownResourceNames.Ram, ramGb);
-							capabilities.Properties.Add($"{KnownResourceNames.Ram}={ramGb}");
+							AddMemInfo(_settings.CurrentValue, capabilities, ramGb);
 						}
 					}
 				}
@@ -542,6 +540,14 @@ namespace HordeAgent.Services
 				capabilities.Properties.Add($"{KnownResourceNames.PhysicalCores}={numPhysicalCores}");
 			}
 		}
+ 
+		public static void AddMemInfo(AgentSettings settings, RpcAgentCapabilities capabilities, int ramGb)
+		{
+			// Agent settings can artificially limit available RAM
+			int adjustedRamGb = settings.RamGb is > 0 ? settings.RamGb.Value : ramGb;
+			capabilities.Resources.Add(KnownResourceNames.Ram, adjustedRamGb);
+			capabilities.Properties.Add($"{KnownResourceNames.Ram}={adjustedRamGb}");
+		}
 
 		static async Task<List<Dictionary<string, string>>?> ReadLinuxHwPropsAsync(string fileName, ILogger logger)
 		{
diff --git a/Engine/Source/Programs/Horde/HordeAgent/Services/WindowsCapabilities.cs b/Engine/Source/Programs/Horde/HordeAgent/Services/WindowsCapabilities.cs
--- a/Engine/Source/Programs/Horde/HordeAgent/Services/WindowsCapabilities.cs
+++ b/Engine/Source/Programs/Horde/HordeAgent/Services/WindowsCapabilities.cs
@@ -50,8 +50,7 @@ public class WindowsCapabilities(IOptionsMonitor<AgentSettings> settings, ILogge
 		if (GlobalMemoryStatusEx(memStatus))
 		{
 			int ramGb = (int)(memStatus.ullTotalPhys / (1024 * 1024 * 1024));
-			SetResource(caps, KnownResourceNames.Ram, ramGb);
-			SetProp(caps, KnownResourceNames.Ram, ramGb);
+			CapabilitiesService.AddMemInfo(settings.CurrentValue, caps, ramGb);
 		}
 		else
 		{
