Horde: Remove Win64 agent type fallback for installer tools

Should only use a licensee-based view when building.

diff --git a/Engine/Source/Programs/Horde/BuildHorde.xml b/Engine/Source/Programs/Horde/BuildHorde.xml
--- a/Engine/Source/Programs/Horde/BuildHorde.xml
+++ b/Engine/Source/Programs/Horde/BuildHorde.xml
@@ -144,7 +144,7 @@
     <!-- Changing the windows architecture value here means updates elsewhere in this XML, and in the Installer wxs file -->
 	<Property Name="ServerArchList" Value="win-x64;osx-arm64;linux-x64"/>
 	
-	<Agent Name="HordeInstallerTools" Type="Win64_Licensee;Win64">
+	<Agent Name="HordeInstallerTools" Type="Win64_Licensee">
 		<Node Name="Build HordeCmd" Produces="#HordeCmdFiles">
 			<Delete Files="$(ServerToolsStagingDir_HordeCmd)/..."/>
 			<DotNet	Arguments="publish &quot;$(RootDir)/Engine/Source/Programs/Horde/HordeCmd/HordeCmd.csproj&quot; --output &quot;$(ServerToolsStagingDir_HordeCmd)&quot; --runtime win-x64 --self-contained $(VersionArguments)" DotNetPath="$(DotNetPath)" />
