Update most nuget packages to latest patch version

#rnx

diff --git a/Engine/Source/Programs/AutomationTool.Tests/AutomationTool.Tests.csproj b/Engine/Source/Programs/AutomationTool.Tests/AutomationTool.Tests.csproj
--- a/Engine/Source/Programs/AutomationTool.Tests/AutomationTool.Tests.csproj
+++ b/Engine/Source/Programs/AutomationTool.Tests/AutomationTool.Tests.csproj
@@ -18,7 +18,7 @@
       <PrivateAssets>all</PrivateAssets>
       <IncludeAssets>runtime; build; native; contentfiles; analyzers; buildtransitive</IncludeAssets>
     </PackageReference>
-    <PackageReference Include="System.Text.Json" Version="8.0.5" />
+    <PackageReference Include="System.Text.Json" Version="8.0.6" />
   </ItemGroup>
 
   <ItemGroup>
diff --git a/Engine/Source/Programs/AutomationTool/Apple/Apple.Automation.csproj b/Engine/Source/Programs/AutomationTool/Apple/Apple.Automation.csproj
--- a/Engine/Source/Programs/AutomationTool/Apple/Apple.Automation.csproj
+++ b/Engine/Source/Programs/AutomationTool/Apple/Apple.Automation.csproj
@@ -41,7 +41,7 @@
 	<Reference Include="BouncyCastle.Crypto, Version=1.8.3.0, Culture=neutral, PublicKeyToken=0e99375e54769942">
 		<HintPath>..\..\..\..\Binaries\ThirdParty\IOS\BouncyCastle.Crypto.dll</HintPath>
 	</Reference>
-	<PackageReference Include="Portable.BouncyCastle" Version="1.8.6.7" />
+	<PackageReference Include="Portable.BouncyCastle" Version="1.8.10" />
 
     <ProjectReference Include="..\..\Shared\EpicGames.Core\EpicGames.Core.csproj" PrivateAssets="All"><Private>false</Private></ProjectReference>
     <ProjectReference Include="..\..\UnrealBuildTool\UnrealBuildTool.csproj" PrivateAssets="All"><Private>false</Private></ProjectReference>
diff --git a/Engine/Source/Programs/AutomationTool/AutomationTool.csproj b/Engine/Source/Programs/AutomationTool/AutomationTool.csproj
--- a/Engine/Source/Programs/AutomationTool/AutomationTool.csproj
+++ b/Engine/Source/Programs/AutomationTool/AutomationTool.csproj
@@ -88,15 +88,15 @@
   </ItemGroup>
 
   <ItemGroup>
-    <PackageReference Include="AWSSDK.SecurityToken" Version="3.7.300.105" />
-    <PackageReference Include="Magick.NET-Q16-HDRI-AnyCPU" Version="14.8.1" />
-    <PackageReference Include="Microsoft.Build" Version="17.11.4" ExcludeAssets="runtime" PrivateAssets="all">
+    <PackageReference Include="AWSSDK.SecurityToken" Version="3.7.501.14" />
+    <PackageReference Include="Magick.NET-Q16-HDRI-AnyCPU" Version="14.8.2" />
+    <PackageReference Include="Microsoft.Build" Version="17.11.31" ExcludeAssets="runtime" PrivateAssets="all">
       <IncludeAssets>compile; build; native; contentfiles; analyzers; buildtransitive</IncludeAssets>
     </PackageReference>
     <PackageReference Include="Microsoft.Build.Locator" Version="1.7.8" />
     <!-- Added as a temporary workaround to running from visual studio where the system version of dotnet is run. 
          Adding a loader context resolves many of the issues, but it still fails when MSBuild tries to write the json -->
-    <PackageReference Include="Newtonsoft.Json" Version="13.0.3" />
+    <PackageReference Include="Newtonsoft.Json" Version="13.0.4" />
   </ItemGroup>
   
   <!-- This set of dependencies should be kept as small as possible -->
diff --git a/Engine/Source/Programs/AutomationTool/AutomationUtils/AutomationUtils.Automation.csproj b/Engine/Source/Programs/AutomationTool/AutomationUtils/AutomationUtils.Automation.csproj
--- a/Engine/Source/Programs/AutomationTool/AutomationUtils/AutomationUtils.Automation.csproj
+++ b/Engine/Source/Programs/AutomationTool/AutomationUtils/AutomationUtils.Automation.csproj
@@ -64,16 +64,16 @@
     <ProjectReference Include="..\..\UnrealBuildTool\UnrealBuildTool.csproj" PrivateAssets="All"><Private>true</Private></ProjectReference>
   </ItemGroup>
   <ItemGroup>
-    <PackageReference Include="Magick.NET-Q16-HDRI-AnyCPU" Version="14.8.1" PrivateAssets="all" />
-    <PackageReference Include="System.Drawing.Common" Version="8.0.10" />
+    <PackageReference Include="Magick.NET-Q16-HDRI-AnyCPU" Version="14.8.2" PrivateAssets="all" />
+    <PackageReference Include="System.Drawing.Common" Version="8.0.20" />
     <PackageReference Include="System.Net.Http" Version="4.3.4" />
 
     <!--Required by Ionic -->
     <PackageReference Include="System.Security.Permissions" Version="8.0.0" />
     <PackageReference Include="System.Text.Encoding.CodePages" Version="8.0.0" />
 
-    <PackageReference Include="Google.Apis.Drive.v3" Version="1.33.1.1225" PrivateAssets="all" />
-    <PackageReference Include="Google.Apis.Sheets.v4" Version="1.33.1.1229" PrivateAssets="all" />
+    <PackageReference Include="Google.Apis.Drive.v3" Version="1.33.1.1232" PrivateAssets="all" />
+    <PackageReference Include="Google.Apis.Sheets.v4" Version="1.33.1.1240" PrivateAssets="all" />
     <PackageReference Include="System.Text.RegularExpressions" Version="4.3.1" />
   </ItemGroup>
 </Project>
\ No newline at end of file
diff --git a/Engine/Source/Programs/AutomationTool/Gauntlet/Gauntlet.Automation.csproj b/Engine/Source/Programs/AutomationTool/Gauntlet/Gauntlet.Automation.csproj
--- a/Engine/Source/Programs/AutomationTool/Gauntlet/Gauntlet.Automation.csproj
+++ b/Engine/Source/Programs/AutomationTool/Gauntlet/Gauntlet.Automation.csproj
@@ -90,7 +90,7 @@
     <ProjectReference Include="..\Android\Android.Automation.csproj" PrivateAssets="All"><Private>false</Private></ProjectReference>
   </ItemGroup>
   <ItemGroup>
-    <PackageReference Include="Magick.NET-Q16-HDRI-AnyCPU" Version="14.8.1" PrivateAssets="all" />
+    <PackageReference Include="Magick.NET-Q16-HDRI-AnyCPU" Version="14.8.2" PrivateAssets="all" />
     <PackageReference Include="Microsoft.CSharp" Version="4.7.0" />
     <PackageReference Include="MySql.Data" Version="6.10.9" PrivateAssets="all" />
   </ItemGroup>
diff --git a/Engine/Source/Programs/AutomationTool/IOS/IOS.Automation.csproj b/Engine/Source/Programs/AutomationTool/IOS/IOS.Automation.csproj
--- a/Engine/Source/Programs/AutomationTool/IOS/IOS.Automation.csproj
+++ b/Engine/Source/Programs/AutomationTool/IOS/IOS.Automation.csproj
@@ -55,6 +55,6 @@
   <ItemGroup>
     <PackageReference Include="Microsoft.CSharp" Version="4.7.0" PrivateAssets="all" />
     <PackageReference Include="System.Resources.Extensions" Version="4.7.1" PrivateAssets="all" />
-    <PackageReference Include="Portable.BouncyCastle" Version="1.8.6.7" PrivateAssets="all" />
+    <PackageReference Include="Portable.BouncyCastle" Version="1.8.10" PrivateAssets="all" />
   </ItemGroup>
 </Project>
\ No newline at end of file
diff --git a/Engine/Source/Programs/AutomationTool/Turnkey/Turnkey.Automation.csproj b/Engine/Source/Programs/AutomationTool/Turnkey/Turnkey.Automation.csproj
--- a/Engine/Source/Programs/AutomationTool/Turnkey/Turnkey.Automation.csproj
+++ b/Engine/Source/Programs/AutomationTool/Turnkey/Turnkey.Automation.csproj
@@ -57,6 +57,6 @@
   <ItemGroup>
     <PackageReference Include="Microsoft.CSharp" Version="4.7.0" />
     <PackageReference Include="System.IO.FileSystem.AccessControl" Version="4.7.0" />
-    <PackageReference Include="Google.Apis.Drive.v3" Version="1.33.1.1225" PrivateAssets="all" />
+    <PackageReference Include="Google.Apis.Drive.v3" Version="1.33.1.1232" PrivateAssets="all" />
   </ItemGroup>
 </Project>
\ No newline at end of file
diff --git a/Engine/Source/Programs/Horde/Drivers/JobDriver.Tests/JobDriver.Tests.csproj b/Engine/Source/Programs/Horde/Drivers/JobDriver.Tests/JobDriver.Tests.csproj
--- a/Engine/Source/Programs/Horde/Drivers/JobDriver.Tests/JobDriver.Tests.csproj
+++ b/Engine/Source/Programs/Horde/Drivers/JobDriver.Tests/JobDriver.Tests.csproj
@@ -19,7 +19,7 @@
   </PropertyGroup>
 
   <ItemGroup>
-    <PackageReference Include="Microsoft.Extensions.Logging.Abstractions" Version="8.0.2" />
+    <PackageReference Include="Microsoft.Extensions.Logging.Abstractions" Version="8.0.3" />
     <PackageReference Include="Microsoft.NET.Test.Sdk" Version="17.11.1" />
     <PackageReference Include="Moq" Version="4.20.70" />
     <PackageReference Include="MSTest.TestAdapter" Version="3.6.2" />
diff --git a/Engine/Source/Programs/Horde/Drivers/JobDriver/JobDriver.csproj b/Engine/Source/Programs/Horde/Drivers/JobDriver/JobDriver.csproj
--- a/Engine/Source/Programs/Horde/Drivers/JobDriver/JobDriver.csproj
+++ b/Engine/Source/Programs/Horde/Drivers/JobDriver/JobDriver.csproj
@@ -52,16 +52,16 @@
   </ItemGroup>
 
   <ItemGroup>
-    <PackageReference Include="AWSSDK.Core" Version="3.7.300.12" />
-    <PackageReference Include="AWSSDK.EC2" Version="3.7.306.6" />
-    <PackageReference Include="AWSSDK.SimpleSystemsManagement" Version="3.7.301.8" />
-    <PackageReference Include="Google.Protobuf" Version="3.25.1" />
+    <PackageReference Include="AWSSDK.Core" Version="3.7.500.23" />
+    <PackageReference Include="AWSSDK.EC2" Version="3.7.514" />
+    <PackageReference Include="AWSSDK.SimpleSystemsManagement" Version="3.7.502" />
+    <PackageReference Include="Google.Protobuf" Version="3.25.8" />
     <PackageReference Include="Grpc.Net.Client" Version="2.59.0" />
     <PackageReference Include="Grpc.Tools" Version="2.59.0">
       <PrivateAssets>all</PrivateAssets>
       <IncludeAssets>runtime; build; native; contentfiles; analyzers; buildtransitive</IncludeAssets>
     </PackageReference>
-    <PackageReference Include="Meziantou.Analyzer" Version="2.0.139">
+    <PackageReference Include="Meziantou.Analyzer" Version="2.0.220">
       <PrivateAssets>all</PrivateAssets>
       <IncludeAssets>runtime; build; native; contentfiles; analyzers; buildtransitive</IncludeAssets>
     </PackageReference>
@@ -70,28 +70,28 @@
     <PackageReference Include="Microsoft.Extensions.Hosting.Abstractions" Version="8.0.1" />
     <PackageReference Include="Microsoft.Extensions.Hosting.WindowsServices" Version="8.0.1" />
     <PackageReference Include="Microsoft.Extensions.Http" Version="8.0.1" />
-    <PackageReference Include="Microsoft.Extensions.Http.Polly" Version="8.0.10" />
-    <PackageReference Include="Microsoft.Extensions.Logging.Abstractions" Version="8.0.2" />
+    <PackageReference Include="Microsoft.Extensions.Http.Polly" Version="8.0.20" />
+    <PackageReference Include="Microsoft.Extensions.Logging.Abstractions" Version="8.0.3" />
     <PackageReference Include="Microsoft.Extensions.Options.DataAnnotations" Version="8.0.0" />
     <PackageReference Include="Microsoft.VisualStudio.Threading.Analyzers" Version="17.11.20">
       <PrivateAssets>all</PrivateAssets>
       <IncludeAssets>runtime; build; native; contentfiles; analyzers; buildtransitive</IncludeAssets>
     </PackageReference>
-    <PackageReference Include="OpenTelemetry" Version="1.8.0" />
-    <PackageReference Include="OpenTelemetry.Exporter.Console" Version="1.8.0" />
-    <PackageReference Include="OpenTelemetry.Exporter.OpenTelemetryProtocol" Version="1.8.0" />
-    <PackageReference Include="OpenTelemetry.Extensions.Hosting" Version="1.8.0" />
+    <PackageReference Include="OpenTelemetry" Version="1.8.1" />
+    <PackageReference Include="OpenTelemetry.Exporter.Console" Version="1.8.1" />
+    <PackageReference Include="OpenTelemetry.Exporter.OpenTelemetryProtocol" Version="1.8.1" />
+    <PackageReference Include="OpenTelemetry.Extensions.Hosting" Version="1.8.1" />
     <PackageReference Include="OpenTelemetry.Instrumentation.GrpcNetClient" Version="1.8.0-beta.1" />
     <PackageReference Include="OpenTelemetry.Instrumentation.Http" Version="1.8.1" />
     <PackageReference Include="Serilog" Version="3.1.1" />
     <PackageReference Include="Serilog.Extensions.Hosting" Version="8.0.0" />
     <PackageReference Include="Serilog.Extensions.Logging" Version="8.0.0" />
-    <PackageReference Include="Serilog.Settings.Configuration" Version="8.0.0" />
+    <PackageReference Include="Serilog.Settings.Configuration" Version="8.0.4" />
     <PackageReference Include="Serilog.Sinks.Console" Version="5.0.1" />
     <PackageReference Include="Serilog.Sinks.File" Version="5.0.0" />
     <PackageReference Include="System.Diagnostics.PerformanceCounter" Version="8.0.1" />
     <PackageReference Include="System.ServiceProcess.ServiceController" Version="8.0.1" />
-    <PackageReference Include="System.Text.Json" Version="8.0.5" />
+    <PackageReference Include="System.Text.Json" Version="8.0.6" />
   </ItemGroup>
   
   <ItemGroup>
@@ -120,7 +120,7 @@
       <FilesToDelete Include="$(PublishDir)**/oo2core*winuwp*.dll" />
     </ItemGroup>
     <Delete Files="@(FilesToDelete)">
-      <Output TaskParameter="DeletedFiles" ItemName="DeletedList"/>
+      <Output TaskParameter="DeletedFiles" ItemName="DeletedList" />
     </Delete>
     <Message Text="Deleted files:%0A@(DeletedList->'  %(Identity)','%0A')" Importance="high" />
   </Target>
diff --git a/Engine/Source/Programs/Horde/HordeAgent.Tests/HordeAgent.Tests.csproj b/Engine/Source/Programs/Horde/HordeAgent.Tests/HordeAgent.Tests.csproj
--- a/Engine/Source/Programs/Horde/HordeAgent.Tests/HordeAgent.Tests.csproj
+++ b/Engine/Source/Programs/Horde/HordeAgent.Tests/HordeAgent.Tests.csproj
@@ -20,7 +20,7 @@
   </PropertyGroup>
 
   <ItemGroup>
-    <PackageReference Include="Microsoft.Extensions.Logging.Abstractions" Version="8.0.2" />
+    <PackageReference Include="Microsoft.Extensions.Logging.Abstractions" Version="8.0.3" />
     <PackageReference Include="Microsoft.NET.Test.Sdk" Version="17.11.1" />
     <PackageReference Include="Moq" Version="4.20.70" />
     <PackageReference Include="MSTest.TestAdapter" Version="3.6.2" />
diff --git a/Engine/Source/Programs/Horde/HordeAgent/HordeAgent.csproj b/Engine/Source/Programs/Horde/HordeAgent/HordeAgent.csproj
--- a/Engine/Source/Programs/Horde/HordeAgent/HordeAgent.csproj
+++ b/Engine/Source/Programs/Horde/HordeAgent/HordeAgent.csproj
@@ -34,16 +34,16 @@
   </ItemGroup>
 
   <ItemGroup>
-    <PackageReference Include="AWSSDK.Core" Version="3.7.300.12" />
-    <PackageReference Include="AWSSDK.EC2" Version="3.7.306.6" />
-    <PackageReference Include="AWSSDK.SimpleSystemsManagement" Version="3.7.301.8" />
-    <PackageReference Include="Google.Protobuf" Version="3.25.1" />
+    <PackageReference Include="AWSSDK.Core" Version="3.7.500.23" />
+    <PackageReference Include="AWSSDK.EC2" Version="3.7.514" />
+    <PackageReference Include="AWSSDK.SimpleSystemsManagement" Version="3.7.502" />
+    <PackageReference Include="Google.Protobuf" Version="3.25.8" />
     <PackageReference Include="Grpc.Net.Client" Version="2.59.0" />
     <PackageReference Include="Grpc.Tools" Version="2.59.0">
       <PrivateAssets>all</PrivateAssets>
       <IncludeAssets>runtime; build; native; contentfiles; analyzers; buildtransitive</IncludeAssets>
     </PackageReference>
-    <PackageReference Include="Meziantou.Analyzer" Version="2.0.139">
+    <PackageReference Include="Meziantou.Analyzer" Version="2.0.220">
       <PrivateAssets>all</PrivateAssets>
       <IncludeAssets>runtime; build; native; contentfiles; analyzers; buildtransitive</IncludeAssets>
     </PackageReference>
@@ -52,28 +52,28 @@
     <PackageReference Include="Microsoft.Extensions.Hosting.Abstractions" Version="8.0.1" />
     <PackageReference Include="Microsoft.Extensions.Hosting.WindowsServices" Version="8.0.1" />
     <PackageReference Include="Microsoft.Extensions.Http" Version="8.0.1" />
-    <PackageReference Include="Microsoft.Extensions.Http.Polly" Version="8.0.10" />
-    <PackageReference Include="Microsoft.Extensions.Logging.Abstractions" Version="8.0.2" />
+    <PackageReference Include="Microsoft.Extensions.Http.Polly" Version="8.0.20" />
+    <PackageReference Include="Microsoft.Extensions.Logging.Abstractions" Version="8.0.3" />
     <PackageReference Include="Microsoft.Extensions.Options.DataAnnotations" Version="8.0.0" />
     <PackageReference Include="Microsoft.VisualStudio.Threading.Analyzers" Version="17.11.20">
       <PrivateAssets>all</PrivateAssets>
       <IncludeAssets>runtime; build; native; contentfiles; analyzers; buildtransitive</IncludeAssets>
     </PackageReference>
-    <PackageReference Include="OpenTelemetry" Version="1.8.0" />
-    <PackageReference Include="OpenTelemetry.Exporter.Console" Version="1.8.0" />
-    <PackageReference Include="OpenTelemetry.Exporter.OpenTelemetryProtocol" Version="1.8.0" />
-    <PackageReference Include="OpenTelemetry.Extensions.Hosting" Version="1.8.0" />
+    <PackageReference Include="OpenTelemetry" Version="1.8.1" />
+    <PackageReference Include="OpenTelemetry.Exporter.Console" Version="1.8.1" />
+    <PackageReference Include="OpenTelemetry.Exporter.OpenTelemetryProtocol" Version="1.8.1" />
+    <PackageReference Include="OpenTelemetry.Extensions.Hosting" Version="1.8.1" />
     <PackageReference Include="OpenTelemetry.Instrumentation.GrpcNetClient" Version="1.8.0-beta.1" />
     <PackageReference Include="OpenTelemetry.Instrumentation.Http" Version="1.8.1" />
     <PackageReference Include="Serilog" Version="3.1.1" />
     <PackageReference Include="Serilog.Extensions.Hosting" Version="8.0.0" />
     <PackageReference Include="Serilog.Extensions.Logging" Version="8.0.0" />
-    <PackageReference Include="Serilog.Settings.Configuration" Version="8.0.0" />
+    <PackageReference Include="Serilog.Settings.Configuration" Version="8.0.4" />
     <PackageReference Include="Serilog.Sinks.Console" Version="5.0.1" />
     <PackageReference Include="Serilog.Sinks.File" Version="5.0.0" />
     <PackageReference Include="System.Diagnostics.PerformanceCounter" Version="8.0.1" />
     <PackageReference Include="System.ServiceProcess.ServiceController" Version="8.0.1" />
-    <PackageReference Include="System.Text.Json" Version="8.0.5" />
+    <PackageReference Include="System.Text.Json" Version="8.0.6" />
   </ItemGroup>
 
   <ItemGroup>
@@ -207,7 +207,7 @@
       <FilesToDelete Include="$(PublishDir)**/oo2core*winuwp*.dll" />
     </ItemGroup>
     <Delete Files="@(FilesToDelete)">
-      <Output TaskParameter="DeletedFiles" ItemName="DeletedList"/>
+      <Output TaskParameter="DeletedFiles" ItemName="DeletedList" />
     </Delete>
     <Message Text="Deleted files:%0A@(DeletedList->'  %(Identity)','%0A')" Importance="high" />
   </Target>
diff --git a/Engine/Source/Programs/Horde/HordeCmd/HordeCmd.csproj b/Engine/Source/Programs/Horde/HordeCmd/HordeCmd.csproj
--- a/Engine/Source/Programs/Horde/HordeCmd/HordeCmd.csproj
+++ b/Engine/Source/Programs/Horde/HordeCmd/HordeCmd.csproj
@@ -30,10 +30,10 @@
     <PackageReference Include="Serilog" Version="3.1.1" />
     <PackageReference Include="Serilog.Extensions.Hosting" Version="8.0.0" />
     <PackageReference Include="Serilog.Extensions.Logging" Version="8.0.0" />
-    <PackageReference Include="Serilog.Settings.Configuration" Version="8.0.0" />
+    <PackageReference Include="Serilog.Settings.Configuration" Version="8.0.4" />
     <PackageReference Include="Serilog.Sinks.Console" Version="5.0.1" />
     <PackageReference Include="Serilog.Sinks.File" Version="5.0.0" />
-    <PackageReference Include="System.Text.Json" Version="8.0.5" />
+    <PackageReference Include="System.Text.Json" Version="8.0.6" />
   </ItemGroup>
 
   <ItemGroup>
diff --git a/Engine/Source/Programs/Horde/HordeServer.Shared/HordeServer.Shared.csproj b/Engine/Source/Programs/Horde/HordeServer.Shared/HordeServer.Shared.csproj
--- a/Engine/Source/Programs/Horde/HordeServer.Shared/HordeServer.Shared.csproj
+++ b/Engine/Source/Programs/Horde/HordeServer.Shared/HordeServer.Shared.csproj
@@ -29,9 +29,9 @@
 
   <ItemGroup>
     <PackageReference Include="MongoDB.Driver" Version="2.22.0" />
-    <PackageReference Include="OpenTelemetry" Version="1.8.0" />
-    <PackageReference Include="OpenTelemetry.Exporter.Console" Version="1.8.0" />
-    <PackageReference Include="OpenTelemetry.Exporter.OpenTelemetryProtocol" Version="1.8.0" />
+    <PackageReference Include="OpenTelemetry" Version="1.8.1" />
+    <PackageReference Include="OpenTelemetry.Exporter.Console" Version="1.8.1" />
+    <PackageReference Include="OpenTelemetry.Exporter.OpenTelemetryProtocol" Version="1.8.1" />
   </ItemGroup>
 
 </Project>
diff --git a/Engine/Source/Programs/Horde/HordeServer.Tests.Shared/HordeServer.Tests.Shared.csproj b/Engine/Source/Programs/Horde/HordeServer.Tests.Shared/HordeServer.Tests.Shared.csproj
--- a/Engine/Source/Programs/Horde/HordeServer.Tests.Shared/HordeServer.Tests.Shared.csproj
+++ b/Engine/Source/Programs/Horde/HordeServer.Tests.Shared/HordeServer.Tests.Shared.csproj
@@ -22,7 +22,7 @@
   </ItemGroup>
 
   <ItemGroup>
-    <PackageReference Include="Microsoft.AspNetCore.Mvc.Testing" Version="8.0.10" />
+    <PackageReference Include="Microsoft.AspNetCore.Mvc.Testing" Version="8.0.20" />
     <PackageReference Include="Microsoft.NET.Test.Sdk" Version="17.11.1" />
     <PackageReference Include="Moq" Version="4.20.70" />
     <PackageReference Include="MSTest.TestAdapter" Version="3.6.2" />
diff --git a/Engine/Source/Programs/Horde/HordeServer.Tests/HordeServer.Tests.csproj b/Engine/Source/Programs/Horde/HordeServer.Tests/HordeServer.Tests.csproj
--- a/Engine/Source/Programs/Horde/HordeServer.Tests/HordeServer.Tests.csproj
+++ b/Engine/Source/Programs/Horde/HordeServer.Tests/HordeServer.Tests.csproj
@@ -19,7 +19,7 @@
   </PropertyGroup>
 
   <ItemGroup>
-    <PackageReference Include="Microsoft.AspNetCore.Mvc.Testing" Version="8.0.10" />
+    <PackageReference Include="Microsoft.AspNetCore.Mvc.Testing" Version="8.0.20" />
     <PackageReference Include="Microsoft.NET.Test.Sdk" Version="17.11.1" />
     <PackageReference Include="Moq" Version="4.20.70" />
     <PackageReference Include="MSTest.TestAdapter" Version="3.6.2" />
diff --git a/Engine/Source/Programs/Horde/HordeServer/HordeServer.csproj b/Engine/Source/Programs/Horde/HordeServer/HordeServer.csproj
--- a/Engine/Source/Programs/Horde/HordeServer/HordeServer.csproj
+++ b/Engine/Source/Programs/Horde/HordeServer/HordeServer.csproj
@@ -20,14 +20,14 @@
   </PropertyGroup>
 
   <ItemGroup>
-    <PackageReference Include="AWSSDK.AutoScaling" Version="3.7.301.9" />
-    <PackageReference Include="AWSSDK.CloudWatch" Version="3.7.300.12" />
-    <PackageReference Include="AWSSDK.EC2" Version="3.7.306.6" />
-    <PackageReference Include="AWSSDK.Extensions.NETCore.Setup" Version="3.7.300" />
-    <PackageReference Include="AWSSDK.S3" Version="3.7.304.1" />
-    <PackageReference Include="AWSSDK.SimpleSystemsManagement" Version="3.7.304.24" />
-    <PackageReference Include="AWSSDK.SQS" Version="3.7.300.12" />
-    <PackageReference Include="AWSSDK.SecurityToken" Version="3.7.300.13" />
+    <PackageReference Include="AWSSDK.AutoScaling" Version="3.7.501.7" />
+    <PackageReference Include="AWSSDK.CloudWatch" Version="3.7.502.7" />
+    <PackageReference Include="AWSSDK.EC2" Version="3.7.514" />
+    <PackageReference Include="AWSSDK.Extensions.NETCore.Setup" Version="3.7.400" />
+    <PackageReference Include="AWSSDK.S3" Version="3.7.501.8" />
+    <PackageReference Include="AWSSDK.SimpleSystemsManagement" Version="3.7.502" />
+    <PackageReference Include="AWSSDK.SQS" Version="3.7.501.14" />
+    <PackageReference Include="AWSSDK.SecurityToken" Version="3.7.501.14" />
     <PackageReference Include="Azure.Storage.Blobs" Version="12.19.1" />
     <PackageReference Include="Crc32.NET" Version="1.2.0" />
     <PackageReference Include="Grpc.AspNetCore" Version="2.59.0" />
@@ -38,51 +38,51 @@
       <IncludeAssets>runtime; build; native; contentfiles; analyzers; buildtransitive</IncludeAssets>
     </PackageReference>
     <PackageReference Include="JetBrains.Profiler.SelfApi" Version="2.5.0" />
-    <PackageReference Include="JsonPath.Net" Version="0.6.7" />
-    <PackageReference Include="Meziantou.Analyzer" Version="2.0.139">
+    <PackageReference Include="JsonPath.Net" Version="0.6.8" />
+    <PackageReference Include="Meziantou.Analyzer" Version="2.0.220">
       <PrivateAssets>all</PrivateAssets>
       <IncludeAssets>runtime; build; native; contentfiles; analyzers; buildtransitive</IncludeAssets>
     </PackageReference>
     <PackageReference Include="Microsoft.AspNetCore.Authentication.Cookies" Version="2.2.0" />
-    <PackageReference Include="Microsoft.AspNetCore.Authentication.JwtBearer" Version="8.0.10" />
-    <PackageReference Include="Microsoft.AspNetCore.Authentication.OpenIdConnect" Version="8.0.10" />
-    <PackageReference Include="Microsoft.AspNetCore.DataProtection.StackExchangeRedis" Version="8.0.10" />
-    <PackageReference Include="Microsoft.AspNetCore.Mvc.Razor.RuntimeCompilation" Version="8.0.10" />
-    <PackageReference Include="Microsoft.AspNetCore.SpaServices.Extensions" Version="8.0.10" />
+    <PackageReference Include="Microsoft.AspNetCore.Authentication.JwtBearer" Version="8.0.20" />
+    <PackageReference Include="Microsoft.AspNetCore.Authentication.OpenIdConnect" Version="8.0.20" />
+    <PackageReference Include="Microsoft.AspNetCore.DataProtection.StackExchangeRedis" Version="8.0.20" />
+    <PackageReference Include="Microsoft.AspNetCore.Mvc.Razor.RuntimeCompilation" Version="8.0.20" />
+    <PackageReference Include="Microsoft.AspNetCore.SpaServices.Extensions" Version="8.0.20" />
     <PackageReference Include="Microsoft.Extensions.Caching.Memory" Version="8.0.1" />
     <PackageReference Include="Microsoft.Extensions.Hosting.WindowsServices" Version="8.0.1" />
     <PackageReference Include="Microsoft.Extensions.Http" Version="8.0.1" />
-    <PackageReference Include="Microsoft.Extensions.Http.Polly" Version="8.0.10" />
+    <PackageReference Include="Microsoft.Extensions.Http.Polly" Version="8.0.20" />
     <PackageReference Include="Microsoft.Extensions.Logging" Version="8.0.1" />
     <PackageReference Include="Microsoft.VisualStudio.Threading.Analyzers" Version="17.11.20">
       <PrivateAssets>all</PrivateAssets>
       <IncludeAssets>runtime; build; native; contentfiles; analyzers; buildtransitive</IncludeAssets>
     </PackageReference>
     <PackageReference Include="MongoDB.Driver" Version="2.22.0" />
-    <PackageReference Include="OpenTelemetry" Version="1.8.0" />
-    <PackageReference Include="OpenTelemetry.Exporter.Console" Version="1.8.0" />
-    <PackageReference Include="OpenTelemetry.Exporter.OpenTelemetryProtocol" Version="1.8.0" />
-    <PackageReference Include="OpenTelemetry.Extensions.Hosting" Version="1.8.0" />
+    <PackageReference Include="OpenTelemetry" Version="1.8.1" />
+    <PackageReference Include="OpenTelemetry.Exporter.Console" Version="1.8.1" />
+    <PackageReference Include="OpenTelemetry.Exporter.OpenTelemetryProtocol" Version="1.8.1" />
+    <PackageReference Include="OpenTelemetry.Extensions.Hosting" Version="1.8.1" />
     <PackageReference Include="OpenTelemetry.Instrumentation.AspNetCore" Version="1.8.1" />
     <PackageReference Include="OpenTelemetry.Instrumentation.GrpcNetClient" Version="1.8.0-beta.1" />
     <PackageReference Include="OpenTelemetry.Instrumentation.Http" Version="1.8.1" />
-    <PackageReference Include="OpenTelemetry.Instrumentation.StackExchangeRedis" Version="1.0.0-rc9.14" />
-    <PackageReference Include="Polly" Version="8.2.0" />
+    <PackageReference Include="OpenTelemetry.Instrumentation.StackExchangeRedis" Version="1.0.0-rc9.15" />
+    <PackageReference Include="Polly" Version="8.2.1" />
     <PackageReference Include="Polly.Extensions.Http" Version="3.0.0" />
-    <PackageReference Include="Serilog.AspNetCore" Version="8.0.0" />
+    <PackageReference Include="Serilog.AspNetCore" Version="8.0.3" />
     <PackageReference Include="Serilog.Exceptions.Grpc" Version="8.4.0" />
     <PackageReference Include="Serilog.Extensions.Hosting" Version="8.0.0" />
     <PackageReference Include="Serilog.Sinks.Console" Version="5.0.1" />
     <PackageReference Include="Serilog.Sinks.File" Version="5.0.0" />
     <PackageReference Include="SerilogAnalyzer" Version="0.15.0" />
     <PackageReference Include="SharpZipLib" Version="1.4.2" />
-    <PackageReference Include="StackExchange.Redis" Version="2.8.16" />
+    <PackageReference Include="StackExchange.Redis" Version="2.8.58" />
     <PackageReference Include="Swashbuckle.AspNetCore" Version="6.5.0" />
-    <PackageReference Include="System.Linq.Async" Version="6.0.1" />
-    <PackageReference Include="System.Text.Json" Version="8.0.5" />
+    <PackageReference Include="System.Linq.Async" Version="6.0.3" />
+    <PackageReference Include="System.Text.Json" Version="8.0.6" />
     <PackageReference Include="T-Digest.NET" Version="1.0.3" />
     <PackageReference Include="TimeZoneConverter" Version="6.1.0" />
-    <PackageReference Include="ZstdSharp.Port" Version="0.8.1" />
+    <PackageReference Include="ZstdSharp.Port" Version="0.8.6" />
   </ItemGroup>
 
   <ItemGroup>
diff --git a/Engine/Source/Programs/Horde/Plugins/Analytics/HordeServer.Analytics/HordeServer.Analytics.csproj b/Engine/Source/Programs/Horde/Plugins/Analytics/HordeServer.Analytics/HordeServer.Analytics.csproj
--- a/Engine/Source/Programs/Horde/Plugins/Analytics/HordeServer.Analytics/HordeServer.Analytics.csproj
+++ b/Engine/Source/Programs/Horde/Plugins/Analytics/HordeServer.Analytics/HordeServer.Analytics.csproj
@@ -24,8 +24,8 @@
   </ItemGroup>
 
   <ItemGroup>
-    <PackageReference Include="JsonPath.Net" Version="0.6.7" />
-    <PackageReference Include="Meziantou.Analyzer" Version="2.0.139">
+    <PackageReference Include="JsonPath.Net" Version="0.6.8" />
+    <PackageReference Include="Meziantou.Analyzer" Version="2.0.220">
       <PrivateAssets>all</PrivateAssets>
       <IncludeAssets>runtime; build; native; contentfiles; analyzers; buildtransitive</IncludeAssets>
     </PackageReference>
diff --git a/Engine/Source/Programs/Horde/Plugins/Build/HordeServer.Build.Tests/HordeServer.Build.Tests.csproj b/Engine/Source/Programs/Horde/Plugins/Build/HordeServer.Build.Tests/HordeServer.Build.Tests.csproj
--- a/Engine/Source/Programs/Horde/Plugins/Build/HordeServer.Build.Tests/HordeServer.Build.Tests.csproj
+++ b/Engine/Source/Programs/Horde/Plugins/Build/HordeServer.Build.Tests/HordeServer.Build.Tests.csproj
@@ -27,7 +27,7 @@
     <PackageReference Include="Microsoft.NET.Test.Sdk" Version="17.11.1" />
     <PackageReference Include="MSTest.TestAdapter" Version="3.6.2" />
     <PackageReference Include="MSTest.TestFramework" Version="3.6.2" />
-    <PackageReference Include="Microsoft.AspNetCore.Mvc.Testing" Version="8.0.10" />
+    <PackageReference Include="Microsoft.AspNetCore.Mvc.Testing" Version="8.0.20" />
     <PackageReference Include="Moq" Version="4.20.70" />
     <PackageReference Include="Microsoft.Extensions.Caching.Memory" Version="8.0.1" />
   </ItemGroup>
diff --git a/Engine/Source/Programs/Horde/Plugins/Build/HordeServer.Build/HordeServer.Build.csproj b/Engine/Source/Programs/Horde/Plugins/Build/HordeServer.Build/HordeServer.Build.csproj
--- a/Engine/Source/Programs/Horde/Plugins/Build/HordeServer.Build/HordeServer.Build.csproj
+++ b/Engine/Source/Programs/Horde/Plugins/Build/HordeServer.Build/HordeServer.Build.csproj
@@ -25,7 +25,7 @@
       <PrivateAssets>all</PrivateAssets>
       <IncludeAssets>runtime; build; native; contentfiles; analyzers; buildtransitive</IncludeAssets>
     </PackageReference>
-    <PackageReference Include="Polly" Version="8.2.0" />
+    <PackageReference Include="Polly" Version="8.2.1" />
     <PackageReference Include="TimeZoneConverter" Version="6.1.0" />
   </ItemGroup>
 
diff --git a/Engine/Source/Programs/Horde/Plugins/Compute/HordeServer.Compute/Aws/AwsCloudWatchMultiplexer.cs b/Engine/Source/Programs/Horde/Plugins/Compute/HordeServer.Compute/Aws/AwsCloudWatchMultiplexer.cs
--- a/Engine/Source/Programs/Horde/Plugins/Compute/HordeServer.Compute/Aws/AwsCloudWatchMultiplexer.cs
+++ b/Engine/Source/Programs/Horde/Plugins/Compute/HordeServer.Compute/Aws/AwsCloudWatchMultiplexer.cs
@@ -78,6 +78,7 @@ public class AwsCloudWatchMultiplexer : IAmazonCloudWatch
 	public Task<DeleteDashboardsResponse> DeleteDashboardsAsync(DeleteDashboardsRequest request, CancellationToken cancellationToken = default) { throw new NotImplementedException(); }
 	public Task<DeleteInsightRulesResponse> DeleteInsightRulesAsync(DeleteInsightRulesRequest request, CancellationToken cancellationToken = default) { throw new NotImplementedException(); }
 	public Task<DeleteMetricStreamResponse> DeleteMetricStreamAsync(DeleteMetricStreamRequest request, CancellationToken cancellationToken = default) { throw new NotImplementedException(); }
+	public Task<DescribeAlarmContributorsResponse> DescribeAlarmContributorsAsync(DescribeAlarmContributorsRequest request, CancellationToken cancellationToken) { throw new NotImplementedException(); }
 	public Task<DescribeAlarmHistoryResponse> DescribeAlarmHistoryAsync(CancellationToken cancellationToken = default) { throw new NotImplementedException(); }
 	public Task<DescribeAlarmHistoryResponse> DescribeAlarmHistoryAsync(DescribeAlarmHistoryRequest request, CancellationToken cancellationToken = default) { throw new NotImplementedException(); }
 	public Task<DescribeAlarmsResponse> DescribeAlarmsAsync(CancellationToken cancellationToken = default) { throw new NotImplementedException(); }
diff --git a/Engine/Source/Programs/Horde/Plugins/Compute/HordeServer.Compute/HordeServer.Compute.csproj b/Engine/Source/Programs/Horde/Plugins/Compute/HordeServer.Compute/HordeServer.Compute.csproj
--- a/Engine/Source/Programs/Horde/Plugins/Compute/HordeServer.Compute/HordeServer.Compute.csproj
+++ b/Engine/Source/Programs/Horde/Plugins/Compute/HordeServer.Compute/HordeServer.Compute.csproj
@@ -44,14 +44,14 @@
   </ItemGroup>
 
   <ItemGroup>
-    <PackageReference Include="AWSSDK.AutoScaling" Version="3.7.301.9" />
-    <PackageReference Include="AWSSDK.CloudWatch" Version="3.7.300.12" />
-    <PackageReference Include="AWSSDK.EC2" Version="3.7.306.6" />
-    <PackageReference Include="AWSSDK.Extensions.NETCore.Setup" Version="3.7.300" />
-    <PackageReference Include="AWSSDK.S3" Version="3.7.304.1" />
-    <PackageReference Include="AWSSDK.SimpleSystemsManagement" Version="3.7.304.24" />
-    <PackageReference Include="AWSSDK.SQS" Version="3.7.300.12" />
-    <PackageReference Include="AWSSDK.SecurityToken" Version="3.7.300.13" />
+    <PackageReference Include="AWSSDK.AutoScaling" Version="3.7.501.7" />
+    <PackageReference Include="AWSSDK.CloudWatch" Version="3.7.502.7" />
+    <PackageReference Include="AWSSDK.EC2" Version="3.7.514" />
+    <PackageReference Include="AWSSDK.Extensions.NETCore.Setup" Version="3.7.400" />
+    <PackageReference Include="AWSSDK.S3" Version="3.7.501.8" />
+    <PackageReference Include="AWSSDK.SimpleSystemsManagement" Version="3.7.502" />
+    <PackageReference Include="AWSSDK.SQS" Version="3.7.501.14" />
+    <PackageReference Include="AWSSDK.SecurityToken" Version="3.7.501.14" />
     <PackageReference Include="Azure.Storage.Blobs" Version="12.19.1" />
     <PackageReference Include="MongoDB.Driver" Version="2.22.0" />
     <PackageReference Include="Grpc.AspNetCore" Version="2.59.0" />
@@ -61,7 +61,7 @@
       <PrivateAssets>all</PrivateAssets>
       <IncludeAssets>runtime; build; native; contentfiles; analyzers; buildtransitive</IncludeAssets>
     </PackageReference>
-    <PackageReference Include="System.Linq.Async" Version="6.0.1" />
+    <PackageReference Include="System.Linq.Async" Version="6.0.3" />
   </ItemGroup>
 
   <ItemGroup>
diff --git a/Engine/Source/Programs/Horde/Plugins/Ddc/HordeServer.Ddc/HordeServer.Ddc.csproj b/Engine/Source/Programs/Horde/Plugins/Ddc/HordeServer.Ddc/HordeServer.Ddc.csproj
--- a/Engine/Source/Programs/Horde/Plugins/Ddc/HordeServer.Ddc/HordeServer.Ddc.csproj
+++ b/Engine/Source/Programs/Horde/Plugins/Ddc/HordeServer.Ddc/HordeServer.Ddc.csproj
@@ -23,7 +23,7 @@
   <ItemGroup>
     <PackageReference Include="Crc32.NET" Version="1.2.0" />
     <PackageReference Include="Serilog.Extensions.Hosting" Version="8.0.0" />
-    <PackageReference Include="System.Linq.Async" Version="6.0.1" />
+    <PackageReference Include="System.Linq.Async" Version="6.0.3" />
   </ItemGroup>
 
   <ItemGroup>
diff --git a/Engine/Source/Programs/Horde/Plugins/Experimental/HordeServer.Experimental.Tests/HordeServer.Experimental.Tests.csproj b/Engine/Source/Programs/Horde/Plugins/Experimental/HordeServer.Experimental.Tests/HordeServer.Experimental.Tests.csproj
--- a/Engine/Source/Programs/Horde/Plugins/Experimental/HordeServer.Experimental.Tests/HordeServer.Experimental.Tests.csproj
+++ b/Engine/Source/Programs/Horde/Plugins/Experimental/HordeServer.Experimental.Tests/HordeServer.Experimental.Tests.csproj
@@ -27,7 +27,7 @@
     <PackageReference Include="Microsoft.NET.Test.Sdk" Version="17.11.1" />
     <PackageReference Include="MSTest.TestAdapter" Version="3.6.2" />
     <PackageReference Include="MSTest.TestFramework" Version="3.6.2" />
-    <PackageReference Include="Microsoft.AspNetCore.Mvc.Testing" Version="8.0.10" />
+    <PackageReference Include="Microsoft.AspNetCore.Mvc.Testing" Version="8.0.20" />
     <PackageReference Include="Moq" Version="4.20.70" />
     <PackageReference Include="Microsoft.Extensions.Caching.Memory" Version="8.0.1" />
   </ItemGroup>
diff --git a/Engine/Source/Programs/Horde/Plugins/Experimental/HordeServer.Experimental/HordeServer.Experimental.csproj b/Engine/Source/Programs/Horde/Plugins/Experimental/HordeServer.Experimental/HordeServer.Experimental.csproj
--- a/Engine/Source/Programs/Horde/Plugins/Experimental/HordeServer.Experimental/HordeServer.Experimental.csproj
+++ b/Engine/Source/Programs/Horde/Plugins/Experimental/HordeServer.Experimental/HordeServer.Experimental.csproj
@@ -25,7 +25,7 @@
       <PrivateAssets>all</PrivateAssets>
       <IncludeAssets>runtime; build; native; contentfiles; analyzers; buildtransitive</IncludeAssets>
     </PackageReference>
-    <PackageReference Include="Polly" Version="8.2.0" />
+    <PackageReference Include="Polly" Version="8.2.1" />
     <PackageReference Include="TimeZoneConverter" Version="6.1.0" />
   </ItemGroup>
 
diff --git a/Engine/Source/Programs/Horde/Plugins/Secrets/HordeServer.Secrets/HordeServer.Secrets.csproj b/Engine/Source/Programs/Horde/Plugins/Secrets/HordeServer.Secrets/HordeServer.Secrets.csproj
--- a/Engine/Source/Programs/Horde/Plugins/Secrets/HordeServer.Secrets/HordeServer.Secrets.csproj
+++ b/Engine/Source/Programs/Horde/Plugins/Secrets/HordeServer.Secrets/HordeServer.Secrets.csproj
@@ -22,7 +22,7 @@
   </ItemGroup>
 
   <ItemGroup>
-    <PackageReference Include="AWSSDK.SimpleSystemsManagement" Version="3.7.304.24" />
+    <PackageReference Include="AWSSDK.SimpleSystemsManagement" Version="3.7.502" />
     <PackageReference Include="Azure.Identity" Version="1.14.2" />
     <PackageReference Include="Azure.Security.KeyVault.Secrets" Version="4.8.0" />
   </ItemGroup>
diff --git a/Engine/Source/Programs/Horde/Plugins/Storage/HordeServer.Storage/HordeServer.Storage.csproj b/Engine/Source/Programs/Horde/Plugins/Storage/HordeServer.Storage/HordeServer.Storage.csproj
--- a/Engine/Source/Programs/Horde/Plugins/Storage/HordeServer.Storage/HordeServer.Storage.csproj
+++ b/Engine/Source/Programs/Horde/Plugins/Storage/HordeServer.Storage/HordeServer.Storage.csproj
@@ -21,8 +21,8 @@
   </ItemGroup>
 
   <ItemGroup>
-    <PackageReference Include="AWSSDK.S3" Version="3.7.304.1" />
-    <PackageReference Include="AWSSDK.Extensions.NETCore.Setup" Version="3.7.300" />
+    <PackageReference Include="AWSSDK.S3" Version="3.7.501.8" />
+    <PackageReference Include="AWSSDK.Extensions.NETCore.Setup" Version="3.7.400" />
     <PackageReference Include="Azure.Storage.Blobs" Version="12.19.1" />
     <PackageReference Include="Google.Cloud.Storage.V1" Version="4.10.0" />
     <PackageReference Include="MongoDB.Driver" Version="2.22.0" />
diff --git a/Engine/Source/Programs/Horde/Tools/ReleaseNotes/ReleaseNotes.csproj b/Engine/Source/Programs/Horde/Tools/ReleaseNotes/ReleaseNotes.csproj
--- a/Engine/Source/Programs/Horde/Tools/ReleaseNotes/ReleaseNotes.csproj
+++ b/Engine/Source/Programs/Horde/Tools/ReleaseNotes/ReleaseNotes.csproj
@@ -10,7 +10,7 @@
 
   <ItemGroup>
     <PackageReference Include="Microsoft.Extensions.Logging.Console" Version="8.0.1" />
-    <PackageReference Include="System.Text.Json" Version="8.0.5" />
+    <PackageReference Include="System.Text.Json" Version="8.0.6" />
   </ItemGroup>
 
   <ItemGroup>
diff --git a/Engine/Source/Programs/Shared/EpicGames.AspNet/EpicGames.AspNet.csproj b/Engine/Source/Programs/Shared/EpicGames.AspNet/EpicGames.AspNet.csproj
--- a/Engine/Source/Programs/Shared/EpicGames.AspNet/EpicGames.AspNet.csproj
+++ b/Engine/Source/Programs/Shared/EpicGames.AspNet/EpicGames.AspNet.csproj
@@ -44,6 +44,6 @@
   </ItemGroup>
   <ItemGroup>
     <PackageReference Include="Microsoft.CSharp" Version="4.7.0" />
-    <PackageReference Include="System.Text.Json" Version="8.0.5" />
+    <PackageReference Include="System.Text.Json" Version="8.0.6" />
   </ItemGroup>
 </Project>
\ No newline at end of file
diff --git a/Engine/Source/Programs/Shared/EpicGames.BuildGraph.Tests/EpicGames.BuildGraph.Tests.csproj b/Engine/Source/Programs/Shared/EpicGames.BuildGraph.Tests/EpicGames.BuildGraph.Tests.csproj
--- a/Engine/Source/Programs/Shared/EpicGames.BuildGraph.Tests/EpicGames.BuildGraph.Tests.csproj
+++ b/Engine/Source/Programs/Shared/EpicGames.BuildGraph.Tests/EpicGames.BuildGraph.Tests.csproj
@@ -53,7 +53,7 @@
       <PrivateAssets>all</PrivateAssets>
       <IncludeAssets>runtime; build; native; contentfiles; analyzers; buildtransitive</IncludeAssets>
     </PackageReference>
-    <PackageReference Include="System.Text.Json" Version="8.0.5" />
+    <PackageReference Include="System.Text.Json" Version="8.0.6" />
   </ItemGroup>
 
   <ItemGroup>
diff --git a/Engine/Source/Programs/Shared/EpicGames.Core/EpicGames.Core.csproj b/Engine/Source/Programs/Shared/EpicGames.Core/EpicGames.Core.csproj
--- a/Engine/Source/Programs/Shared/EpicGames.Core/EpicGames.Core.csproj
+++ b/Engine/Source/Programs/Shared/EpicGames.Core/EpicGames.Core.csproj
@@ -52,11 +52,11 @@
   <ItemGroup>
     <PackageReference Include="JetBrains.Annotations" Version="2024.3.0" />
     <PackageReference Include="Microsoft.Extensions.Configuration" Version="8.0.0" />
-    <PackageReference Include="Microsoft.Extensions.Http.Polly" Version="8.0.10" />
+    <PackageReference Include="Microsoft.Extensions.Http.Polly" Version="8.0.20" />
     <PackageReference Include="Microsoft.Extensions.Logging.Console" Version="8.0.1" />
-    <PackageReference Include="Microsoft.Extensions.ObjectPool" Version="8.0.10" />
+    <PackageReference Include="Microsoft.Extensions.ObjectPool" Version="8.0.20" />
     <PackageReference Include="OpenTracing" Version="0.12.1" />
     <PackageReference Include="System.Memory" Version="4.5.5" />
-    <PackageReference Include="System.Text.Json" Version="8.0.5" />
+    <PackageReference Include="System.Text.Json" Version="8.0.6" />
   </ItemGroup>
 </Project>
\ No newline at end of file
diff --git a/Engine/Source/Programs/Shared/EpicGames.Horde.Tests/EpicGames.Horde.Tests.csproj b/Engine/Source/Programs/Shared/EpicGames.Horde.Tests/EpicGames.Horde.Tests.csproj
--- a/Engine/Source/Programs/Shared/EpicGames.Horde.Tests/EpicGames.Horde.Tests.csproj
+++ b/Engine/Source/Programs/Shared/EpicGames.Horde.Tests/EpicGames.Horde.Tests.csproj
@@ -39,7 +39,7 @@
       <PrivateAssets>all</PrivateAssets>
       <IncludeAssets>runtime; build; native; contentfiles; analyzers; buildtransitive</IncludeAssets>
     </PackageReference>
-    <PackageReference Include="System.Linq.Async" Version="6.0.1" />
+    <PackageReference Include="System.Linq.Async" Version="6.0.3" />
   </ItemGroup>
 
   <ItemGroup>
diff --git a/Engine/Source/Programs/Shared/EpicGames.Horde/EpicGames.Horde.csproj b/Engine/Source/Programs/Shared/EpicGames.Horde/EpicGames.Horde.csproj
--- a/Engine/Source/Programs/Shared/EpicGames.Horde/EpicGames.Horde.csproj
+++ b/Engine/Source/Programs/Shared/EpicGames.Horde/EpicGames.Horde.csproj
@@ -39,8 +39,8 @@
   </ItemGroup>
   <ItemGroup>
     <PackageReference Include="BitFaster.Caching" Version="2.4.1" />
-    <PackageReference Include="Dapper" Version="2.1.24" />
-    <PackageReference Include="Google.Protobuf" Version="3.25.1" />
+    <PackageReference Include="Dapper" Version="2.1.66" />
+    <PackageReference Include="Google.Protobuf" Version="3.25.8" />
     <PackageReference Include="Grpc.Net.Client" Version="2.59.0" />
     <PackageReference Include="Grpc.Tools" Version="2.59.0">
       <PrivateAssets>all</PrivateAssets>
@@ -48,15 +48,15 @@
     </PackageReference>
     <PackageReference Include="K4os.Compression.LZ4" Version="1.2.16" />
     <PackageReference Include="System.IO.Pipelines" Version="8.0.0" />
-    <PackageReference Include="Microsoft.Data.Sqlite" Version="8.0.10" />
+    <PackageReference Include="Microsoft.Data.Sqlite" Version="8.0.20" />
     <PackageReference Include="Microsoft.Extensions.Caching.Abstractions" Version="8.0.0" />
     <PackageReference Include="Microsoft.Extensions.Caching.Memory" Version="8.0.1" />
     <PackageReference Include="Microsoft.VisualStudio.Threading.Analyzers" Version="17.11.20">
       <PrivateAssets>all</PrivateAssets>
       <IncludeAssets>runtime; build; native; contentfiles; analyzers; buildtransitive</IncludeAssets>
     </PackageReference>
-    <PackageReference Include="System.Linq.Async" Version="6.0.1" />
-    <PackageReference Include="ZstdSharp.Port" Version="0.8.1" />
+    <PackageReference Include="System.Linq.Async" Version="6.0.3" />
+    <PackageReference Include="ZstdSharp.Port" Version="0.8.6" />
   </ItemGroup>
   <ItemGroup>
     <ProjectReference Include="..\EpicGames.Core\EpicGames.Core.csproj" />
diff --git a/Engine/Source/Programs/Shared/EpicGames.MongoDB/EpicGames.MongoDB.csproj b/Engine/Source/Programs/Shared/EpicGames.MongoDB/EpicGames.MongoDB.csproj
--- a/Engine/Source/Programs/Shared/EpicGames.MongoDB/EpicGames.MongoDB.csproj
+++ b/Engine/Source/Programs/Shared/EpicGames.MongoDB/EpicGames.MongoDB.csproj
@@ -32,6 +32,6 @@
     </Compile>
   </ItemGroup>
   <ItemGroup>
-    <PackageReference Include="MongoDB.Driver" Version="2.19.0" />
+    <PackageReference Include="MongoDB.Driver" Version="2.19.2" />
   </ItemGroup>
 </Project>
diff --git a/Engine/Source/Programs/Shared/EpicGames.MsBuild/EpicGames.MsBuild.csproj b/Engine/Source/Programs/Shared/EpicGames.MsBuild/EpicGames.MsBuild.csproj
--- a/Engine/Source/Programs/Shared/EpicGames.MsBuild/EpicGames.MsBuild.csproj
+++ b/Engine/Source/Programs/Shared/EpicGames.MsBuild/EpicGames.MsBuild.csproj
@@ -36,9 +36,9 @@
   </ItemGroup>
 
   <ItemGroup>
-    <PackageReference Include="Microsoft.Build" Version="17.11.4" ExcludeAssets="runtime" />
+    <PackageReference Include="Microsoft.Build" Version="17.11.31" ExcludeAssets="runtime" />
     <PackageReference Include="Microsoft.Build.Locator" Version="1.7.8" />
-    <PackageReference Include="System.Drawing.Common" Version="8.0.10" />
+    <PackageReference Include="System.Drawing.Common" Version="8.0.20" />
   </ItemGroup>
 
   <ItemGroup>
diff --git a/Engine/Source/Programs/Shared/EpicGames.OIDC/EpicGames.OIDC.csproj b/Engine/Source/Programs/Shared/EpicGames.OIDC/EpicGames.OIDC.csproj
--- a/Engine/Source/Programs/Shared/EpicGames.OIDC/EpicGames.OIDC.csproj
+++ b/Engine/Source/Programs/Shared/EpicGames.OIDC/EpicGames.OIDC.csproj
@@ -40,7 +40,7 @@
     <PackageReference Include="Microsoft.Extensions.Configuration.Binder" Version="8.0.2" />
     <PackageReference Include="Microsoft.Extensions.Configuration.Json" Version="8.0.1" />
     <PackageReference Include="Microsoft.Extensions.Options" Version="8.0.2" />
-    <PackageReference Include="System.Text.Json" Version="8.0.5" />
+    <PackageReference Include="System.Text.Json" Version="8.0.6" />
   </ItemGroup>
 
 </Project>
diff --git a/Engine/Source/Programs/Shared/EpicGames.Perforce.Fixture/EpicGames.Perforce.Fixture.csproj b/Engine/Source/Programs/Shared/EpicGames.Perforce.Fixture/EpicGames.Perforce.Fixture.csproj
--- a/Engine/Source/Programs/Shared/EpicGames.Perforce.Fixture/EpicGames.Perforce.Fixture.csproj
+++ b/Engine/Source/Programs/Shared/EpicGames.Perforce.Fixture/EpicGames.Perforce.Fixture.csproj
@@ -42,7 +42,7 @@
       <PrivateAssets>all</PrivateAssets>
       <IncludeAssets>runtime; build; native; contentfiles; analyzers; buildtransitive</IncludeAssets>
     </PackageReference>
-    <PackageReference Include="System.Text.Json" Version="8.0.5" />
+    <PackageReference Include="System.Text.Json" Version="8.0.6" />
   </ItemGroup>
 
   <ItemGroup>
diff --git a/Engine/Source/Programs/Shared/EpicGames.Perforce.Managed/EpicGames.Perforce.Managed.csproj b/Engine/Source/Programs/Shared/EpicGames.Perforce.Managed/EpicGames.Perforce.Managed.csproj
--- a/Engine/Source/Programs/Shared/EpicGames.Perforce.Managed/EpicGames.Perforce.Managed.csproj
+++ b/Engine/Source/Programs/Shared/EpicGames.Perforce.Managed/EpicGames.Perforce.Managed.csproj
@@ -29,7 +29,7 @@
   </PropertyGroup>
 
   <ItemGroup>
-    <PackageReference Include="OpenTelemetry" Version="1.8.0" />
+    <PackageReference Include="OpenTelemetry" Version="1.8.1" />
   </ItemGroup>
   
   <ItemGroup>
diff --git a/Engine/Source/Programs/Shared/EpicGames.Perforce/EpicGames.Perforce.csproj b/Engine/Source/Programs/Shared/EpicGames.Perforce/EpicGames.Perforce.csproj
--- a/Engine/Source/Programs/Shared/EpicGames.Perforce/EpicGames.Perforce.csproj
+++ b/Engine/Source/Programs/Shared/EpicGames.Perforce/EpicGames.Perforce.csproj
@@ -74,7 +74,7 @@
   </ItemGroup>
 
   <ItemGroup>
-    <PackageReference Include="System.Linq.Async" Version="6.0.1" />
+    <PackageReference Include="System.Linq.Async" Version="6.0.3" />
   </ItemGroup>
 
 </Project>
diff --git a/Engine/Source/Programs/Shared/EpicGames.Redis/EpicGames.Redis.csproj b/Engine/Source/Programs/Shared/EpicGames.Redis/EpicGames.Redis.csproj
--- a/Engine/Source/Programs/Shared/EpicGames.Redis/EpicGames.Redis.csproj
+++ b/Engine/Source/Programs/Shared/EpicGames.Redis/EpicGames.Redis.csproj
@@ -35,9 +35,9 @@
     </Compile>
   </ItemGroup>
   <ItemGroup>
-    <PackageReference Include="Google.Protobuf" Version="3.25.1" />
-    <PackageReference Include="protobuf-net" Version="3.2.30" />
-    <PackageReference Include="StackExchange.Redis" Version="2.8.16" />
+    <PackageReference Include="Google.Protobuf" Version="3.25.8" />
+    <PackageReference Include="protobuf-net" Version="3.2.56" />
+    <PackageReference Include="StackExchange.Redis" Version="2.8.58" />
   </ItemGroup>
   <ItemGroup>
     <ProjectReference Include="..\EpicGames.Core\EpicGames.Core.csproj" />
diff --git a/Engine/Source/Programs/Shared/EpicGames.Slack/EpicGames.Slack.csproj b/Engine/Source/Programs/Shared/EpicGames.Slack/EpicGames.Slack.csproj
--- a/Engine/Source/Programs/Shared/EpicGames.Slack/EpicGames.Slack.csproj
+++ b/Engine/Source/Programs/Shared/EpicGames.Slack/EpicGames.Slack.csproj
@@ -38,6 +38,6 @@
     <Folder Include="Blocks\" />
   </ItemGroup>
   <ItemGroup>
-    <PackageReference Include="Microsoft.Extensions.Logging.Abstractions" Version="8.0.2" />
+    <PackageReference Include="Microsoft.Extensions.Logging.Abstractions" Version="8.0.3" />
   </ItemGroup>
 </Project>
\ No newline at end of file
diff --git a/Engine/Source/Programs/UnrealBuildTool/UnrealBuildTool.csproj b/Engine/Source/Programs/UnrealBuildTool/UnrealBuildTool.csproj
--- a/Engine/Source/Programs/UnrealBuildTool/UnrealBuildTool.csproj
+++ b/Engine/Source/Programs/UnrealBuildTool/UnrealBuildTool.csproj
@@ -110,7 +110,7 @@
 
     <!-- Added as a temporary workaround to running from visual studio where the system version of dotnet is run. 
          Adding a loader context resolves many of the issues, but it still fails when MSBuild tries to write the json -->
-    <PackageReference Include="Newtonsoft.Json" Version="13.0.3" />
+    <PackageReference Include="Newtonsoft.Json" Version="13.0.4" />
   </ItemGroup>
 
   <!-- Various batch files and shell scripts that build or run UnrealBuildTool -->
diff --git a/Engine/Source/Programs/UnrealToolbox/UnrealToolbox.csproj b/Engine/Source/Programs/UnrealToolbox/UnrealToolbox.csproj
--- a/Engine/Source/Programs/UnrealToolbox/UnrealToolbox.csproj
+++ b/Engine/Source/Programs/UnrealToolbox/UnrealToolbox.csproj
@@ -77,16 +77,16 @@
   </ItemGroup>
 
   <ItemGroup>
-    <PackageReference Include="Avalonia" Version="11.1.3" />
-    <PackageReference Include="Avalonia.Desktop" Version="11.1.3" />
-    <PackageReference Include="Avalonia.Themes.Fluent" Version="11.1.3" />
-    <PackageReference Include="Avalonia.Fonts.Inter" Version="11.1.3" />
+    <PackageReference Include="Avalonia" Version="11.1.5" />
+    <PackageReference Include="Avalonia.Desktop" Version="11.1.5" />
+    <PackageReference Include="Avalonia.Themes.Fluent" Version="11.1.5" />
+    <PackageReference Include="Avalonia.Fonts.Inter" Version="11.1.5" />
     <!--Condition below is needed to remove Avalonia.Diagnostics package from build output in Release configuration.-->
-    <PackageReference Condition="'$(Configuration)' == 'Debug'" Include="Avalonia.Diagnostics" Version="11.1.0" />
+    <PackageReference Condition="'$(Configuration)' == 'Debug'" Include="Avalonia.Diagnostics" Version="11.1.5" />
     <PackageReference Include="FluentAvaloniaUI" Version="2.1.0" />
-    <PackageReference Include="CommunityToolkit.Mvvm" Version="8.3.1" />
+    <PackageReference Include="CommunityToolkit.Mvvm" Version="8.3.2" />
     <PackageReference Include="DesktopNotifications.Avalonia" Version="1.3.1" />
-    <PackageReference Include="System.Drawing.Common" Version="8.0.10" />
+    <PackageReference Include="System.Drawing.Common" Version="8.0.20" />
   </ItemGroup>
 
 </Project>
\ No newline at end of file
