Update test nuget packages to latest patch version

#rnx

diff --git a/Engine/Source/Programs/AutomationTool.Tests/AutomationTool.Tests.csproj b/Engine/Source/Programs/AutomationTool.Tests/AutomationTool.Tests.csproj
--- a/Engine/Source/Programs/AutomationTool.Tests/AutomationTool.Tests.csproj
+++ b/Engine/Source/Programs/AutomationTool.Tests/AutomationTool.Tests.csproj
@@ -12,9 +12,9 @@
 
   <ItemGroup>
     <PackageReference Include="Microsoft.NET.Test.Sdk" Version="17.11.1" />
-    <PackageReference Include="MSTest.TestAdapter" Version="3.6.2" />
-    <PackageReference Include="MSTest.TestFramework" Version="3.6.2" />
-    <PackageReference Include="coverlet.collector" Version="6.0.2">
+    <PackageReference Include="MSTest.TestAdapter" Version="3.6.4" />
+    <PackageReference Include="MSTest.TestFramework" Version="3.6.4" />
+    <PackageReference Include="coverlet.collector" Version="6.0.4">
       <PrivateAssets>all</PrivateAssets>
       <IncludeAssets>runtime; build; native; contentfiles; analyzers; buildtransitive</IncludeAssets>
     </PackageReference>
diff --git a/Engine/Source/Programs/Horde/Drivers/JobDriver.Tests/JobDriver.Tests.csproj b/Engine/Source/Programs/Horde/Drivers/JobDriver.Tests/JobDriver.Tests.csproj
--- a/Engine/Source/Programs/Horde/Drivers/JobDriver.Tests/JobDriver.Tests.csproj
+++ b/Engine/Source/Programs/Horde/Drivers/JobDriver.Tests/JobDriver.Tests.csproj
@@ -21,10 +21,10 @@
   <ItemGroup>
     <PackageReference Include="Microsoft.Extensions.Logging.Abstractions" Version="8.0.3" />
     <PackageReference Include="Microsoft.NET.Test.Sdk" Version="17.11.1" />
-    <PackageReference Include="Moq" Version="4.20.70" />
-    <PackageReference Include="MSTest.TestAdapter" Version="3.6.2" />
-    <PackageReference Include="MSTest.TestFramework" Version="3.6.2" />
-    <PackageReference Include="coverlet.collector" Version="6.0.2">
+    <PackageReference Include="Moq" Version="4.20.72" />
+    <PackageReference Include="MSTest.TestAdapter" Version="3.6.4" />
+    <PackageReference Include="MSTest.TestFramework" Version="3.6.4" />
+    <PackageReference Include="coverlet.collector" Version="6.0.4">
       <PrivateAssets>all</PrivateAssets>
       <IncludeAssets>runtime; build; native; contentfiles; analyzers; buildtransitive</IncludeAssets>
     </PackageReference>
diff --git a/Engine/Source/Programs/Horde/HordeAgent.Tests/HordeAgent.Tests.csproj b/Engine/Source/Programs/Horde/HordeAgent.Tests/HordeAgent.Tests.csproj
--- a/Engine/Source/Programs/Horde/HordeAgent.Tests/HordeAgent.Tests.csproj
+++ b/Engine/Source/Programs/Horde/HordeAgent.Tests/HordeAgent.Tests.csproj
@@ -22,10 +22,10 @@
   <ItemGroup>
     <PackageReference Include="Microsoft.Extensions.Logging.Abstractions" Version="8.0.3" />
     <PackageReference Include="Microsoft.NET.Test.Sdk" Version="17.11.1" />
-    <PackageReference Include="Moq" Version="4.20.70" />
-    <PackageReference Include="MSTest.TestAdapter" Version="3.6.2" />
-    <PackageReference Include="MSTest.TestFramework" Version="3.6.2" />
-    <PackageReference Include="coverlet.collector" Version="6.0.2">
+    <PackageReference Include="Moq" Version="4.20.72" />
+    <PackageReference Include="MSTest.TestAdapter" Version="3.6.4" />
+    <PackageReference Include="MSTest.TestFramework" Version="3.6.4" />
+    <PackageReference Include="coverlet.collector" Version="6.0.4">
       <PrivateAssets>all</PrivateAssets>
       <IncludeAssets>runtime; build; native; contentfiles; analyzers; buildtransitive</IncludeAssets>
     </PackageReference>
diff --git a/Engine/Source/Programs/Horde/HordeServer.Tests.Shared/HordeServer.Tests.Shared.csproj b/Engine/Source/Programs/Horde/HordeServer.Tests.Shared/HordeServer.Tests.Shared.csproj
--- a/Engine/Source/Programs/Horde/HordeServer.Tests.Shared/HordeServer.Tests.Shared.csproj
+++ b/Engine/Source/Programs/Horde/HordeServer.Tests.Shared/HordeServer.Tests.Shared.csproj
@@ -24,11 +24,11 @@
   <ItemGroup>
     <PackageReference Include="Microsoft.AspNetCore.Mvc.Testing" Version="8.0.20" />
     <PackageReference Include="Microsoft.NET.Test.Sdk" Version="17.11.1" />
-    <PackageReference Include="Moq" Version="4.20.70" />
-    <PackageReference Include="MSTest.TestAdapter" Version="3.6.2" />
-    <PackageReference Include="MSTest.TestFramework" Version="3.6.2" />
+    <PackageReference Include="Moq" Version="4.20.72" />
+    <PackageReference Include="MSTest.TestAdapter" Version="3.6.4" />
+    <PackageReference Include="MSTest.TestFramework" Version="3.6.4" />
     <PackageReference Include="Microsoft.Extensions.Caching.Memory" Version="8.0.1" />
-    <PackageReference Include="coverlet.collector" Version="6.0.2">
+    <PackageReference Include="coverlet.collector" Version="6.0.4">
       <PrivateAssets>all</PrivateAssets>
       <IncludeAssets>runtime; build; native; contentfiles; analyzers; buildtransitive</IncludeAssets>
     </PackageReference>
diff --git a/Engine/Source/Programs/Horde/HordeServer.Tests/HordeServer.Tests.csproj b/Engine/Source/Programs/Horde/HordeServer.Tests/HordeServer.Tests.csproj
--- a/Engine/Source/Programs/Horde/HordeServer.Tests/HordeServer.Tests.csproj
+++ b/Engine/Source/Programs/Horde/HordeServer.Tests/HordeServer.Tests.csproj
@@ -21,11 +21,11 @@
   <ItemGroup>
     <PackageReference Include="Microsoft.AspNetCore.Mvc.Testing" Version="8.0.20" />
     <PackageReference Include="Microsoft.NET.Test.Sdk" Version="17.11.1" />
-    <PackageReference Include="Moq" Version="4.20.70" />
-    <PackageReference Include="MSTest.TestAdapter" Version="3.6.2" />
-    <PackageReference Include="MSTest.TestFramework" Version="3.6.2" />
+    <PackageReference Include="Moq" Version="4.20.72" />
+    <PackageReference Include="MSTest.TestAdapter" Version="3.6.4" />
+    <PackageReference Include="MSTest.TestFramework" Version="3.6.4" />
     <PackageReference Include="Microsoft.Extensions.Caching.Memory" Version="8.0.1" />
-    <PackageReference Include="coverlet.collector" Version="6.0.2">
+    <PackageReference Include="coverlet.collector" Version="6.0.4">
       <PrivateAssets>all</PrivateAssets>
       <IncludeAssets>runtime; build; native; contentfiles; analyzers; buildtransitive</IncludeAssets>
     </PackageReference>
diff --git a/Engine/Source/Programs/Horde/Plugins/Analytics/HordeServer.Analytics.Integration.Tests/HordeServer.Analytics.Integration.Tests.csproj b/Engine/Source/Programs/Horde/Plugins/Analytics/HordeServer.Analytics.Integration.Tests/HordeServer.Analytics.Integration.Tests.csproj
--- a/Engine/Source/Programs/Horde/Plugins/Analytics/HordeServer.Analytics.Integration.Tests/HordeServer.Analytics.Integration.Tests.csproj
+++ b/Engine/Source/Programs/Horde/Plugins/Analytics/HordeServer.Analytics.Integration.Tests/HordeServer.Analytics.Integration.Tests.csproj
@@ -19,13 +19,13 @@
   </PropertyGroup>
 
   <ItemGroup>
-    <PackageReference Include="coverlet.collector" Version="6.0.2">
+    <PackageReference Include="coverlet.collector" Version="6.0.4">
       <PrivateAssets>all</PrivateAssets>
       <IncludeAssets>runtime; build; native; contentfiles; analyzers; buildtransitive</IncludeAssets>
     </PackageReference>
     <PackageReference Include="Microsoft.NET.Test.Sdk" Version="17.11.1" />
-    <PackageReference Include="MSTest.TestAdapter" Version="3.6.2" />
-    <PackageReference Include="MSTest.TestFramework" Version="3.6.2" />
+    <PackageReference Include="MSTest.TestAdapter" Version="3.6.4" />
+    <PackageReference Include="MSTest.TestFramework" Version="3.6.4" />
   </ItemGroup>
 
   <ItemGroup>
diff --git a/Engine/Source/Programs/Horde/Plugins/Analytics/HordeServer.Analytics.Tests/HordeServer.Analytics.Tests.csproj b/Engine/Source/Programs/Horde/Plugins/Analytics/HordeServer.Analytics.Tests/HordeServer.Analytics.Tests.csproj
--- a/Engine/Source/Programs/Horde/Plugins/Analytics/HordeServer.Analytics.Tests/HordeServer.Analytics.Tests.csproj
+++ b/Engine/Source/Programs/Horde/Plugins/Analytics/HordeServer.Analytics.Tests/HordeServer.Analytics.Tests.csproj
@@ -19,13 +19,13 @@
   </PropertyGroup>
 
   <ItemGroup>
-    <PackageReference Include="coverlet.collector" Version="6.0.2">
+    <PackageReference Include="coverlet.collector" Version="6.0.4">
       <PrivateAssets>all</PrivateAssets>
       <IncludeAssets>runtime; build; native; contentfiles; analyzers; buildtransitive</IncludeAssets>
     </PackageReference>
     <PackageReference Include="Microsoft.NET.Test.Sdk" Version="17.11.1" />
-    <PackageReference Include="MSTest.TestAdapter" Version="3.6.2" />
-    <PackageReference Include="MSTest.TestFramework" Version="3.6.2" />
+    <PackageReference Include="MSTest.TestAdapter" Version="3.6.4" />
+    <PackageReference Include="MSTest.TestFramework" Version="3.6.4" />
   </ItemGroup>
 
   <ItemGroup>
diff --git a/Engine/Source/Programs/Horde/Plugins/Build/HordeServer.Build.Tests/HordeServer.Build.Tests.csproj b/Engine/Source/Programs/Horde/Plugins/Build/HordeServer.Build.Tests/HordeServer.Build.Tests.csproj
--- a/Engine/Source/Programs/Horde/Plugins/Build/HordeServer.Build.Tests/HordeServer.Build.Tests.csproj
+++ b/Engine/Source/Programs/Horde/Plugins/Build/HordeServer.Build.Tests/HordeServer.Build.Tests.csproj
@@ -20,15 +20,15 @@
   </PropertyGroup>
 
   <ItemGroup>
-    <PackageReference Include="coverlet.collector" Version="6.0.2">
+    <PackageReference Include="coverlet.collector" Version="6.0.4">
       <PrivateAssets>all</PrivateAssets>
       <IncludeAssets>runtime; build; native; contentfiles; analyzers; buildtransitive</IncludeAssets>
     </PackageReference>
     <PackageReference Include="Microsoft.NET.Test.Sdk" Version="17.11.1" />
-    <PackageReference Include="MSTest.TestAdapter" Version="3.6.2" />
-    <PackageReference Include="MSTest.TestFramework" Version="3.6.2" />
+    <PackageReference Include="MSTest.TestAdapter" Version="3.6.4" />
+    <PackageReference Include="MSTest.TestFramework" Version="3.6.4" />
     <PackageReference Include="Microsoft.AspNetCore.Mvc.Testing" Version="8.0.20" />
-    <PackageReference Include="Moq" Version="4.20.70" />
+    <PackageReference Include="Moq" Version="4.20.72" />
     <PackageReference Include="Microsoft.Extensions.Caching.Memory" Version="8.0.1" />
   </ItemGroup>
 
diff --git a/Engine/Source/Programs/Horde/Plugins/Compute/HordeServer.Compute.Tests/HordeServer.Compute.Tests.csproj b/Engine/Source/Programs/Horde/Plugins/Compute/HordeServer.Compute.Tests/HordeServer.Compute.Tests.csproj
--- a/Engine/Source/Programs/Horde/Plugins/Compute/HordeServer.Compute.Tests/HordeServer.Compute.Tests.csproj
+++ b/Engine/Source/Programs/Horde/Plugins/Compute/HordeServer.Compute.Tests/HordeServer.Compute.Tests.csproj
@@ -20,13 +20,13 @@
   </PropertyGroup>
 
   <ItemGroup>
-    <PackageReference Include="coverlet.collector" Version="6.0.2">
+    <PackageReference Include="coverlet.collector" Version="6.0.4">
       <PrivateAssets>all</PrivateAssets>
       <IncludeAssets>runtime; build; native; contentfiles; analyzers; buildtransitive</IncludeAssets>
     </PackageReference>
     <PackageReference Include="Microsoft.NET.Test.Sdk" Version="17.11.1" />
-    <PackageReference Include="MSTest.TestAdapter" Version="3.6.2" />
-    <PackageReference Include="MSTest.TestFramework" Version="3.6.2" />
+    <PackageReference Include="MSTest.TestAdapter" Version="3.6.4" />
+    <PackageReference Include="MSTest.TestFramework" Version="3.6.4" />
   </ItemGroup>
 
   <ItemGroup>
diff --git a/Engine/Source/Programs/Horde/Plugins/Ddc/HordeServer.Ddc.Tests/HordeServer.Ddc.Tests.csproj b/Engine/Source/Programs/Horde/Plugins/Ddc/HordeServer.Ddc.Tests/HordeServer.Ddc.Tests.csproj
--- a/Engine/Source/Programs/Horde/Plugins/Ddc/HordeServer.Ddc.Tests/HordeServer.Ddc.Tests.csproj
+++ b/Engine/Source/Programs/Horde/Plugins/Ddc/HordeServer.Ddc.Tests/HordeServer.Ddc.Tests.csproj
@@ -20,13 +20,13 @@
   </PropertyGroup>
 
   <ItemGroup>
-    <PackageReference Include="coverlet.collector" Version="6.0.2">
+    <PackageReference Include="coverlet.collector" Version="6.0.4">
       <PrivateAssets>all</PrivateAssets>
       <IncludeAssets>runtime; build; native; contentfiles; analyzers; buildtransitive</IncludeAssets>
     </PackageReference>
     <PackageReference Include="Microsoft.NET.Test.Sdk" Version="17.11.1" />
-    <PackageReference Include="MSTest.TestAdapter" Version="3.6.2" />
-    <PackageReference Include="MSTest.TestFramework" Version="3.6.2" />
+    <PackageReference Include="MSTest.TestAdapter" Version="3.6.4" />
+    <PackageReference Include="MSTest.TestFramework" Version="3.6.4" />
   </ItemGroup>
 
   <ItemGroup>
diff --git a/Engine/Source/Programs/Horde/Plugins/Experimental/HordeServer.Experimental.Tests/HordeServer.Experimental.Tests.csproj b/Engine/Source/Programs/Horde/Plugins/Experimental/HordeServer.Experimental.Tests/HordeServer.Experimental.Tests.csproj
--- a/Engine/Source/Programs/Horde/Plugins/Experimental/HordeServer.Experimental.Tests/HordeServer.Experimental.Tests.csproj
+++ b/Engine/Source/Programs/Horde/Plugins/Experimental/HordeServer.Experimental.Tests/HordeServer.Experimental.Tests.csproj
@@ -20,15 +20,15 @@
   </PropertyGroup>
 
   <ItemGroup>
-    <PackageReference Include="coverlet.collector" Version="6.0.2">
+    <PackageReference Include="coverlet.collector" Version="6.0.4">
       <PrivateAssets>all</PrivateAssets>
       <IncludeAssets>runtime; build; native; contentfiles; analyzers; buildtransitive</IncludeAssets>
     </PackageReference>
     <PackageReference Include="Microsoft.NET.Test.Sdk" Version="17.11.1" />
-    <PackageReference Include="MSTest.TestAdapter" Version="3.6.2" />
-    <PackageReference Include="MSTest.TestFramework" Version="3.6.2" />
+    <PackageReference Include="MSTest.TestAdapter" Version="3.6.4" />
+    <PackageReference Include="MSTest.TestFramework" Version="3.6.4" />
     <PackageReference Include="Microsoft.AspNetCore.Mvc.Testing" Version="8.0.20" />
-    <PackageReference Include="Moq" Version="4.20.70" />
+    <PackageReference Include="Moq" Version="4.20.72" />
     <PackageReference Include="Microsoft.Extensions.Caching.Memory" Version="8.0.1" />
   </ItemGroup>
 
diff --git a/Engine/Source/Programs/Horde/Plugins/Secrets/HordeServer.Secrets.Tests/HordeServer.Secrets.Tests.csproj b/Engine/Source/Programs/Horde/Plugins/Secrets/HordeServer.Secrets.Tests/HordeServer.Secrets.Tests.csproj
--- a/Engine/Source/Programs/Horde/Plugins/Secrets/HordeServer.Secrets.Tests/HordeServer.Secrets.Tests.csproj
+++ b/Engine/Source/Programs/Horde/Plugins/Secrets/HordeServer.Secrets.Tests/HordeServer.Secrets.Tests.csproj
@@ -20,15 +20,15 @@
   </PropertyGroup>
 
   <ItemGroup>
-    <PackageReference Include="AWSSDK.SSO" Version="3.7.400.108" />
-    <PackageReference Include="AWSSDK.SSOOIDC" Version="3.7.400.109" />
-    <PackageReference Include="coverlet.collector" Version="6.0.2">
+    <PackageReference Include="AWSSDK.SSO" Version="3.7.501.12" />
+    <PackageReference Include="AWSSDK.SSOOIDC" Version="3.7.502" />
+    <PackageReference Include="coverlet.collector" Version="6.0.4">
       <PrivateAssets>all</PrivateAssets>
       <IncludeAssets>runtime; build; native; contentfiles; analyzers; buildtransitive</IncludeAssets>
     </PackageReference>
     <PackageReference Include="Microsoft.NET.Test.Sdk" Version="17.11.1" />
-    <PackageReference Include="MSTest.TestAdapter" Version="3.6.2" />
-    <PackageReference Include="MSTest.TestFramework" Version="3.6.2" />
+    <PackageReference Include="MSTest.TestAdapter" Version="3.6.4" />
+    <PackageReference Include="MSTest.TestFramework" Version="3.6.4" />
   </ItemGroup>
 
   <ItemGroup>
diff --git a/Engine/Source/Programs/Horde/Plugins/Storage/HordeServer.Storage.Tests/HordeServer.Storage.Tests.csproj b/Engine/Source/Programs/Horde/Plugins/Storage/HordeServer.Storage.Tests/HordeServer.Storage.Tests.csproj
--- a/Engine/Source/Programs/Horde/Plugins/Storage/HordeServer.Storage.Tests/HordeServer.Storage.Tests.csproj
+++ b/Engine/Source/Programs/Horde/Plugins/Storage/HordeServer.Storage.Tests/HordeServer.Storage.Tests.csproj
@@ -20,13 +20,13 @@
   </PropertyGroup>
 
   <ItemGroup>
-    <PackageReference Include="coverlet.collector" Version="6.0.2">
+    <PackageReference Include="coverlet.collector" Version="6.0.4">
       <PrivateAssets>all</PrivateAssets>
       <IncludeAssets>runtime; build; native; contentfiles; analyzers; buildtransitive</IncludeAssets>
     </PackageReference>
     <PackageReference Include="Microsoft.NET.Test.Sdk" Version="17.11.1" />
-    <PackageReference Include="MSTest.TestAdapter" Version="3.6.2" />
-    <PackageReference Include="MSTest.TestFramework" Version="3.6.2" />
+    <PackageReference Include="MSTest.TestAdapter" Version="3.6.4" />
+    <PackageReference Include="MSTest.TestFramework" Version="3.6.4" />
   </ItemGroup>
 
   <ItemGroup>
diff --git a/Engine/Source/Programs/Horde/Plugins/Tools/HordeServer.Tools.Tests/HordeServer.Tools.Tests.csproj b/Engine/Source/Programs/Horde/Plugins/Tools/HordeServer.Tools.Tests/HordeServer.Tools.Tests.csproj
--- a/Engine/Source/Programs/Horde/Plugins/Tools/HordeServer.Tools.Tests/HordeServer.Tools.Tests.csproj
+++ b/Engine/Source/Programs/Horde/Plugins/Tools/HordeServer.Tools.Tests/HordeServer.Tools.Tests.csproj
@@ -20,13 +20,13 @@
   </PropertyGroup>
 
   <ItemGroup>
-    <PackageReference Include="coverlet.collector" Version="6.0.2">
+    <PackageReference Include="coverlet.collector" Version="6.0.4">
       <PrivateAssets>all</PrivateAssets>
       <IncludeAssets>runtime; build; native; contentfiles; analyzers; buildtransitive</IncludeAssets>
     </PackageReference>
     <PackageReference Include="Microsoft.NET.Test.Sdk" Version="17.11.1" />
-    <PackageReference Include="MSTest.TestAdapter" Version="3.6.2" />
-    <PackageReference Include="MSTest.TestFramework" Version="3.6.2" />
+    <PackageReference Include="MSTest.TestAdapter" Version="3.6.4" />
+    <PackageReference Include="MSTest.TestFramework" Version="3.6.4" />
   </ItemGroup>
 
   <ItemGroup>
diff --git a/Engine/Source/Programs/Shared/EpicGames.BuildGraph.Tests/EpicGames.BuildGraph.Tests.csproj b/Engine/Source/Programs/Shared/EpicGames.BuildGraph.Tests/EpicGames.BuildGraph.Tests.csproj
--- a/Engine/Source/Programs/Shared/EpicGames.BuildGraph.Tests/EpicGames.BuildGraph.Tests.csproj
+++ b/Engine/Source/Programs/Shared/EpicGames.BuildGraph.Tests/EpicGames.BuildGraph.Tests.csproj
@@ -45,11 +45,11 @@
 
   <ItemGroup>
     <PackageReference Include="Microsoft.NET.Test.Sdk" Version="17.11.1" />
-    <PackageReference Include="MSTest.TestAdapter" Version="3.6.2">
+    <PackageReference Include="MSTest.TestAdapter" Version="3.6.4">
       <TreatAsUsed>true</TreatAsUsed>
     </PackageReference>
-    <PackageReference Include="MSTest.TestFramework" Version="3.6.2" />
-    <PackageReference Include="coverlet.collector" Version="6.0.2">
+    <PackageReference Include="MSTest.TestFramework" Version="3.6.4" />
+    <PackageReference Include="coverlet.collector" Version="6.0.4">
       <PrivateAssets>all</PrivateAssets>
       <IncludeAssets>runtime; build; native; contentfiles; analyzers; buildtransitive</IncludeAssets>
     </PackageReference>
diff --git a/Engine/Source/Programs/Shared/EpicGames.Core.Tests/EpicGames.Core.Tests.csproj b/Engine/Source/Programs/Shared/EpicGames.Core.Tests/EpicGames.Core.Tests.csproj
--- a/Engine/Source/Programs/Shared/EpicGames.Core.Tests/EpicGames.Core.Tests.csproj
+++ b/Engine/Source/Programs/Shared/EpicGames.Core.Tests/EpicGames.Core.Tests.csproj
@@ -49,11 +49,11 @@
 
   <ItemGroup>
     <PackageReference Include="Microsoft.NET.Test.Sdk" Version="17.11.1" />
-    <PackageReference Include="MSTest.TestAdapter" Version="3.6.2">
+    <PackageReference Include="MSTest.TestAdapter" Version="3.6.4">
       <TreatAsUsed>true</TreatAsUsed>
     </PackageReference>
-    <PackageReference Include="MSTest.TestFramework" Version="3.6.2" />
-    <PackageReference Include="coverlet.collector" Version="6.0.2">
+    <PackageReference Include="MSTest.TestFramework" Version="3.6.4" />
+    <PackageReference Include="coverlet.collector" Version="6.0.4">
       <PrivateAssets>all</PrivateAssets>
       <IncludeAssets>runtime; build; native; contentfiles; analyzers; buildtransitive</IncludeAssets>
     </PackageReference>
diff --git a/Engine/Source/Programs/Shared/EpicGames.Horde.Tests/EpicGames.Horde.Tests.csproj b/Engine/Source/Programs/Shared/EpicGames.Horde.Tests/EpicGames.Horde.Tests.csproj
--- a/Engine/Source/Programs/Shared/EpicGames.Horde.Tests/EpicGames.Horde.Tests.csproj
+++ b/Engine/Source/Programs/Shared/EpicGames.Horde.Tests/EpicGames.Horde.Tests.csproj
@@ -31,11 +31,11 @@
 
   <ItemGroup>
     <PackageReference Include="Microsoft.NET.Test.Sdk" Version="17.11.1" />
-    <PackageReference Include="MSTest.TestAdapter" Version="3.6.2">
+    <PackageReference Include="MSTest.TestAdapter" Version="3.6.4">
       <TreatAsUsed>true</TreatAsUsed>
     </PackageReference>
-    <PackageReference Include="MSTest.TestFramework" Version="3.6.2" />
-    <PackageReference Include="coverlet.collector" Version="6.0.2">
+    <PackageReference Include="MSTest.TestFramework" Version="3.6.4" />
+    <PackageReference Include="coverlet.collector" Version="6.0.4">
       <PrivateAssets>all</PrivateAssets>
       <IncludeAssets>runtime; build; native; contentfiles; analyzers; buildtransitive</IncludeAssets>
     </PackageReference>
diff --git a/Engine/Source/Programs/Shared/EpicGames.IoHash.Tests/EpicGames.IoHash.Tests.csproj b/Engine/Source/Programs/Shared/EpicGames.IoHash.Tests/EpicGames.IoHash.Tests.csproj
--- a/Engine/Source/Programs/Shared/EpicGames.IoHash.Tests/EpicGames.IoHash.Tests.csproj
+++ b/Engine/Source/Programs/Shared/EpicGames.IoHash.Tests/EpicGames.IoHash.Tests.csproj
@@ -18,11 +18,11 @@
 
   <ItemGroup>
     <PackageReference Include="Microsoft.NET.Test.Sdk" Version="17.11.1" />
-    <PackageReference Include="MSTest.TestAdapter" Version="3.6.2">
+    <PackageReference Include="MSTest.TestAdapter" Version="3.6.4">
       <TreatAsUsed>true</TreatAsUsed>
     </PackageReference>
-    <PackageReference Include="MSTest.TestFramework" Version="3.6.2" />
-    <PackageReference Include="coverlet.collector" Version="6.0.2">
+    <PackageReference Include="MSTest.TestFramework" Version="3.6.4" />
+    <PackageReference Include="coverlet.collector" Version="6.0.4">
       <PrivateAssets>all</PrivateAssets>
       <IncludeAssets>runtime; build; native; contentfiles; analyzers; buildtransitive</IncludeAssets>
     </PackageReference>
diff --git a/Engine/Source/Programs/Shared/EpicGames.Perforce.Fixture/EpicGames.Perforce.Fixture.csproj b/Engine/Source/Programs/Shared/EpicGames.Perforce.Fixture/EpicGames.Perforce.Fixture.csproj
--- a/Engine/Source/Programs/Shared/EpicGames.Perforce.Fixture/EpicGames.Perforce.Fixture.csproj
+++ b/Engine/Source/Programs/Shared/EpicGames.Perforce.Fixture/EpicGames.Perforce.Fixture.csproj
@@ -34,11 +34,11 @@
   <ItemGroup>
     <PackageReference Include="Microsoft.Extensions.Logging.Console" Version="8.0.1" />
     <PackageReference Include="Microsoft.NET.Test.Sdk" Version="17.11.1" />
-    <PackageReference Include="MSTest.TestAdapter" Version="3.6.2">
+    <PackageReference Include="MSTest.TestAdapter" Version="3.6.4">
       <TreatAsUsed>true</TreatAsUsed>
     </PackageReference>
-    <PackageReference Include="MSTest.TestFramework" Version="3.6.2" />
-    <PackageReference Include="coverlet.collector" Version="6.0.2">
+    <PackageReference Include="MSTest.TestFramework" Version="3.6.4" />
+    <PackageReference Include="coverlet.collector" Version="6.0.4">
       <PrivateAssets>all</PrivateAssets>
       <IncludeAssets>runtime; build; native; contentfiles; analyzers; buildtransitive</IncludeAssets>
     </PackageReference>
diff --git a/Engine/Source/Programs/Shared/EpicGames.Perforce.Managed.Tests/EpicGames.Perforce.Managed.Tests.csproj b/Engine/Source/Programs/Shared/EpicGames.Perforce.Managed.Tests/EpicGames.Perforce.Managed.Tests.csproj
--- a/Engine/Source/Programs/Shared/EpicGames.Perforce.Managed.Tests/EpicGames.Perforce.Managed.Tests.csproj
+++ b/Engine/Source/Programs/Shared/EpicGames.Perforce.Managed.Tests/EpicGames.Perforce.Managed.Tests.csproj
@@ -34,13 +34,13 @@
     <PackageReference Include="Microsoft.NET.Test.Sdk" Version="17.11.1">
       <TreatAsUsed>true</TreatAsUsed>
     </PackageReference>
-    <PackageReference Include="MSTest.TestAdapter" Version="3.6.2">
+    <PackageReference Include="MSTest.TestAdapter" Version="3.6.4">
       <TreatAsUsed>true</TreatAsUsed>
     </PackageReference>
-    <PackageReference Include="MSTest.TestFramework" Version="3.6.2">
+    <PackageReference Include="MSTest.TestFramework" Version="3.6.4">
       <TreatAsUsed>true</TreatAsUsed>
     </PackageReference>
-    <PackageReference Include="coverlet.collector" Version="6.0.2">
+    <PackageReference Include="coverlet.collector" Version="6.0.4">
       <PrivateAssets>all</PrivateAssets>
       <IncludeAssets>runtime; build; native; contentfiles; analyzers; buildtransitive</IncludeAssets>
     </PackageReference>
diff --git a/Engine/Source/Programs/Shared/EpicGames.Perforce.Tests/EpicGames.Perforce.Tests.csproj b/Engine/Source/Programs/Shared/EpicGames.Perforce.Tests/EpicGames.Perforce.Tests.csproj
--- a/Engine/Source/Programs/Shared/EpicGames.Perforce.Tests/EpicGames.Perforce.Tests.csproj
+++ b/Engine/Source/Programs/Shared/EpicGames.Perforce.Tests/EpicGames.Perforce.Tests.csproj
@@ -31,11 +31,11 @@
 
   <ItemGroup>
     <PackageReference Include="Microsoft.NET.Test.Sdk" Version="17.11.1" />
-    <PackageReference Include="MSTest.TestAdapter" Version="3.6.2">
+    <PackageReference Include="MSTest.TestAdapter" Version="3.6.4">
       <TreatAsUsed>true</TreatAsUsed>
     </PackageReference>
-    <PackageReference Include="MSTest.TestFramework" Version="3.6.2" />
-    <PackageReference Include="coverlet.collector" Version="6.0.2">
+    <PackageReference Include="MSTest.TestFramework" Version="3.6.4" />
+    <PackageReference Include="coverlet.collector" Version="6.0.4">
       <PrivateAssets>all</PrivateAssets>
       <IncludeAssets>runtime; build; native; contentfiles; analyzers; buildtransitive</IncludeAssets>
     </PackageReference>
diff --git a/Engine/Source/Programs/Shared/EpicGames.Redis.Tests/EpicGames.Redis.Tests.csproj b/Engine/Source/Programs/Shared/EpicGames.Redis.Tests/EpicGames.Redis.Tests.csproj
--- a/Engine/Source/Programs/Shared/EpicGames.Redis.Tests/EpicGames.Redis.Tests.csproj
+++ b/Engine/Source/Programs/Shared/EpicGames.Redis.Tests/EpicGames.Redis.Tests.csproj
@@ -31,11 +31,11 @@
 
   <ItemGroup>
     <PackageReference Include="Microsoft.NET.Test.Sdk" Version="17.11.1" />
-    <PackageReference Include="MSTest.TestAdapter" Version="3.6.2">
+    <PackageReference Include="MSTest.TestAdapter" Version="3.6.4">
       <TreatAsUsed>true</TreatAsUsed>
     </PackageReference>
-    <PackageReference Include="MSTest.TestFramework" Version="3.6.2" />
-    <PackageReference Include="coverlet.collector" Version="6.0.2">
+    <PackageReference Include="MSTest.TestFramework" Version="3.6.4" />
+    <PackageReference Include="coverlet.collector" Version="6.0.4">
       <PrivateAssets>all</PrivateAssets>
       <IncludeAssets>runtime; build; native; contentfiles; analyzers; buildtransitive</IncludeAssets>
     </PackageReference>
diff --git a/Engine/Source/Programs/Shared/EpicGames.Serialization.Tests/EpicGames.Serialization.Tests.csproj b/Engine/Source/Programs/Shared/EpicGames.Serialization.Tests/EpicGames.Serialization.Tests.csproj
--- a/Engine/Source/Programs/Shared/EpicGames.Serialization.Tests/EpicGames.Serialization.Tests.csproj
+++ b/Engine/Source/Programs/Shared/EpicGames.Serialization.Tests/EpicGames.Serialization.Tests.csproj
@@ -34,11 +34,11 @@
 
   <ItemGroup>
     <PackageReference Include="Microsoft.NET.Test.Sdk" Version="17.11.1" />
-    <PackageReference Include="MSTest.TestAdapter" Version="3.6.2">
+    <PackageReference Include="MSTest.TestAdapter" Version="3.6.4">
       <TreatAsUsed>true</TreatAsUsed>
     </PackageReference>
-    <PackageReference Include="MSTest.TestFramework" Version="3.6.2" />
-    <PackageReference Include="coverlet.collector" Version="6.0.2">
+    <PackageReference Include="MSTest.TestFramework" Version="3.6.4" />
+    <PackageReference Include="coverlet.collector" Version="6.0.4">
       <PrivateAssets>all</PrivateAssets>
       <IncludeAssets>runtime; build; native; contentfiles; analyzers; buildtransitive</IncludeAssets>
     </PackageReference>
diff --git a/Engine/Source/Programs/Shared/EpicGames.Tracing.Tests/EpicGames.Tracing.Tests.csproj b/Engine/Source/Programs/Shared/EpicGames.Tracing.Tests/EpicGames.Tracing.Tests.csproj
--- a/Engine/Source/Programs/Shared/EpicGames.Tracing.Tests/EpicGames.Tracing.Tests.csproj
+++ b/Engine/Source/Programs/Shared/EpicGames.Tracing.Tests/EpicGames.Tracing.Tests.csproj
@@ -30,11 +30,11 @@
 
   <ItemGroup>
     <PackageReference Include="Microsoft.NET.Test.Sdk" Version="17.11.1" />
-    <PackageReference Include="MSTest.TestAdapter" Version="3.6.2">
+    <PackageReference Include="MSTest.TestAdapter" Version="3.6.4">
       <TreatAsUsed>true</TreatAsUsed>
     </PackageReference>
-    <PackageReference Include="MSTest.TestFramework" Version="3.6.2" />
-    <PackageReference Include="coverlet.collector" Version="6.0.2">
+    <PackageReference Include="MSTest.TestFramework" Version="3.6.4" />
+    <PackageReference Include="coverlet.collector" Version="6.0.4">
       <PrivateAssets>all</PrivateAssets>
       <IncludeAssets>runtime; build; native; contentfiles; analyzers; buildtransitive</IncludeAssets>
     </PackageReference>
diff --git a/Engine/Source/Programs/Shared/UnrealBuildTool.Tests/UnrealBuildTool.Tests.csproj b/Engine/Source/Programs/Shared/UnrealBuildTool.Tests/UnrealBuildTool.Tests.csproj
--- a/Engine/Source/Programs/Shared/UnrealBuildTool.Tests/UnrealBuildTool.Tests.csproj
+++ b/Engine/Source/Programs/Shared/UnrealBuildTool.Tests/UnrealBuildTool.Tests.csproj
@@ -17,11 +17,11 @@
 
   <ItemGroup>
     <PackageReference Include="Microsoft.NET.Test.Sdk" Version="17.11.1" />
-    <PackageReference Include="MSTest.TestAdapter" Version="3.6.2">
+    <PackageReference Include="MSTest.TestAdapter" Version="3.6.4">
       <TreatAsUsed>true</TreatAsUsed>
     </PackageReference>
-    <PackageReference Include="MSTest.TestFramework" Version="3.6.2" />
-    <PackageReference Include="coverlet.collector" Version="6.0.2">
+    <PackageReference Include="MSTest.TestFramework" Version="3.6.4" />
+    <PackageReference Include="coverlet.collector" Version="6.0.4">
       <PrivateAssets>all</PrivateAssets>
       <IncludeAssets>runtime; build; native; contentfiles; analyzers; buildtransitive</IncludeAssets>
     </PackageReference>
