Rollback AWSSDK.SimpleSystemsManagement to 3.7.304.24

#rnx

diff --git a/Engine/Source/Programs/Horde/Drivers/JobDriver/JobDriver.csproj b/Engine/Source/Programs/Horde/Drivers/JobDriver/JobDriver.csproj
--- a/Engine/Source/Programs/Horde/Drivers/JobDriver/JobDriver.csproj
+++ b/Engine/Source/Programs/Horde/Drivers/JobDriver/JobDriver.csproj
@@ -54,7 +54,7 @@
   <ItemGroup>
     <PackageReference Include="AWSSDK.Core" Version="3.7.500.23" />
     <PackageReference Include="AWSSDK.EC2" Version="3.7.514" />
-    <PackageReference Include="AWSSDK.SimpleSystemsManagement" Version="3.7.502" />
+    <PackageReference Include="AWSSDK.SimpleSystemsManagement" Version="3.7.304.24" />
     <PackageReference Include="Google.Protobuf" Version="3.25.8" />
     <PackageReference Include="Grpc.Net.Client" Version="2.59.0" />
     <PackageReference Include="Grpc.Tools" Version="2.59.0">
diff --git a/Engine/Source/Programs/Horde/HordeAgent/HordeAgent.csproj b/Engine/Source/Programs/Horde/HordeAgent/HordeAgent.csproj
--- a/Engine/Source/Programs/Horde/HordeAgent/HordeAgent.csproj
+++ b/Engine/Source/Programs/Horde/HordeAgent/HordeAgent.csproj
@@ -36,7 +36,7 @@
   <ItemGroup>
     <PackageReference Include="AWSSDK.Core" Version="3.7.500.23" />
     <PackageReference Include="AWSSDK.EC2" Version="3.7.514" />
-    <PackageReference Include="AWSSDK.SimpleSystemsManagement" Version="3.7.502" />
+    <PackageReference Include="AWSSDK.SimpleSystemsManagement" Version="3.7.304.24" />
     <PackageReference Include="Google.Protobuf" Version="3.25.8" />
     <PackageReference Include="Grpc.Net.Client" Version="2.59.0" />
     <PackageReference Include="Grpc.Tools" Version="2.59.0">
diff --git a/Engine/Source/Programs/Horde/HordeServer/HordeServer.csproj b/Engine/Source/Programs/Horde/HordeServer/HordeServer.csproj
--- a/Engine/Source/Programs/Horde/HordeServer/HordeServer.csproj
+++ b/Engine/Source/Programs/Horde/HordeServer/HordeServer.csproj
@@ -25,7 +25,7 @@
     <PackageReference Include="AWSSDK.EC2" Version="3.7.514" />
     <PackageReference Include="AWSSDK.Extensions.NETCore.Setup" Version="3.7.400" />
     <PackageReference Include="AWSSDK.S3" Version="3.7.501.8" />
-    <PackageReference Include="AWSSDK.SimpleSystemsManagement" Version="3.7.502" />
+    <PackageReference Include="AWSSDK.SimpleSystemsManagement" Version="3.7.304.24" />
     <PackageReference Include="AWSSDK.SQS" Version="3.7.501.14" />
     <PackageReference Include="AWSSDK.SecurityToken" Version="3.7.501.14" />
     <PackageReference Include="Azure.Storage.Blobs" Version="12.19.1" />
diff --git a/Engine/Source/Programs/Horde/Plugins/Compute/HordeServer.Compute/HordeServer.Compute.csproj b/Engine/Source/Programs/Horde/Plugins/Compute/HordeServer.Compute/HordeServer.Compute.csproj
--- a/Engine/Source/Programs/Horde/Plugins/Compute/HordeServer.Compute/HordeServer.Compute.csproj
+++ b/Engine/Source/Programs/Horde/Plugins/Compute/HordeServer.Compute/HordeServer.Compute.csproj
@@ -49,7 +49,7 @@
     <PackageReference Include="AWSSDK.EC2" Version="3.7.514" />
     <PackageReference Include="AWSSDK.Extensions.NETCore.Setup" Version="3.7.400" />
     <PackageReference Include="AWSSDK.S3" Version="3.7.501.8" />
-    <PackageReference Include="AWSSDK.SimpleSystemsManagement" Version="3.7.502" />
+    <PackageReference Include="AWSSDK.SimpleSystemsManagement" Version="3.7.304.24" />
     <PackageReference Include="AWSSDK.SQS" Version="3.7.501.14" />
     <PackageReference Include="AWSSDK.SecurityToken" Version="3.7.501.14" />
     <PackageReference Include="Azure.Storage.Blobs" Version="12.19.1" />
diff --git a/Engine/Source/Programs/Horde/Plugins/Secrets/HordeServer.Secrets/HordeServer.Secrets.csproj b/Engine/Source/Programs/Horde/Plugins/Secrets/HordeServer.Secrets/HordeServer.Secrets.csproj
--- a/Engine/Source/Programs/Horde/Plugins/Secrets/HordeServer.Secrets/HordeServer.Secrets.csproj
+++ b/Engine/Source/Programs/Horde/Plugins/Secrets/HordeServer.Secrets/HordeServer.Secrets.csproj
@@ -22,7 +22,7 @@
   </ItemGroup>
 
   <ItemGroup>
-    <PackageReference Include="AWSSDK.SimpleSystemsManagement" Version="3.7.502" />
+    <PackageReference Include="AWSSDK.SimpleSystemsManagement" Version="3.7.304.24" />
     <PackageReference Include="Azure.Identity" Version="1.14.2" />
     <PackageReference Include="Azure.Security.KeyVault.Secrets" Version="4.8.0" />
   </ItemGroup>
