Rollback other AWSSDK packages

#rnx

diff --git a/Engine/Source/Programs/Horde/Drivers/JobDriver/JobDriver.csproj b/Engine/Source/Programs/Horde/Drivers/JobDriver/JobDriver.csproj
--- a/Engine/Source/Programs/Horde/Drivers/JobDriver/JobDriver.csproj
+++ b/Engine/Source/Programs/Horde/Drivers/JobDriver/JobDriver.csproj
@@ -52,9 +52,9 @@
   </ItemGroup>
 
   <ItemGroup>
-    <PackageReference Include="AWSSDK.Core" Version="3.7.500.23" />
-    <PackageReference Include="AWSSDK.EC2" Version="3.7.514" />
-    <PackageReference Include="AWSSDK.SimpleSystemsManagement" Version="3.7.304.24" />
+    <PackageReference Include="AWSSDK.Core" Version="3.7.300.12" />
+    <PackageReference Include="AWSSDK.EC2" Version="3.7.306.6" />
+    <PackageReference Include="AWSSDK.SimpleSystemsManagement" Version="3.7.301.8" />
     <PackageReference Include="Google.Protobuf" Version="3.25.8" />
     <PackageReference Include="Grpc.Net.Client" Version="2.59.0" />
     <PackageReference Include="Grpc.Tools" Version="2.59.0">
@@ -120,7 +120,7 @@
       <FilesToDelete Include="$(PublishDir)**/oo2core*winuwp*.dll" />
     </ItemGroup>
     <Delete Files="@(FilesToDelete)">
-      <Output TaskParameter="DeletedFiles" ItemName="DeletedList" />
+      <Output TaskParameter="DeletedFiles" ItemName="DeletedList"/>
     </Delete>
     <Message Text="Deleted files:%0A@(DeletedList->'  %(Identity)','%0A')" Importance="high" />
   </Target>
diff --git a/Engine/Source/Programs/Horde/HordeAgent/HordeAgent.csproj b/Engine/Source/Programs/Horde/HordeAgent/HordeAgent.csproj
--- a/Engine/Source/Programs/Horde/HordeAgent/HordeAgent.csproj
+++ b/Engine/Source/Programs/Horde/HordeAgent/HordeAgent.csproj
@@ -34,9 +34,9 @@
   </ItemGroup>
 
   <ItemGroup>
-    <PackageReference Include="AWSSDK.Core" Version="3.7.500.23" />
-    <PackageReference Include="AWSSDK.EC2" Version="3.7.514" />
-    <PackageReference Include="AWSSDK.SimpleSystemsManagement" Version="3.7.304.24" />
+    <PackageReference Include="AWSSDK.Core" Version="3.7.300.12" />
+    <PackageReference Include="AWSSDK.EC2" Version="3.7.306.6" />
+    <PackageReference Include="AWSSDK.SimpleSystemsManagement" Version="3.7.301.8" />
     <PackageReference Include="Google.Protobuf" Version="3.25.8" />
     <PackageReference Include="Grpc.Net.Client" Version="2.59.0" />
     <PackageReference Include="Grpc.Tools" Version="2.59.0">
@@ -207,7 +207,7 @@
       <FilesToDelete Include="$(PublishDir)**/oo2core*winuwp*.dll" />
     </ItemGroup>
     <Delete Files="@(FilesToDelete)">
-      <Output TaskParameter="DeletedFiles" ItemName="DeletedList" />
+      <Output TaskParameter="DeletedFiles" ItemName="DeletedList"/>
     </Delete>
     <Message Text="Deleted files:%0A@(DeletedList->'  %(Identity)','%0A')" Importance="high" />
   </Target>
diff --git a/Engine/Source/Programs/Horde/HordeServer/HordeServer.csproj b/Engine/Source/Programs/Horde/HordeServer/HordeServer.csproj
--- a/Engine/Source/Programs/Horde/HordeServer/HordeServer.csproj
+++ b/Engine/Source/Programs/Horde/HordeServer/HordeServer.csproj
@@ -20,14 +20,14 @@
   </PropertyGroup>
 
   <ItemGroup>
-    <PackageReference Include="AWSSDK.AutoScaling" Version="3.7.501.7" />
-    <PackageReference Include="AWSSDK.CloudWatch" Version="3.7.502.7" />
-    <PackageReference Include="AWSSDK.EC2" Version="3.7.514" />
-    <PackageReference Include="AWSSDK.Extensions.NETCore.Setup" Version="3.7.400" />
-    <PackageReference Include="AWSSDK.S3" Version="3.7.501.8" />
+    <PackageReference Include="AWSSDK.AutoScaling" Version="3.7.301.9" />
+    <PackageReference Include="AWSSDK.CloudWatch" Version="3.7.300.12" />
+    <PackageReference Include="AWSSDK.EC2" Version="3.7.306.6" />
+    <PackageReference Include="AWSSDK.Extensions.NETCore.Setup" Version="3.7.300" />
+    <PackageReference Include="AWSSDK.S3" Version="3.7.304.1" />
     <PackageReference Include="AWSSDK.SimpleSystemsManagement" Version="3.7.304.24" />
-    <PackageReference Include="AWSSDK.SQS" Version="3.7.501.14" />
-    <PackageReference Include="AWSSDK.SecurityToken" Version="3.7.501.14" />
+    <PackageReference Include="AWSSDK.SQS" Version="3.7.300.12" />
+    <PackageReference Include="AWSSDK.SecurityToken" Version="3.7.300.13" />
     <PackageReference Include="Azure.Storage.Blobs" Version="12.19.1" />
     <PackageReference Include="Crc32.NET" Version="1.2.0" />
     <PackageReference Include="Grpc.AspNetCore" Version="2.59.0" />
diff --git a/Engine/Source/Programs/Horde/Plugins/Compute/HordeServer.Compute/Aws/AwsCloudWatchMultiplexer.cs b/Engine/Source/Programs/Horde/Plugins/Compute/HordeServer.Compute/Aws/AwsCloudWatchMultiplexer.cs
--- a/Engine/Source/Programs/Horde/Plugins/Compute/HordeServer.Compute/Aws/AwsCloudWatchMultiplexer.cs
+++ b/Engine/Source/Programs/Horde/Plugins/Compute/HordeServer.Compute/Aws/AwsCloudWatchMultiplexer.cs
@@ -78,7 +78,6 @@ public class AwsCloudWatchMultiplexer : IAmazonCloudWatch
 	public Task<DeleteDashboardsResponse> DeleteDashboardsAsync(DeleteDashboardsRequest request, CancellationToken cancellationToken = default) { throw new NotImplementedException(); }
 	public Task<DeleteInsightRulesResponse> DeleteInsightRulesAsync(DeleteInsightRulesRequest request, CancellationToken cancellationToken = default) { throw new NotImplementedException(); }
 	public Task<DeleteMetricStreamResponse> DeleteMetricStreamAsync(DeleteMetricStreamRequest request, CancellationToken cancellationToken = default) { throw new NotImplementedException(); }
-	public Task<DescribeAlarmContributorsResponse> DescribeAlarmContributorsAsync(DescribeAlarmContributorsRequest request, CancellationToken cancellationToken) { throw new NotImplementedException(); }
 	public Task<DescribeAlarmHistoryResponse> DescribeAlarmHistoryAsync(CancellationToken cancellationToken = default) { throw new NotImplementedException(); }
 	public Task<DescribeAlarmHistoryResponse> DescribeAlarmHistoryAsync(DescribeAlarmHistoryRequest request, CancellationToken cancellationToken = default) { throw new NotImplementedException(); }
 	public Task<DescribeAlarmsResponse> DescribeAlarmsAsync(CancellationToken cancellationToken = default) { throw new NotImplementedException(); }
diff --git a/Engine/Source/Programs/Horde/Plugins/Compute/HordeServer.Compute/HordeServer.Compute.csproj b/Engine/Source/Programs/Horde/Plugins/Compute/HordeServer.Compute/HordeServer.Compute.csproj
--- a/Engine/Source/Programs/Horde/Plugins/Compute/HordeServer.Compute/HordeServer.Compute.csproj
+++ b/Engine/Source/Programs/Horde/Plugins/Compute/HordeServer.Compute/HordeServer.Compute.csproj
@@ -44,14 +44,14 @@
   </ItemGroup>
 
   <ItemGroup>
-    <PackageReference Include="AWSSDK.AutoScaling" Version="3.7.501.7" />
-    <PackageReference Include="AWSSDK.CloudWatch" Version="3.7.502.7" />
-    <PackageReference Include="AWSSDK.EC2" Version="3.7.514" />
-    <PackageReference Include="AWSSDK.Extensions.NETCore.Setup" Version="3.7.400" />
-    <PackageReference Include="AWSSDK.S3" Version="3.7.501.8" />
+    <PackageReference Include="AWSSDK.AutoScaling" Version="3.7.301.9" />
+    <PackageReference Include="AWSSDK.CloudWatch" Version="3.7.300.12" />
+    <PackageReference Include="AWSSDK.EC2" Version="3.7.306.6" />
+    <PackageReference Include="AWSSDK.Extensions.NETCore.Setup" Version="3.7.300" />
+    <PackageReference Include="AWSSDK.S3" Version="3.7.304.1" />
     <PackageReference Include="AWSSDK.SimpleSystemsManagement" Version="3.7.304.24" />
-    <PackageReference Include="AWSSDK.SQS" Version="3.7.501.14" />
-    <PackageReference Include="AWSSDK.SecurityToken" Version="3.7.501.14" />
+    <PackageReference Include="AWSSDK.SQS" Version="3.7.300.12" />
+    <PackageReference Include="AWSSDK.SecurityToken" Version="3.7.300.13" />
     <PackageReference Include="Azure.Storage.Blobs" Version="12.19.1" />
     <PackageReference Include="MongoDB.Driver" Version="2.22.0" />
     <PackageReference Include="Grpc.AspNetCore" Version="2.59.0" />
diff --git a/Engine/Source/Programs/Horde/Plugins/Storage/HordeServer.Storage/HordeServer.Storage.csproj b/Engine/Source/Programs/Horde/Plugins/Storage/HordeServer.Storage/HordeServer.Storage.csproj
--- a/Engine/Source/Programs/Horde/Plugins/Storage/HordeServer.Storage/HordeServer.Storage.csproj
+++ b/Engine/Source/Programs/Horde/Plugins/Storage/HordeServer.Storage/HordeServer.Storage.csproj
@@ -21,8 +21,8 @@
   </ItemGroup>
 
   <ItemGroup>
-    <PackageReference Include="AWSSDK.S3" Version="3.7.501.8" />
-    <PackageReference Include="AWSSDK.Extensions.NETCore.Setup" Version="3.7.400" />
+    <PackageReference Include="AWSSDK.S3" Version="3.7.304.1" />
+    <PackageReference Include="AWSSDK.Extensions.NETCore.Setup" Version="3.7.300" />
     <PackageReference Include="Azure.Storage.Blobs" Version="12.19.1" />
     <PackageReference Include="Google.Cloud.Storage.V1" Version="4.10.0" />
     <PackageReference Include="MongoDB.Driver" Version="2.22.0" />
