[Horde][BuildHealth][Server] Move Job/Step logging to information level

diff --git a/Engine/Source/Programs/Horde/Plugins/Build/HordeServer.Build/Issues/IssueService.cs b/Engine/Source/Programs/Horde/Plugins/Build/HordeServer.Build/Issues/IssueService.cs
--- a/Engine/Source/Programs/Horde/Plugins/Build/HordeServer.Build/Issues/IssueService.cs
+++ b/Engine/Source/Programs/Horde/Plugins/Build/HordeServer.Build/Issues/IssueService.cs
@@ -787,7 +787,7 @@ namespace HordeServer.Issues
 			{
 				// Attach issue audit log buffers to each issue event
 				issueEvent.AuditLogger = new IssueAuditLogger();
-				issueEvent.AuditLogger.LogDebug("Logging events from Job {JobId}, step {StepId}", job.Id, step.Id);
+				issueEvent.AuditLogger.LogInformation("Logging events from Job {JobId}, step {StepId}", job.Id, step.Id);
 
 				foreach (IssueHandler handler in handlers)
 				{
