Horde: Update release notes

#rnx

diff --git a/Engine/Source/Programs/Horde/Docs/ReleaseNotes.md b/Engine/Source/Programs/Horde/Docs/ReleaseNotes.md
--- a/Engine/Source/Programs/Horde/Docs/ReleaseNotes.md
+++ b/Engine/Source/Programs/Horde/Docs/ReleaseNotes.md
@@ -1,5 +1,31 @@
 # Release Notes
 
+## 2025-09-11
+
+* Summary column & hours label for title & improving the console.log for change summary request count so we can see the *actual* request (45714224)
+* Silence static analysis complaint about large function (45701060)
+* Fix for MSVC compiler bug - provide all designated initializers for FRigVMCompiledInFunction including for predicates. (45698719)
+* JobSummary telemetry (45677295)
+* Remove left-click to add highlighted text to search box in log view (45606471)
+* Put error handling in own component to keep root view from needing to be an observer (45543682)
+* Increase title size to match changelist/date anchors & checkbox to hide date anchors (45535364)
+* Move Job/Step logging to information level (45494424)
+* Change Headers anchored by submit times (45494125)
+* Expand streams/{streamId}/changes api to support batch semantics. It can take a manually specified set of changes to obtain. (45494088)
+* Add using declaration to new JsonDocument reference to ensure disposal (45482577)
+* Resolve workspace materializer through workspace config/type. Also add a factory for the PerforceMaterializer. (45478698)
+* Add support for file structured logging property to hashed issue handler (45441165)
+* Add support for LatestChangeNumber constant in PerforceMaterializer (45437537)
+* Add new PerforceMaterializer. Uses plain workspaces in Perforce, currently optimized for single stream use. (45432379)
+* Selective dropdown updates on change & Step filters no longer being applied across streams (45340847)
+* Backing out JobExecutor killing jobs on UAT process hangs, we need to tweak this logic as it's killing jobs too early. (45313353)
+* Minor UI:
+  - Last refresh time
+  - Snap to right (most recent date)
+  - Add a 2 week option for time filter (45301672)
+* Minor UI:
+  - Disable date anchors in headers for preflights. Changelist <> Dates can fall out of sync easily, causing date header anchor instability. (45280879)
+
 ## 2025-08-28
 
 * [Horde][BuildHealth][Dashboard] Minor UI enhancements for 5.7:
