Horde: Filter more files and improve description generation in release notes tool

diff --git a/Engine/Source/Programs/Horde/Tools/ReleaseNotes/Program.cs b/Engine/Source/Programs/Horde/Tools/ReleaseNotes/Program.cs
--- a/Engine/Source/Programs/Horde/Tools/ReleaseNotes/Program.cs
+++ b/Engine/Source/Programs/Horde/Tools/ReleaseNotes/Program.cs
@@ -13,33 +13,48 @@ namespace CreateReleaseNotes
 	internal class Program
 	{
 		static readonly string[] s_paths =
-		{
-			$"Engine/Source/Programs/Horde/...",
-			$"Engine/Source/Programs/Shared/...",
-		};
-
+		[
+			"Engine/Source/Programs/Horde/...",
+			"Engine/Source/Programs/Shared/...",
+		];
+		
+		private static readonly string[] s_ignoreFilter =
+		[
+			// Ignore these as they have (almost) no impact on Horde
+			"Engine/Source/Programs/Shared/EpicGames.UBA/...",
+			"Engine/Source/Programs/Shared/EpicGames.Build/...",
+			"Engine/Source/Programs/Shared/EpicGames.UHT/...",
+			"Engine/Source/Programs/Shared/UnrealBuildTool.Tests/...",
+			"Engine/Source/Programs/Horde/Docs/ReleaseNotes.md",
+			
+			// TODO: Flag as low risk changes
+			"Engine/Source/Programs/Shared/EpicGames.AspNet/...",
+			"Engine/Source/Programs/Shared/EpicGames.Core/...",
+			"Engine/Source/Programs/Shared/EpicGames.Serialization/...",
+		];
+		
 		static readonly string[] s_serverFilter =
-		{
+		[
 			"...",
 			"-/Engine/Source/Programs/Horde/HordeDashboard/...",
-			"-/Engine/Source/Programs/Horde/HordeAgent/...",
-		};
+			"-/Engine/Source/Programs/Horde/HordeAgent/..."
+		];
 
 		static readonly string[] s_agentFilter =
-		{
+		[
 			"...",
 			"-/Engine/Source/Programs/Horde/HordeDashboard/...",
 			"-/Engine/Source/Programs/Horde/HordeServer*/...",
-			"-/Engine/Source/Programs/Horde/Plugins/*/HordeServer*/...",
-		};
+			"-/Engine/Source/Programs/Horde/Plugins/*/HordeServer*/..."
+		];
 
 		static readonly string[] s_dashboardFilter =
-		{
+		[
 			"-/...",
-			"/Engine/Source/Programs/Horde/HordeDashboard/...",
-		};
+			"/Engine/Source/Programs/Horde/HordeDashboard/..."
+		];
 
-		record class ChangeInfo(int Number, string Description, string? JiraTicket);
+		record ChangeInfo(int Number, string Description, string? JiraTicket);
 
 		class Options
 		{
@@ -120,7 +135,7 @@ namespace CreateReleaseNotes
 
 			List<string> paths = s_paths.Select(x => $"//{perforceInfo.ClientName}/{x}").ToList();
 
-			List<ChangesRecord> changes = await perforce.GetChangesAsync(ChangesOptions.None, null, change + 1, -1, ChangeStatus.Submitted, null, paths);
+			List<ChangesRecord> changes = await perforce.GetChangesAsync(ChangesOptions.LongOutput, null, change + 1, -1, ChangeStatus.Submitted, null, paths);
 			changes.RemoveAll(x => x.User.Equals("buildmachine", StringComparison.OrdinalIgnoreCase));
 			changes = changes.DistinctBy(x => x.Number).ToList();
 
@@ -132,11 +147,11 @@ namespace CreateReleaseNotes
 
 			List<DescribeRecord> describeRecords = await perforce.DescribeAsync(changes.Select(x => x.Number).ToArray());
 
-			HashSet<int> changeNumbers = new HashSet<int>();
-
-			FileFilter serverFilter = new FileFilter(s_serverFilter);
-			FileFilter agentFilter = new FileFilter(s_agentFilter);
-			FileFilter dashboardFilter = new FileFilter(s_dashboardFilter);
+			HashSet<int> changeNumbers = [];
+			FileFilter ignoreFilter = new (s_ignoreFilter);
+			FileFilter serverFilter = new (s_serverFilter);
+			FileFilter agentFilter = new (s_agentFilter);
+			FileFilter dashboardFilter = new (s_dashboardFilter);
 
 			logger.LogInformation("");
 			foreach (DescribeRecord describeRecord in describeRecords.OrderBy(x => x.Number))
@@ -147,15 +162,12 @@ namespace CreateReleaseNotes
 				string? jiraTicket = jiraMatch.Success ? jiraMatch.Groups[1].Value : null;
 
 				description = Regex.Replace(description, @"^[a-zA-Z]+:\s*", "");
+				description = Regex.Replace(description, @"^(?:\[[A-Za-z0-9_-]+\]\s*)+(?::\s*)?", "");
 				description = Regex.Replace(description, @"^\s*#.*$", "", RegexOptions.Multiline);
 				description = Regex.Replace(description, @"\n\s*", "\n");
+				description = Regex.Replace(description, @"^\s*\*\s*", "");
 				description = description.Trim();
 
-				if (!Regex.IsMatch(describeRecord.Description, @"^\s*#rnx\s*$", RegexOptions.Multiline))
-				{
-					parsedChanges.Add(new ChangeInfo(describeRecord.Number, description, jiraTicket));
-				}
-
 				string logDescription = Regex.Replace(description, @"\s+", " ");
 
 				const int MaxDescriptionLength = 80;
@@ -167,12 +179,19 @@ namespace CreateReleaseNotes
 				bool affectsServer = false;
 				bool affectsAgent = false;
 				bool affectsDashboard = false;
+				bool ignoreChangelist = false;
 
 				foreach (DescribeFileRecord fileRecord in describeRecord.Files)
 				{
 					string path = fileRecord.DepotFile;
 					path = Regex.Replace(path, @"^//[^/]+/[^/]+/", "");
-
+					
+					if (ignoreFilter.Matches(path))
+					{
+						ignoreChangelist = true;
+						break;
+					}
+					
 					if (describeRecord.Number > serverChange && serverFilter.Matches(path))
 					{
 						affectsServer = true;
@@ -186,6 +205,16 @@ namespace CreateReleaseNotes
 						affectsDashboard = true;
 					}
 				}
+				
+				if (ignoreChangelist)
+				{
+					continue;
+				}
+				
+				if (!Regex.IsMatch(describeRecord.Description, @"^\s*#rnx\s*$", RegexOptions.Multiline))
+				{
+					parsedChanges.Add(new ChangeInfo(describeRecord.Number, description, jiraTicket));
+				}
 
 				if (affectsServer || affectsAgent || affectsDashboard)
 				{
