[Horde][BuildHealth][Dashboard] Export to csv button.

diff --git a/Engine/Source/Programs/Horde/HordeDashboard/src/components/buildhealth/stepoutcome/StepOutcomeD3Component.tsx b/Engine/Source/Programs/Horde/HordeDashboard/src/components/buildhealth/stepoutcome/StepOutcomeD3Component.tsx
--- a/Engine/Source/Programs/Horde/HordeDashboard/src/components/buildhealth/stepoutcome/StepOutcomeD3Component.tsx
+++ b/Engine/Source/Programs/Horde/HordeDashboard/src/components/buildhealth/stepoutcome/StepOutcomeD3Component.tsx
@@ -1,6 +1,6 @@
 // Copyright Epic Games, Inc. All Rights Reserved.
 
-import { mergeStyleSets } from "@fluentui/react";
+import { DefaultButton, DirectionalHint, mergeStyleSets, Stack, TooltipHost } from "@fluentui/react";
 import * as d3 from "d3";
 import { getHordeStyling, preloadFonts } from "horde/styles/Styles";
 import { useEffect, useRef } from "react";
@@ -27,6 +27,42 @@ const tooltipClasses = mergeStyleSets({
 
 // #endregion -- Styles --
 
+// #region -- Step Outcome Table Download Handlers --
+
+function saveFile(content: string, filename: string) {
+    const blob = new Blob([content], { type: "text/csv" });
+    const url = URL.createObjectURL(blob);
+
+    const a = document.createElement("a");
+    a.href = url;
+    a.download = filename;
+    a.click();
+
+    URL.revokeObjectURL(url);
+}
+
+function exportStepOutcomeDataToCSV(stepOutcomeTable: StepOutcomeTable) {
+    const defaultName = new Date()
+        .toLocaleDateString("en-US", {
+            year: "numeric",
+            month: "long",
+            day: "numeric",
+        })
+        .replace(/,/g, "")
+        .replace(/ /g, "-");
+
+    const filename = window.prompt("Enter name of export:", "ProjectExport");
+    if (!filename) return;
+
+    const stepOrderedCsvContent = stepOutcomeTable.toRowOrderedCSV();
+    const changeOrderedCsvContent = stepOutcomeTable.toColOrderedCSV();
+
+    saveFile(stepOrderedCsvContent, `${filename}-StepOrderedData-${defaultName}.csv`);
+    saveFile(changeOrderedCsvContent, `${filename}-ChangeOrderedData-${defaultName}.csv`);
+}
+
+// #endregion -- Step Outcome Table Download Handlers --
+
 /**
  * Function to create a D3 format table for the step outcome table.
  * @param stepOutcomeTable The data to base the table off of.
@@ -639,7 +675,44 @@ function GenerateD3Table(stepOutcomeTable: StepOutcomeTable, handler: StepOutcom
                 }}
             />
             {/* Right side: incremental spinner */}
-            <IncrementalRefreshIndicatorPanel handler={handler} />
+            <Stack
+                key={"export-incremental-container"}
+                verticalAlign="stretch"
+                styles={{
+                    root: {
+                        display: "flex",
+                        flexDirection: "column",
+                        justifyContent: "space-between",
+                        height: "100%",
+                    },
+                }}
+            >
+                {/* Top: spinner */}
+                <IncrementalRefreshIndicatorPanel handler={handler} />
+                {/* Bottom: button */}
+                <TooltipHost
+                    content="Export the current step outcome history data to CSV."
+                    directionalHint={DirectionalHint.topLeftEdge}
+                >
+                    <DefaultButton
+                        text=">"
+                        styles={{
+                            root: {
+                                width: 30,
+                                height: 30,
+                                minWidth: 0,
+                                padding: 0,
+                                lineHeight: "15px",
+                            },
+                            label: {
+                                fontSize: 10,
+                                lineHeight: "15px",
+                            },
+                        }}
+                        onClick={() => exportStepOutcomeDataToCSV(stepOutcomeTable)}
+                    />
+                </TooltipHost>
+            </Stack>
         </div>
     );
 };
diff --git a/Engine/Source/Programs/Horde/HordeDashboard/src/components/buildhealth/stepoutcome/StepOutcomeDataTypes.ts b/Engine/Source/Programs/Horde/HordeDashboard/src/components/buildhealth/stepoutcome/StepOutcomeDataTypes.ts
--- a/Engine/Source/Programs/Horde/HordeDashboard/src/components/buildhealth/stepoutcome/StepOutcomeDataTypes.ts
+++ b/Engine/Source/Programs/Horde/HordeDashboard/src/components/buildhealth/stepoutcome/StepOutcomeDataTypes.ts
@@ -153,6 +153,22 @@
 
         return durationStr;
     }
+
+    /**
+     * Gets the CSV representation of the the table entry. Property ordering matches @see getCSVHeader .
+     * @returns CSV representation of the table entry.
+     */
+    toCSVRepresntation(): string {
+        return `${this.change},${this.streamId},${this.jobName},${this?.jobId},${this?.stepResponse.name},${this?.stepResponse.state},${this?.stepResponse.outcome},${this?.stepResponse.stepDurationMs / 1000}\n`;
+    }
+
+    /**
+     * Gets the CSV property names of the step outcome entry. Property ordering matches @see toCSVRepresntation .
+     * @returns The comma separated properties serialization.
+     */
+    static getCSVHeader() {
+        return "Change,StreamId,JobName,JobId,StepName,StepState,StepOutcome,Duration(Sec)\n";
+    }
     // #endregion -- Public API --
 }
 
@@ -237,7 +253,6 @@
 
     /**
      * Sets whether @see getDateAnchoredChanges is supported for the given data.
-     * @todo - UE-315881 will remove the need for this constraint.
      */
     supportsDateAnchoredChanges: boolean = true;
 
@@ -286,6 +301,43 @@
         this.tableStreamMetadata.set(streamId, priorCount);
     }
 
+    /**
+     * Gets the the date anchored changes. 
+     * @returns * Generator of all of the @see ChangeHeader items separated by @see DateHeader that represent dates. 
+     */
+    private *getDateAnchoredChanges(): Generator<ChangeHeader | DateHeader> {
+        if (!this.supportsDateAnchoredChanges) {
+            for (const header of this.changeColHeaders) {
+                yield header;
+            }
+
+            return;
+        }
+
+        let prevDay: string | null = null;
+        for (let i: number = 0; i < this.changeColHeaders.length; ++i) {
+
+            // If the date associated with the changelist header has changed since the last, we are on a new day. Set that as the colimn.
+            let changeHeader = this.changeColHeaders[i];
+            if (changeHeader.date !== undefined) {
+                const headerDate = new Date(changeHeader.date!);
+
+                // Extract the day portion as YYYY-MM-DD
+                const dayStr = `${headerDate.getUTCFullYear()}-${headerDate.getUTCMonth() + 1}-${headerDate.getUTCDate()}`;
+
+                if (prevDay !== dayStr) {
+                    prevDay = dayStr;
+                    let dateAnchorChangelistHeader: DateHeader = {
+                        type: "date",
+                        date: dayStr
+                    }
+                    yield dateAnchorChangelistHeader;
+                }
+                yield changeHeader;
+            }
+        }
+    }
+
     // #endregion -- Private API --
 
     // #region -- Public API --
@@ -524,40 +576,43 @@
     }
 
     /**
-     * Gets the the date anchored changes. 
-     * @returns * Generator of all of the @see ChangeHeader items separated by @see DateHeader that represent dates. 
+     * Generates the csv string representation of the step outcome table, in row order.
+     * @returns A string in CSV format that represents the underlying step outcome table, row order.
      */
-    private *getDateAnchoredChanges(): Generator<ChangeHeader | DateHeader> {
-        if (!this.supportsDateAnchoredChanges) {
-            for (const header of this.changeColHeaders) {
-                yield header;
+    toRowOrderedCSV(): string {
+        let resultString: string = StepOutcomeTableEntry.getCSVHeader();
+        for (let rowIndex = 0; rowIndex < this.tableEntries.length; rowIndex++) {
+            const row = this.tableEntries[rowIndex];
+            for (let i: number = 0; i < row.length; ++i) {
+                if (!row[i]) {
+                    continue;
+                }
+                let tableEntry: StepOutcomeTableEntry = row[i]!;
+                resultString += tableEntry.toCSVRepresntation();
             }
-
-            return;
         }
 
-        let prevDay: string | null = null;
-        for (let i: number = 0; i < this.changeColHeaders.length; ++i) {
+        return resultString;
+    }
 
-            // If the date associated with the changelist header has changed since the last, we are on a new day. Set that as the colimn.
-            let changeHeader = this.changeColHeaders[i];
-            if (changeHeader.date !== undefined) {
-                const headerDate = new Date(changeHeader.date!);
-
-                // Extract the day portion as YYYY-MM-DD
-                const dayStr = `${headerDate.getUTCFullYear()}-${headerDate.getUTCMonth() + 1}-${headerDate.getUTCDate()}`;
-
-                if (prevDay !== dayStr) {
-                    prevDay = dayStr;
-                    let dateAnchorChangelistHeader: DateHeader = {
-                        type: "date",
-                        date: dayStr
-                    }
-                    yield dateAnchorChangelistHeader;
+    /**
+     * Generates the csv string representation of the step outcome table, in column order.
+     * @returns A string in CSV format that represents the underlying step outcome table, column order.
+     */
+    toColOrderedCSV(): string {
+        let resultString: string = StepOutcomeTableEntry.getCSVHeader();
+        for (let changeIdx = 0; changeIdx < this.changeColHeaders.length; changeIdx++) {
+            const changeEntries = [...this.getChangeColumn(this.changeColHeaders[changeIdx].change)];
+            for (let i: number = 0; i < changeEntries.length; ++i) {
+                if (!changeEntries[i]) {
+                    continue;
                 }
-                yield changeHeader;
+                let tableEntry: StepOutcomeTableEntry = changeEntries[i]!;
+                resultString += tableEntry.toCSVRepresntation();
             }
         }
+
+        return resultString;
     }
 
     // #endregion -- Public API --
diff --git a/Engine/Source/Programs/Horde/HordeDashboard/src/components/buildhealth/stepoutcome/StepOutcomeSharedUIComponents.tsx b/Engine/Source/Programs/Horde/HordeDashboard/src/components/buildhealth/stepoutcome/StepOutcomeSharedUIComponents.tsx
--- a/Engine/Source/Programs/Horde/HordeDashboard/src/components/buildhealth/stepoutcome/StepOutcomeSharedUIComponents.tsx
+++ b/Engine/Source/Programs/Horde/HordeDashboard/src/components/buildhealth/stepoutcome/StepOutcomeSharedUIComponents.tsx
@@ -99,7 +99,8 @@ export const IncrementalRefreshIndicatorPanel: React.FC<{ handler: StepOutcomeDa
             gridRow: '1 / span 2',
             gridColumn: 3,
             overflowY: 'auto',
-            width: 25,
+            width: 30,
+            height: 30,
         }}
     >
         {handler.isInIncrementalDataRefresh && <Spinner size={SpinnerSize.small} />}
