[Horde][BuildHealth][Dashboard] Styling pass V3 & general UI improvements:
- Styling for aborted
- Styling for skipped
- Sticky summary columns
- Min width reduced for smaller screens
- Skipped as red (+ view/ui options segregation refactor)
- Warnings as summary success
- Patch changelist UTC skew in UI
- Fall back to job start when no changelist data is available from server
- Clearing project cascades correctly

diff --git a/Engine/Source/Programs/Horde/HordeDashboard/src/components/buildhealth/BuildHealthOptions.ts b/Engine/Source/Programs/Horde/HordeDashboard/src/components/buildhealth/BuildHealthOptions.ts
--- a/Engine/Source/Programs/Horde/HordeDashboard/src/components/buildhealth/BuildHealthOptions.ts
+++ b/Engine/Source/Programs/Horde/HordeDashboard/src/components/buildhealth/BuildHealthOptions.ts
@@ -192,6 +192,19 @@ export async function loadBuildHealthOptionsFromParams(
 /**
  * Build Health View Options.
  */
+export class BuildHealthUIOptionsState {
+    @observable includeDateAnchors: boolean = true;
+    @observable skippedIsFailure: boolean = false;
+    @observable warningsAsSummarySuccess: boolean = false;
+
+    constructor() {
+        makeObservable(this);
+    }
+}
+
+/**
+ * Build Health Filter Options.
+ */
 export class BuildHealthOptionsState {
 
     // #region -- Published Options --
@@ -212,7 +225,6 @@ export class BuildHealthOptionsState {
     @observable enabledSteps: Record<string, string> = {};
     @observable jobHistoryTimeSpan: TimeSpan = DEFAULT_JOB_HISTORY_TIME_SPAN;
     @observable includePreflight: boolean = false;
-    @observable includeDateAnchors: boolean = true;
 
     // #endregion -- User Selected Options --
 
@@ -237,6 +249,8 @@ export class BuildHealthOptionsController {
     // #region -- Private Members --
 
     readonly state: BuildHealthOptionsState;
+    readonly uiState: BuildHealthUIOptionsState;
+
     private lastSynchronize = -1;
 
     // #endregion -- Private Members --
@@ -249,8 +263,9 @@ export class BuildHealthOptionsController {
 
     // #region -- Constructor --
 
-    constructor(state: BuildHealthOptionsState) {
+    constructor(state: BuildHealthOptionsState, uiState: BuildHealthUIOptionsState) {
         this.state = state;
+        this.uiState = uiState;
         makeObservable(this);
     }
 
@@ -271,6 +286,7 @@ export class BuildHealthOptionsController {
     @action
     toggleSingleProject(id: string, name: string) {
         this.state.enabledProjects = { [id]: name };
+         this.clearStreams(true);
         this.setOptionsChanged();
     }
 
@@ -381,8 +397,25 @@ export class BuildHealthOptionsController {
      */
     @action
     setHideDateAnchors(includeDateAnchors: boolean) {
-        this.state.includeDateAnchors = includeDateAnchors;
-        this.setOptionsChanged();
+        this.uiState.includeDateAnchors = includeDateAnchors;
+    }
+
+    /**
+     * Sets whether entries in the table that are skipped state should be rendered as failures.
+     * @param skippedIsFailure True to render as failures, false otherwise.
+     */
+    @action
+    setSkippedIsFailure(skippedIsFailure: boolean) {
+        this.uiState.skippedIsFailure = skippedIsFailure;
+    }
+
+    /**
+ * Sets whether summary rows should consider warnings as a success, or failure.
+ * @param warningsAsSummarySuccess True to include warnings in the numerator of success ratio, false otherwise.
+ */
+    @action
+    setWarningsAsSummarySuccess(warningsAsSummarySuccess: boolean) {
+        this.uiState.warningsAsSummarySuccess = warningsAsSummarySuccess;
     }
 
     /**
diff --git a/Engine/Source/Programs/Horde/HordeDashboard/src/components/buildhealth/BuildhealthView.tsx b/Engine/Source/Programs/Horde/HordeDashboard/src/components/buildhealth/BuildhealthView.tsx
--- a/Engine/Source/Programs/Horde/HordeDashboard/src/components/buildhealth/BuildhealthView.tsx
+++ b/Engine/Source/Programs/Horde/HordeDashboard/src/components/buildhealth/BuildhealthView.tsx
@@ -12,7 +12,7 @@ import { TopNav } from "../TopNav";
 import { StepOutcomeView } from "../buildhealth/stepoutcome/StepOutcome";
 import { BuildHealthDataHandler, DataHandlerRefreshRequest, expandRefreshRequest } from "./BuildHealthDataHandler";
 import { KEY_SEPARATOR, PARAMETER_KEY_PREFIX, StepRefData, TemplateRefData } from "./BuildHealthDataTypes";
-import { BuildHealthOptionsController, BuildHealthOptionsState, JobHistoryTimeSpans, loadBuildHealthOptionsFromParams, parseBuildHealthQueryParams, TimeSpan } from "./BuildHealthOptions";
+import { BuildHealthOptionsController, BuildHealthOptionsState, BuildHealthUIOptionsState, JobHistoryTimeSpans, loadBuildHealthOptionsFromParams, parseBuildHealthQueryParams, TimeSpan } from "./BuildHealthOptions";
 import { encodeStepKey, encodeTemplateKey } from "./BuildHealthUtilities";
 import { stepOutcomeTableClasses } from "./stepoutcome/StepOutcomeSharedUIComponents";
 
@@ -471,10 +471,17 @@ const BuildHealthUrlSync: React.FC<{ buildHealthController: BuildHealthOptionsCo
  * @returns React Component.
  */
 const BuildHealthSidebar: React.FC<{ handler: BuildHealthDataHandler, buildHealthOptions: BuildHealthOptionsController }> = function ConstrucBuildHealthSidebar({ handler, buildHealthOptions }) {
-    const { hordeClasses } = getHordeStyling();
+    const { hordeClasses, modeColors } = getHordeStyling();
     const [timeSpanKey, setTimeSpanKey] = useState(() => buildHealthOptions.state.jobHistoryTimeSpan.key);
     const [includePreflight, setIncludePreflight] = useState(() => buildHealthOptions.state.includePreflight)
-    const [hideDateAnchors, setHideDateAnchors] = useState(() => buildHealthOptions.state.includeDateAnchors)
+
+    // #region -- Visual Options -- 
+
+    const [hideDateAnchors, setHideDateAnchors] = useState(() => buildHealthOptions.uiState.includeDateAnchors)
+    const [skippedAsRed, setSkippedAsRed] = useState(() => buildHealthOptions.uiState.skippedIsFailure)
+    const [warningsAsSummarySuccess, setWarningsAsSummarySuccess] = useState(() => buildHealthOptions.uiState.warningsAsSummarySuccess)
+
+    // #ednregion -- Visual Options --
 
     // #region -- Property Callbacks --
 
@@ -494,6 +501,16 @@ const BuildHealthSidebar: React.FC<{ handler: BuildHealthDataHandler, buildHealt
         buildHealthOptions.setHideDateAnchors(isChecked);
     }
 
+    const updateSkippedIsFailure = (isChecked) => {
+        setSkippedAsRed(isChecked);
+        buildHealthOptions.setSkippedIsFailure(isChecked);
+    }
+
+    const updateWarningsAsSummarySuccess = (isChecked) => {
+        setWarningsAsSummarySuccess(isChecked);
+        buildHealthOptions.setWarningsAsSummarySuccess(isChecked);
+    }
+
     // #endregion -- Property Callbacks --
 
     // #region -- Property Use Effects --
@@ -507,8 +524,16 @@ const BuildHealthSidebar: React.FC<{ handler: BuildHealthDataHandler, buildHealt
     }, [buildHealthOptions.state.includePreflight]);
 
     useEffect(() => {
-        setHideDateAnchors(buildHealthOptions.state.includeDateAnchors);
-    }, [buildHealthOptions.state.includeDateAnchors]);
+        setHideDateAnchors(buildHealthOptions.uiState.includeDateAnchors);
+    }, [buildHealthOptions.uiState.includeDateAnchors]);
+
+    useEffect(() => {
+        setSkippedAsRed(buildHealthOptions.uiState.skippedIsFailure);
+    }, [buildHealthOptions.uiState.skippedIsFailure]);
+
+    useEffect(() => {
+        setWarningsAsSummarySuccess(buildHealthOptions.uiState.warningsAsSummarySuccess);
+    }, [buildHealthOptions.uiState.warningsAsSummarySuccess]);
 
     // #endregion -- Property Use Effects --
 
@@ -516,17 +541,29 @@ const BuildHealthSidebar: React.FC<{ handler: BuildHealthDataHandler, buildHealt
 
     return (
         <Stack style={{ minWidth: 300, maxWidth: 450, paddingRight: 18 }}>
-            <Stack className={hordeClasses.modal}>
-                <Stack>
+            <Stack className={hordeClasses.modal} tokens={{ childrenGap: 12 }}>
+                <Stack
+                    styles={{
+                        root: {
+                            border: `1px solid ${modeColors.content}`,
+                            borderRadius: 4,
+                            padding: 8,
+                        },
+                    }}
+                >
                     <Stack style={{ paddingTop: 0, paddingBottom: 4 }}>
-                        <Label>Step Outcome Filter</Label>
+                        <Label>Step Outcome Filters</Label>
                     </Stack>
                     <ProjectDropdownSingle handler={handler} buildHealthOptions={buildHealthOptions} />
                     <StreamDropdown handler={handler} buildHealthOptions={buildHealthOptions} />
                     <TemplateDropdown handler={handler} buildHealthOptions={buildHealthOptions} />
                     <StepDropdown handler={handler} buildHealthOptions={buildHealthOptions} />
                     <Stack>
-                        <TooltipHost content="Adjust the historical job history time span.">
+                        <TooltipHost content="Adjust the historical job history time span."
+                            calloutProps={{
+                                directionalHint: DirectionalHint.rightCenter
+                            }}
+                        >
                             <Label>Jobs Since:</Label>
                         </TooltipHost>
                         <ComboBox
@@ -538,35 +575,88 @@ const BuildHealthSidebar: React.FC<{ handler: BuildHealthDataHandler, buildHealt
                                 buildHealthOptions.synchronizeDerivedKeys();
                             }}
                         />
-                        {/* @todo - UE-315881 will remove the need for this constraint. */}
                         <TooltipHost content={
                             <div style={{ maxWidth: 300 }}>
                                 <div>
-                                    Include preflights in the job history results, or not.
+                                    Include presubmits in the job history results, or not.
                                 </div>
                             </div>
-                        }>
+                        }
+                            calloutProps={{
+                                directionalHint: DirectionalHint.rightCenter
+                            }}
+                        >
                             <Checkbox
-                                label="Include Preflights"
+                                label="Include presubmits"
                                 checked={includePreflight}
                                 onChange={(_, isChecked) => updateIncludePreflight(!!isChecked)}
                             />
                         </TooltipHost>
-                        <TooltipHost content={
-                            <div style={{ maxWidth: 300 }}>
-                                <div>
-                                    Include change submit date anchors in the header to separate changes.
-                                </div>
-                            </div>
-                        }>
-                            <Checkbox
-                                label="Use Date Anchors"
-                                checked={hideDateAnchors}
-                                onChange={(_, isChecked) => updateIncludeDateAnchors(!!isChecked)}
-                            />
-                        </TooltipHost>
                     </Stack>
                 </Stack>
+                <Stack
+                    styles={{
+                        root: {
+                            border: `1px solid ${modeColors.content}`,
+                            borderRadius: 4,
+                            padding: 8,
+                        },
+                    }}>
+                    <Stack style={{ paddingTop: 0, paddingBottom: 4 }}>
+                        <Label>Step Outcome Visual Options</Label>
+                    </Stack>
+                    <TooltipHost content={
+                        <div style={{ maxWidth: 300 }}>
+                            <div>
+                                Include change submit date anchors in the header to separate changes.
+                            </div>
+                        </div>
+                    }
+                        calloutProps={{
+                            directionalHint: DirectionalHint.rightCenter
+                        }}
+                    >
+                        <Checkbox
+                            label="Use date anchors"
+                            checked={hideDateAnchors}
+                            onChange={(_, isChecked) => updateIncludeDateAnchors(!!isChecked)}
+                        />
+                    </TooltipHost>
+                    <TooltipHost content={
+                        <div style={{ maxWidth: 300 }}>
+                            <div>
+                                Skipped steps (due to dependency failures) render the same as failures.
+                            </div>
+                        </div>
+                    }
+                        calloutProps={{
+                            directionalHint: DirectionalHint.rightCenter
+                        }}
+                    >
+                        <Checkbox
+                            label="Skipped is failure"
+                            checked={skippedAsRed}
+                            onChange={(_, isChecked) => updateSkippedIsFailure(!!isChecked)}
+                        />
+                    </TooltipHost>
+                    <TooltipHost content={
+                        <div style={{ maxWidth: 300 }}>
+                            <div>
+                                Warnings are considered success in succes/failure ratio for summary.
+                            </div>
+                        </div>
+                    }
+                        calloutProps={{
+                            directionalHint: DirectionalHint.rightCenter
+                        }}
+                    >
+                        <Checkbox
+                            label="Warnings as summary success"
+                            checked={warningsAsSummarySuccess}
+                            onChange={(_, isChecked) => updateWarningsAsSummarySuccess(!!isChecked)}
+                        />
+                    </TooltipHost>
+                </Stack>
             </Stack>
         </Stack>
     );
@@ -580,20 +670,33 @@ export const BuildHealthView: React.FC = observer(function ConstructBuildHealthV
     const windowSize = useWindowSize();
     const { hordeClasses, modeColors } = getHordeStyling();
 
+    const jobHistorySpan = useMemo(() => ({
+        start: buildHealthState.startDate,
+        end: buildHealthState.endDate
+    }), [buildHealthState.startDate, buildHealthState.endDate]);
+
     const filter = useMemo(() => ({
         streams: buildHealthState.stepOutcomeEnabledStreamKeys,
         jobs: buildHealthState.stepOutcomeEnabledJobKeys,
         steps: buildHealthState.stepOutcomeEnabledStepKeys,
-        jobHistorySpan: { start: buildHealthState.startDate, end: buildHealthState.endDate },
+        jobHistorySpan,
         includePreflights: buildHealthState.includePreflight,
-        includeDateAnchors: buildHealthState.includeDateAnchors
     }), [
         buildHealthState.stepOutcomeEnabledStreamKeys,
         buildHealthState.stepOutcomeEnabledJobKeys,
         buildHealthState.stepOutcomeEnabledStepKeys,
-        { start: buildHealthState.startDate, end: buildHealthState.endDate },
-        buildHealthState.includePreflight,
-        buildHealthState.includeDateAnchors
+        jobHistorySpan,
+        buildHealthState.includePreflight
+    ]);
+
+    const uiOptions = useMemo(() => ({
+        includeDateAnchors: buildHealthUIState.includeDateAnchors,
+        skippedAsRed: buildHealthUIState.skippedIsFailure,
+        warningsAsSummarySuccess: buildHealthUIState.warningsAsSummarySuccess
+    }), [
+        buildHealthUIState.includeDateAnchors,
+        buildHealthUIState.skippedIsFailure,
+        buildHealthUIState.warningsAsSummarySuccess
     ]);
 
     return (
@@ -607,8 +710,8 @@ export const BuildHealthView: React.FC = observer(function ConstructBuildHealthV
                     <Stack styles={{ root: { width: "100%" } }}>
                         <Stack horizontal styles={{ root: { minHeight: '85vh' } }} >
                             <BuildHealthSidebar handler={handler} buildHealthOptions={buildHealthController} />
-                            <Stack id="parent-buildhealth-stepoutcome" className={stepOutcomeTableClasses.root} grow style={{ overflowX: "auto", overflowY: "visible", minWidth: "1200px", position: "relative", border: `${modeColors.content} solid 2px` }}>
-                                <StepOutcomeView filter={filter} />
+                            <Stack id="parent-buildhealth-stepoutcome" className={stepOutcomeTableClasses.root} grow style={{ overflowX: "auto", overflowY: "visible", minWidth: "800px", position: "relative", border: `${modeColors.content} solid 2px` }}>
+                                <StepOutcomeView key="step-outcome-inline" filter={filter} uiOptions={uiOptions} />
                             </Stack>
                         </Stack>
                     </Stack>
@@ -623,7 +726,8 @@ export const BuildHealthView: React.FC = observer(function ConstructBuildHealthV
 // #region -- Script --
 
 const buildHealthState = new BuildHealthOptionsState();
-const buildHealthController = new BuildHealthOptionsController(buildHealthState);
+const buildHealthUIState = new BuildHealthUIOptionsState();
+const buildHealthController = new BuildHealthOptionsController(buildHealthState, buildHealthUIState);
 const handler = new BuildHealthDataHandler(buildHealthState, buildHealthController);
 
 // #endregion -- Script --
\ No newline at end of file
diff --git a/Engine/Source/Programs/Horde/HordeDashboard/src/components/buildhealth/stepoutcome/StepOutcome.tsx b/Engine/Source/Programs/Horde/HordeDashboard/src/components/buildhealth/stepoutcome/StepOutcome.tsx
--- a/Engine/Source/Programs/Horde/HordeDashboard/src/components/buildhealth/stepoutcome/StepOutcome.tsx
+++ b/Engine/Source/Programs/Horde/HordeDashboard/src/components/buildhealth/stepoutcome/StepOutcome.tsx
@@ -5,21 +5,24 @@ import { Text } from '@fluentui/react/lib/Text';
 import { observer } from "mobx-react-lite";
 import React, { useEffect, useState } from 'react';
 import { StepOutcomeD3Table } from "./StepOutcomeD3Component";
-import { prettyPrintStepOutcomeFilters, prettyPrintStepOutcomeFiltersSummary, StepOutcomeDataHandler, StepOutcomeFilters } from "./StepOutcomeDataHandler";
+import { prettyPrintStepOutcomeFilters, prettyPrintStepOutcomeFiltersSummary, StepOutcomeDataHandler, StepOutcomeFilters, StepOutcomeViewOptions } from "./StepOutcomeDataHandler";
 import { StepOutcomeTableEntry } from "./StepOutcomeDataTypes";
 import { sharedBorderStyle, stepOutcomeTableClasses } from "./StepOutcomeSharedUIComponents";
-import { toDate } from "./StepOutcomeUtilities";
+import { isTerminalState, toDate } from "./StepOutcomeUtilities";
 
 // #region -- View Components --
 
 type StepOutcomeTableProps = {
     filters: StepOutcomeFilters;
+    uiOptions: StepOutcomeViewOptions;
+
     handler: StepOutcomeDataHandler;
     onCellSelected: (item: StepOutcomeTableEntry) => void;
 };
 
 interface StepOutcomeViewProps {
     filter: StepOutcomeFilters;
+    uiOptions: StepOutcomeViewOptions;
     refreshTime?: number;
 }
 
@@ -50,6 +53,19 @@ const StepOutcomeFocusModal: React.FC<StepOutcomeFocusModalProps> = ({ item, onD
         }
     }, []);
 
+    const details: (string | JSX.Element)[][] = [
+        ['Date', startDate ? startDate.toDateString() : "N/A"],
+        ['Stream', item.streamId],
+        ['Code Change', item.change],
+        ['Job', item.jobName],
+        ['Step', item.stepResponse.name],
+        ['Duration', item.getDurationString()],
+        ['State', item.stepResponse.state],
+        isTerminalState(item) ? ['Outcome', item.stepResponse.outcome] : null,
+        ['Job Link', <a href={`/job/${item.jobId}`} target="_blank" rel="noopener noreferrer">View Job</a>],
+        ['Step Link', <a href={`/job/${item.jobId}?step=${item.stepResponse.id}`} target="_blank" rel="noopener noreferrer">View Step</a>],
+    ].filter((d): d is (string | JSX.Element)[] => Boolean(d));
+
     return (
         <Dialog
             modalProps={{
@@ -76,21 +92,9 @@ const StepOutcomeFocusModal: React.FC<StepOutcomeFocusModalProps> = ({ item, onD
             }}>
             <Stack>
                 <Stack horizontal tokens={{ childrenGap: 16 }} styles={{ root: { width: '100%', flexGrow: 0 } }}>
-
-                    {/* Left Column - Step Details*/}
+                    {/* Left Column - Step Details */}
                     <Stack styles={{ root: { width: '35%' } }} tokens={{ childrenGap: 4 }}>
-                        {[
-                            ['Date', startDate ? startDate.toDateString() : "N/A"],
-                            ['Stream', item.streamId],
-                            ['Code Change', item.change],
-                            ['Job', item.jobName],
-                            ['Step', item.stepResponse.name],
-                            ['Duration', item.getDurationString()],
-                            ['State', item.stepResponse.state],
-                            ['Outcome', item.stepResponse.outcome],
-                            ['Job Link', <a href={`/job/${item.jobId}`} target="_blank" rel="noopener noreferrer">View Job</a>],
-                            ['Step Link', <a href={`/job/${item.jobId}?step=${item.stepResponse.id}`} target="_blank" rel="noopener noreferrer">View Step</a>]
-                        ].map(([label, value], index) => (
+                        {details.map(([label, value], index) => (
                             <Stack key={`modal_step_item_${index}`} horizontal tokens={{ childrenGap: 12 }}>
                                 <Label styles={{ root: { minWidth: 100, padding: 0 } }}>{label}:</Label>
                                 <Text>{value}</Text>
@@ -142,8 +146,9 @@ const StepOutcomeFocusModal: React.FC<StepOutcomeFocusModalProps> = ({ item, onD
 /**
  * React Component that acts as a observable binder for data handler updates, and allows for easy swapping of other table types.
  */
-const StepOutcomeTableComponent: React.FC<StepOutcomeTableProps> = observer(({ handler, filters, onCellSelected }) => {
+const StepOutcomeTableComponent: React.FC<StepOutcomeTableProps> = observer(({ handler, filters, uiOptions, onCellSelected }) => {
     handler.updated;
+    uiOptions;
     const [dataCount, setDataCount] = useState(() => handler.StepOutcomeTableData.dataCount);
 
     useEffect(() => {
@@ -196,7 +201,7 @@ const StepOutcomeTableComponent: React.FC<StepOutcomeTableProps> = observer(({ h
  * @param StepOutcomeViewProps The filter to use for the step outcome, and the refresh cadence of the data.
  * @returns React Component.
  */
-export const StepOutcomeView: React.FC<StepOutcomeViewProps> = ({ filter, refreshTime }) => {
+export const StepOutcomeView: React.FC<StepOutcomeViewProps> = ({ filter, uiOptions, refreshTime }) => {
     useEffect(() => {
         handler.start();
 
@@ -210,6 +215,10 @@ export const StepOutcomeView: React.FC<StepOutcomeViewProps> = ({ filter, refres
         handler.setFilter(filter);
     }, [filter]);
 
+    useEffect(() => {
+        handler.StepOutcomeTableData.setViewModelOptions(uiOptions);
+    }, [uiOptions]);
+
     useEffect(() => {
         if (refreshTime) {
             handler.setRefreshTime(refreshTime);
@@ -222,7 +231,7 @@ export const StepOutcomeView: React.FC<StepOutcomeViewProps> = ({ filter, refres
         <Stack>
             <Stack key="step-outcome-table">
                 {selectedItem && (<StepOutcomeFocusModal item={selectedItem} onDismiss={() => setSelectedItem(null)} />)}
-                <StepOutcomeTableComponent handler={handler} filters={filter} onCellSelected={setSelectedItem} />
+                <StepOutcomeTableComponent handler={handler} filters={filter} uiOptions={uiOptions} onCellSelected={setSelectedItem} />
             </Stack>
         </Stack>
     );
diff --git a/Engine/Source/Programs/Horde/HordeDashboard/src/components/buildhealth/stepoutcome/StepOutcomeD3Component.tsx b/Engine/Source/Programs/Horde/HordeDashboard/src/components/buildhealth/stepoutcome/StepOutcomeD3Component.tsx
--- a/Engine/Source/Programs/Horde/HordeDashboard/src/components/buildhealth/stepoutcome/StepOutcomeD3Component.tsx
+++ b/Engine/Source/Programs/Horde/HordeDashboard/src/components/buildhealth/stepoutcome/StepOutcomeD3Component.tsx
@@ -6,8 +6,9 @@ import { getHordeStyling, preloadFonts } from "horde/styles/Styles";
 import { useEffect, useRef } from "react";
 import { StepOutcomeDataHandler } from "./StepOutcomeDataHandler";
 import { ChangeHeader, ColumnHeader, encodeStepNameFromStepNameHeader, isChangeHeader, isDateHeader, isStepOutcome, isSummary, isSummaryHeader, StepNameHeader, StepOutcomeTable, TableEntry } from "./StepOutcomeDataTypes";
-import { getCellStyle, getColorRecords, IncrementalRefreshIndicatorPanel, sharedBorderStyle } from "./StepOutcomeSharedUIComponents";
+import { getCellAccents, getCellClasses, getCellStyle, getColorRecords, IncrementalRefreshIndicatorPanel, resetCellClasses, sharedBorderStyle } from "./StepOutcomeSharedUIComponents";
 import { getHordeTheme } from "horde/styles/theme";
+import { observer } from "mobx-react-lite";
 
 // #region -- Styles --
 
@@ -270,7 +271,7 @@ function GenerateD3Table(stepOutcomeTable: StepOutcomeTable, handler: StepOutcom
                 handler.poll();
             });
 
-        // Add all columns to the header row; these will be the changelists; these control downstream widths
+        // Add all columns to the header row; these will be the changelists & the summary; these control downstream widths
         scaffold.current.theaderrow.selectAll("th.col")
             .data([...stepOutcomeTable.getColumnHeaders()] as ColumnHeader[])
             .join("th")
@@ -281,12 +282,16 @@ function GenerateD3Table(stepOutcomeTable: StepOutcomeTable, handler: StepOutcom
             .style("background-color", modeColors.content)
             .style("width", d => isSummaryHeader(d) ? "70px" : "30px")
             .style("font-size", d => (isDateHeader(d) || isSummaryHeader(d) ? "large" : "small"))
+            .style("right", d => isSummaryHeader(d) ? "0" : null)
+            .style("z-index", d => isSummaryHeader(d) ? 4 : 3)
             .style("box-sizing", "border-box")
             .style("writing-mode", "vertical-rl")
             .style("transform", "rotate(180deg)")
             .style("padding-bottom", "5px")
             .style("text-align", "right")
             .style("border", `1px solid ${modeColors.background}`)
+            .style("border", d => isSummaryHeader(d) ? `2px solid ${modeColors.content}` : `1px solid ${modeColors.background}`)
+
             .text(entry => isSummaryHeader(entry) ? entry.name : (isChangeHeader(entry) ? entry.change : (entry.date ?? "")))
             .on("mouseover", function (_event, d) {
                 const entry = d as unknown as ChangeHeader;
@@ -307,17 +312,21 @@ function GenerateD3Table(stepOutcomeTable: StepOutcomeTable, handler: StepOutcom
                     tooltipContent.style.paddingLeft = "4px";
 
                     const changeHeader = { change: entry.change ?? "-", date: entry.date ?? "-" };
+                    const original = new Date(changeHeader.date!);
+
+                    // UE-324856 - This is a temporary fix until the horde server is fixed.
+                    const corrected = new Date(original.getTime() - 4 * 60 * 60 * 1000);
 
                     const changeDiv = document.createElement("div");
                     changeDiv.innerHTML = `<strong>Change:</strong><p class="${tooltipClasses.tooltipLine}">${changeHeader.change}</p>`;
                     tooltipContent.appendChild(changeDiv);
 
                     const dateUTCDiv = document.createElement("div");
-                    dateUTCDiv.innerHTML = `<strong>Date (UTC):</strong><p class="${tooltipClasses.tooltipLine}">${new Date(changeHeader.date!).toISOString()}</p>`;
+                    dateUTCDiv.innerHTML = `<strong>Date (UTC):</strong><p class="${tooltipClasses.tooltipLine}">${corrected.toISOString()}</p>`;
                     tooltipContent.appendChild(dateUTCDiv);
 
                     const dateLocalDiv = document.createElement("div");
-                    dateLocalDiv.innerHTML = `<strong>Date (Local):</strong><p class="${tooltipClasses.tooltipLine}">${new Date(changeHeader.date!).toString()}</p>`;
+                    dateLocalDiv.innerHTML = `<strong>Date (Local):</strong><p class="${tooltipClasses.tooltipLine}">${corrected.toString()}</p>`;
                     tooltipContent.appendChild(dateLocalDiv);
 
                     attachTargettedToolTip(this as HTMLElement, scaffold.current!.tooltip, tooltipContent);
@@ -406,11 +415,23 @@ function GenerateD3Table(stepOutcomeTable: StepOutcomeTable, handler: StepOutcom
         });
 
         // Accessability row highlights.
-        rows.on("mouseover", function () {
-            d3.select(this).style("background-color", theme.semanticColors.bodyBackgroundHovered);
+        rows.on("mouseover", function (_event, _d) {
+            const row = d3.select(this);
+            // color the row (optional)
+            row.style("background-color", theme.semanticColors.bodyBackgroundHovered, "important");
+
+            // also color the sticky step-name cell in this row so it shows the hover
+            row.select("td.step-name")
+                .style("background-color", theme.semanticColors.bodyBackgroundHovered, "important");
         })
             .on("mouseout", function () {
-                d3.select(this).style("background-color", null);
+                const row = d3.select(this);
+                // remove row-level inline background
+                row.style("background-color", null);
+
+                // restore the sticky cell background to the normal mode color
+                row.select("td.step-name")
+                    .style("background-color", modeColors.content, "important");
             });
 
         // Add the first "step name" column item.
@@ -470,18 +491,22 @@ function GenerateD3Table(stepOutcomeTable: StepOutcomeTable, handler: StepOutcom
                 .style("height", "100%")
                 .style("display", "inline-block")
                 .style("text-align", "center")
+                .style("border-right", d => isSummary(d) ? `2px solid ${modeColors.background}` : null)
+                .style("border-left", d => isSummary(d) ? `2px solid ${modeColors.background}` : null)
+                .style("box-shadow", "none")
                 .each(function (entry) {
                     const div = d3.select(this);
+                    resetCellClasses().forEach(cls => div.classed(cls, false));
 
                     if (!entry) {
                         div.text("").style("background-color", "transparent").style("cursor", "default");
                         return;
                     }
-
                     if (isStepOutcome(entry)) { // StepOutcomeTableEntry
                         div.text(""),
                             div.style("cursor", "pointer")
-                                .style("background-color", getCellStyle(entry).bgColor)
+                                .style("background-color", getCellStyle(entry, stepOutcomeTable.viewModelOptions).bgColor)
+                                .style("box-shadow", getCellAccents(entry).boxShadow)
                                 .on("click", () => onCellClick?.(entry))
                                 .on("mouseover", function () {
                                     const tooltipContent = document.createElement("div");
@@ -515,9 +540,21 @@ function GenerateD3Table(stepOutcomeTable: StepOutcomeTable, handler: StepOutcom
                                 })
                                 .on("mouseout", () => {
                                     scaffold.current!.tooltip.style("opacity", 0);
-                                });
+                                })
+                        getCellClasses(entry).forEach(cls => div.classed(cls, true));
                     } else if (isSummary(entry)) { // SummaryTableEntry
-                        const ratio = entry.steps > 0 ? entry.stepsPass / entry.stepsCompleted : 0;
+                        let ratio: number = 0;
+
+                        if (entry.steps > 0) {
+                            let numerator: number = entry.stepsPass
+
+                            if (stepOutcomeTable.viewModelOptions.warningsAsSummarySuccess) {
+                                numerator += entry.stepsWarning;
+                            }
+
+                            ratio = numerator / entry.stepsCompleted;
+                        }
+
                         let ratioStr: string = `${(ratio * 100).toFixed(1)}%`;
 
                         div.style("cursor", "default")
@@ -616,7 +653,11 @@ function GenerateD3Table(stepOutcomeTable: StepOutcomeTable, handler: StepOutcom
                 .join("td")
                 .attr("class", "step-outcome-data");
 
-            dataCells.style("text-align", "center").style("border-right", `1px solid ${modeColors.content}`);
+            dataCells.style("text-align", "center")
+                .style("border-right", `1px solid ${modeColors.content}`)
+                .style("position", d => isSummary(d) ? "sticky" : null)
+                .style("right", d => isSummary(d) ? "0" : null)
+                .style("z-index", d => isSummary(d) ? 2 : 1);
 
             dataCells.selectAll("div")
                 .data(d => [d])
@@ -625,7 +666,7 @@ function GenerateD3Table(stepOutcomeTable: StepOutcomeTable, handler: StepOutcom
                 .style("height", "25px");
         });
 
-    }, [stepOutcomeTable.tableEntries, getCellStyle]);
+    }, [stepOutcomeTable.tableEntries, stepOutcomeTable.viewModelOptions, handler.lastRefreshDate, getCellStyle]);
 
     const gridStyle: React.CSSProperties = {
         display: 'grid',
@@ -644,7 +685,7 @@ function GenerateD3Table(stepOutcomeTable: StepOutcomeTable, handler: StepOutcom
                 container.scrollLeft = container.scrollWidth;
             });
         }
-    }, [handler.StepOutcomeTableData]);
+    }, [handler.StepOutcomeTableData, stepOutcomeTable.viewModelOptions]);
 
     const updateMaxHeight = () => {
         if (containerRef.current) {
@@ -724,6 +765,6 @@ function GenerateD3Table(stepOutcomeTable: StepOutcomeTable, handler: StepOutcom
  * @param onCellSelected Cell selected callback.
  * @returns React Component.
  */
-export const StepOutcomeD3Table: React.FC<{ stepOutcomeTable: StepOutcomeTable; handler: StepOutcomeDataHandler, onCellSelected: any; }> = ({ stepOutcomeTable, handler, onCellSelected }) => {
+export const StepOutcomeD3Table: React.FC<{ stepOutcomeTable: StepOutcomeTable; handler: StepOutcomeDataHandler, onCellSelected: any; }> = observer(({ stepOutcomeTable, handler, onCellSelected }) => {
     return GenerateD3Table(stepOutcomeTable, handler, onCellSelected);
-};
+});
diff --git a/Engine/Source/Programs/Horde/HordeDashboard/src/components/buildhealth/stepoutcome/StepOutcomeDataHandler.ts b/Engine/Source/Programs/Horde/HordeDashboard/src/components/buildhealth/stepoutcome/StepOutcomeDataHandler.ts
--- a/Engine/Source/Programs/Horde/HordeDashboard/src/components/buildhealth/stepoutcome/StepOutcomeDataHandler.ts
+++ b/Engine/Source/Programs/Horde/HordeDashboard/src/components/buildhealth/stepoutcome/StepOutcomeDataHandler.ts
@@ -77,6 +77,15 @@ export function prettyPrintStepOutcomeFiltersSummary(filters: StepOutcomeFilters
     return [header, subHeader, subSubHeader];
 }
 
+/**
+ * Data structure that controls the visual options for the StepOutcome
+ */
+export interface StepOutcomeViewOptions {
+    includeDateAnchors: boolean;
+    skippedAsRed: boolean;
+    warningsAsSummarySuccess: boolean;
+}
+
 /**
  * Data structure that controls the filtering & queries of the Data Handler.
  * @todo UE-309726 - This data structure can be simplified to reduce duplication, and reduce the need for the internal container to decode.
@@ -88,7 +97,6 @@ export interface StepOutcomeFilters {
     steps?: string[];
     maxJobCount?: number;
     includePreflights: boolean;
-    includeDateAnchors: boolean;
 }
 
 // #endregion -- API Data Types & Utilities --
@@ -233,7 +241,7 @@ export class StepOutcomeDataHandler extends PollBase {
         if (jobHistoryCount.size > 0) {
             // Parallelize backend calls for change summaries
             const changeResults = await Promise.all(
-                changeQueries.map(({ streamId, query }) => backend.getChangeSummaries(streamId, undefined, undefined, jobHistoryCount.get(streamId), query.filter))
+                changeQueries.map(({ streamId, query }) => backend.getChangeSummaries(streamId, undefined, undefined, [...new Set(jobHistoryCount.get(streamId))], query.filter))
             );
 
             let flattenedResults = changeResults.flat();
@@ -276,15 +284,13 @@ export class StepOutcomeDataHandler extends PollBase {
                 // Summarize each row
             }
         }
-        
+
         this.StepOutcomeTableData.summarize();
 
         console.info(`Retrieved: ${this.jobData.length} job instance results for StepOutcome query.`);
 
         this.stepOutcomeTableData.orderTableDataByChange(true);
 
-        this.StepOutcomeTableData.supportsDateAnchoredChanges = this.activeFilter.includeDateAnchors;
-
         this.setUpdated();
 
         runInAction(() => {
diff --git a/Engine/Source/Programs/Horde/HordeDashboard/src/components/buildhealth/stepoutcome/StepOutcomeDataTypes.ts b/Engine/Source/Programs/Horde/HordeDashboard/src/components/buildhealth/stepoutcome/StepOutcomeDataTypes.ts
--- a/Engine/Source/Programs/Horde/HordeDashboard/src/components/buildhealth/stepoutcome/StepOutcomeDataTypes.ts
+++ b/Engine/Source/Programs/Horde/HordeDashboard/src/components/buildhealth/stepoutcome/StepOutcomeDataTypes.ts
@@ -2,6 +2,8 @@
 
 import { ChangeSummaryData, GetStepResponse, IssueData } from "horde/backend/Api";
 import { computeDuration, DURATION_NOT_SET, formatDurationFromMs, toDate } from "./StepOutcomeUtilities";
+import { StepOutcomeViewOptions } from "./StepOutcomeDataHandler";
+import { action, makeObservable, observable } from "mobx";
 
 export const WarningStr: string = "Warnings";
 export const SuccessStr: string = "Success";
@@ -9,6 +11,7 @@ export const SkippedStr: string = "Skipped";
 export const FailureStr: string = "Failure";
 export const WaitingStr: string = "Waiting";
 export const RunningStr: string = "Running";
+export const AbortedStr: string = "Aborted";
 
 /**
  * Encodes a step outcome table entry to a stream qualified step name.
@@ -277,6 +280,9 @@ export class StepOutcomeTable {
 
     tableStreamMetadata: Map<string, number> = new Map<string, number>();
 
+    @observable
+    viewModelOptions: StepOutcomeViewOptions;
+
     /**
      * Table data structure
      */
@@ -284,6 +290,10 @@ export class StepOutcomeTable {
 
     // #endregion --  Public Members --
 
+    constructor() {
+        makeObservable(this);
+    }
+
     // #region -- Private API --
 
     private ensureRow(source: (StepOutcomeTableEntry | null)[], targetSize: number) {
@@ -384,7 +394,8 @@ export class StepOutcomeTable {
         if (insertionCol === undefined) {
             insertionCol = this.changeColHeaders.length;
             this.changeColLookup.set(tableEntry.change, insertionCol);
-            this.changeColHeaders.push({ type: "change", change: tableEntry.change, date: toDate(changeSummaryData?.dateUtc)?.toISOString() });
+            let changeDate = changeSummaryData?.dateUtc ? changeSummaryData?.dateUtc : tableEntry.jobCreateTime;
+            this.changeColHeaders.push({ type: "change", change: tableEntry.change, date: toDate(changeDate)?.toISOString() });
 
             for (const row of this.tableEntries) {
                 if (row) {
@@ -615,6 +626,16 @@ export class StepOutcomeTable {
         return resultString;
     }
 
+    /**
+     * Sets the view model options to use when interacting with the StepOutcomeTable.
+     * @param viewModelOptions The view model options to set.
+     */
+    @action
+    setViewModelOptions(viewModelOptions: StepOutcomeViewOptions) {
+        this.viewModelOptions = viewModelOptions;
+        this.supportsDateAnchoredChanges = viewModelOptions.includeDateAnchors;
+    }
+
     // #endregion -- Public API --
 }
 
diff --git a/Engine/Source/Programs/Horde/HordeDashboard/src/components/buildhealth/stepoutcome/StepOutcomeSharedUIComponents.tsx b/Engine/Source/Programs/Horde/HordeDashboard/src/components/buildhealth/stepoutcome/StepOutcomeSharedUIComponents.tsx
--- a/Engine/Source/Programs/Horde/HordeDashboard/src/components/buildhealth/stepoutcome/StepOutcomeSharedUIComponents.tsx
+++ b/Engine/Source/Programs/Horde/HordeDashboard/src/components/buildhealth/stepoutcome/StepOutcomeSharedUIComponents.tsx
@@ -3,8 +3,8 @@
 import { mergeStyleSets, Spinner, SpinnerSize } from "@fluentui/react";
 import dashboard, { StatusColor } from "horde/backend/Dashboard";
 import { observer } from "mobx-react-lite";
-import { StepOutcomeDataHandler } from "./StepOutcomeDataHandler";
-import { FailureStr, SkippedStr, StepOutcomeTableEntry, SuccessStr, WarningStr } from "./StepOutcomeDataTypes";
+import { StepOutcomeDataHandler, StepOutcomeViewOptions } from "./StepOutcomeDataHandler";
+import { AbortedStr, FailureStr, RunningStr, SkippedStr, StepOutcomeTableEntry, SuccessStr, WarningStr } from "./StepOutcomeDataTypes";
 
 // #region -- Constants --
 
@@ -24,6 +24,29 @@ export const stepOutcomeTableClasses = mergeStyleSets({
     root: {
         maxWidth: "min(85vw, 2400px)",
     },
+    redX: {
+        position: "relative",
+        overflow: "visible",
+        selectors: {
+            "&::before, &::after": {
+                content: "''",
+                position: "absolute",
+                left: "50%",
+                top: "50%",
+                width: "140%",
+                height: "2px",
+                background: "red",
+                transformOrigin: "center",
+                pointerEvents: "none",
+            },
+            "&::before": {
+                transform: "translate(-50%, -50%) rotate(45deg)",
+            },
+            "&::after": {
+                transform: "translate(-50%, -50%) rotate(-45deg)",
+            },
+        },
+    },
 });
 
 // #endregion -- Styles --
@@ -48,16 +71,64 @@ export function getColorRecords(): Record<string, string> {
     return colors;
 }
 
+/**
+ * Gets the set of conditional classes to remove from the cells.
+ * @returns The list of all conditional classes to remove from the cell.
+ */
+export function resetCellClasses(): string[] {
+    const allClasses = [
+        stepOutcomeTableClasses.redX,
+    ];
+
+    return allClasses;
+}
+
+/**
+ * Gets the classes for a step outcome table entry cell.
+ * @param entry The step outcome table entry of which to base the classes on.
+ * @returns The classes to apply to the cell.
+ */
+export function getCellClasses(entry: StepOutcomeTableEntry | null): string[] {
+    const classes: string[] = [];
+
+    if (!entry) return classes;
+
+    if (entry.stepResponse.state === AbortedStr) {
+        classes.push(stepOutcomeTableClasses.redX);
+    }
+
+    return classes;
+}
+
+/**
+ * Gets the accents for a step outcome table entry cell.
+ * @param entry The step outcome table entry of which to base the accents on.
+ * @returns The accents to apply to the cell.
+ */
+export function getCellAccents(entry: StepOutcomeTableEntry | null): { boxShadow: string } {
+    const standard = { boxShadow: "none" };
+
+    if (!entry) {
+        return standard;
+    }
+
+    if (entry.stepResponse.state == SkippedStr || entry.stepResponse.state == AbortedStr) {
+        return { boxShadow: "inset 0 0 0 4px black" };
+    }
+
+    return standard;
+}
+
 /**
  * Gets the cell style given the provided StepOutcomeTableEntry.
  * @param entry The entry to base the style off of.
+ * @param viewModelOptions The view model options to consider while applying styles.
  * @returns The matching style given the state of the entry. 
  */
-export function getCellStyle(entry: StepOutcomeTableEntry | null): { bgColor: string } {
+export function getCellStyle(entry: StepOutcomeTableEntry | null, viewModelOptions: StepOutcomeViewOptions): { bgColor: string } {
     const colors: Record<string, string> = getColorRecords();
 
     let resultColor: string = 'transparent';
-
     if (!entry) {
         return { bgColor: resultColor };
     }
@@ -69,7 +140,7 @@ export function getCellStyle(entry: StepOutcomeTableEntry | null): { bgColor: st
     // narrowly select the edge cases (e.g. outcome=success && state=completed uniquely signifies a completed run with no errors, otherwise outcome=success it could mean ready, waiting, or running).
     // Then we can fall back to the state as the truth.
     if (outcome === FailureStr) {
-        resultColor = state === SkippedStr ? colors[SkippedStr] : colors[FailureStr];
+        resultColor = (state === SkippedStr && !viewModelOptions.skippedAsRed) ? colors[SkippedStr] : colors[FailureStr];
     }
     // A step can be emitting warnings whilst still running, so we narrowly select this colouring.
     else if (outcome === WarningStr) {
diff --git a/Engine/Source/Programs/Horde/HordeDashboard/src/components/buildhealth/stepoutcome/StepOutcomeUtilities.ts b/Engine/Source/Programs/Horde/HordeDashboard/src/components/buildhealth/stepoutcome/StepOutcomeUtilities.ts
--- a/Engine/Source/Programs/Horde/HordeDashboard/src/components/buildhealth/stepoutcome/StepOutcomeUtilities.ts
+++ b/Engine/Source/Programs/Horde/HordeDashboard/src/components/buildhealth/stepoutcome/StepOutcomeUtilities.ts
@@ -1,7 +1,23 @@
 // Copyright Epic Games, Inc. All Rights Reserved.
 
+import { RunningStr, StepOutcomeTableEntry, WaitingStr } from "./StepOutcomeDataTypes";
+
 export const DURATION_NOT_SET: number = -1;
 
+/**
+ * Gets whether the step outcome table entry is in a terminal state (no more work to be done), or results are still pending.
+ * @param stepOutcomeEntry The entry to consider.
+ * @returns True if the step has run to completion, false otherwise.
+ */
+export function isTerminalState(stepOutcomeEntry: StepOutcomeTableEntry): boolean {
+    if(stepOutcomeEntry.stepResponse.state === RunningStr || stepOutcomeEntry.stepResponse.state === WaitingStr)
+    {
+        return false;
+    }
+
+    return true;
+}
+
 /**
  * Normalizes an ambiguous input (Date | string) to a Date object. 
  * @param input the input string union type.
