[Horde][BuildHealth][Dashboard] Text filtering & combo box conversion

diff --git a/Engine/Source/Programs/Horde/HordeDashboard/src/components/buildhealth/BuildhealthView.tsx b/Engine/Source/Programs/Horde/HordeDashboard/src/components/buildhealth/BuildhealthView.tsx
--- a/Engine/Source/Programs/Horde/HordeDashboard/src/components/buildhealth/BuildhealthView.tsx
+++ b/Engine/Source/Programs/Horde/HordeDashboard/src/components/buildhealth/BuildhealthView.tsx
@@ -1,9 +1,9 @@
 // Copyright Epic Games, Inc. All Rights Reserved.
 
-import { Checkbox, ComboBox, DirectionalHint, Dropdown, DropdownMenuItemType, IDropdownOption, IDropdownStyles, Label, Slider, Stack, TooltipHost } from "@fluentui/react";
+import { Checkbox, ComboBox, DirectionalHint, Dropdown, DropdownMenuItemType, IDropdownOption, IDropdownProps, ISelectableOption, Label, SearchBox, Stack, TooltipHost } from "@fluentui/react";
 import { GetProjectResponse, GetStreamResponse } from "horde/backend/Api";
 import { observer } from "mobx-react-lite";
-import { useCallback, useEffect, useMemo, useState } from "react";
+import { useCallback, useEffect, useMemo, useRef, useState } from "react";
 import { useLocation, useNavigate } from "react-router";
 import { useWindowSize } from "../../base/utilities/hooks";
 import { getHordeStyling } from "../../styles/Styles";
@@ -34,6 +34,63 @@ type BuildHealthOptionList = {
 
 // #endregion -- Dropdown Data Types --
 
+// #region -- Base Combobox Components --
+
+interface ISearchableDropdownProps extends IDropdownProps {
+    onSearchValueChanged?: (searchValue: string) => void;
+}
+
+const SearchableDropdown: React.FC<ISearchableDropdownProps> = (props) => {
+    const { onSearchValueChanged, ...dropdownProps } = props;
+
+    // Add dummy search header
+    const getOptions = (): IDropdownOption[] => [
+        { key: 'search', text: '', itemType: 2 },
+        ...dropdownProps.options,
+    ];
+
+    const onRenderOption = (
+        option?: ISelectableOption,
+        defaultRender?: (props?: ISelectableOption) => JSX.Element | null,
+    ): JSX.Element | null => {
+        if (!option) return null;
+
+        // Render SearchBox for dummy header
+        if (option.key === 'search') {
+            return (
+                <SearchBox
+                    placeholder="Filter options"
+                    autoFocus
+                    onChange={(ev, newValue) => onSearchValueChanged?.(newValue || '')}
+                    onClear={() => onSearchValueChanged?.('')}
+                    styles={{ root: { margin: '8px', width: 'calc(100% - 16px)' } }}
+                    onKeyDown={(ev) => {
+                        if (ev.key === 'Enter' || ev.key === 'Escape') ev.stopPropagation();
+                    }}
+                />
+            );
+        }
+
+        // Use custom or default rendering for options
+        return dropdownProps.onRenderOption ? dropdownProps.onRenderOption(option, defaultRender) : defaultRender?.(option) ?? null;
+    };
+
+    return (
+        <Dropdown
+            {...dropdownProps}
+            options={getOptions()}
+            onRenderOption={onRenderOption}
+            calloutProps={{
+                directionalHint: DirectionalHint.rightCenter,
+                styles: { calloutMain: { width: 500, minWidth: 500 } },
+            }}
+        />
+    );
+};
+
+
+// #endregion -- Base Combobox Components --
+
 // #region -- Base Dropdown Components --
 
 /**
@@ -81,9 +138,9 @@ const BasicListParameter: React.FC<{ handler: BuildHealthDataHandler, buildHealt
  * React Component that is a multi select drop down.
  * @returns React Component.
  */
-const MultiListParameter: React.FC<{ handler: BuildHealthDataHandler, buildHealthOptions: BuildHealthOptionsController, param: BuildHealthOptionList, disabled?: boolean, selectedKeys: string[], request: DataHandlerRefreshRequest, onChange: (option: IDropdownOption<BuildHealthOptionData>) => void }> = function ConstructMultilistParamter({ param, buildHealthOptions, selectedKeys, request, onChange }) {
+const MultiListParameter: React.FC<{ handler: BuildHealthDataHandler, buildHealthOptions: BuildHealthOptionsController, param: BuildHealthOptionList, disabled?: boolean, selectedKeys: string[], request: DataHandlerRefreshRequest, onChange: (option: IDropdownOption<BuildHealthOptionData>) => void }> = function ConstructMultilistParamter({ param, buildHealthOptions, disabled, selectedKeys, request, onChange }) {
+    const [filterText, setFilterText] = useState("");
     const changeVersion = buildHealthOptions.optionsChangeVersion;
-    const { modeColors } = getHordeStyling();
     const key = `${PARAMETER_KEY_PREFIX}${KEY_SEPARATOR}${param.label}`;
     const gset: Set<string> = new Set();
 
@@ -121,120 +178,89 @@ const MultiListParameter: React.FC<{ handler: BuildHealthDataHandler, buildHealt
                         key: item.id,
                         text: item.text,
                         data: item,
+                        selected: selectedKeys.includes(item.id),
                     });
                 }
             });
         });
 
         return options;
-    }, [param.items]);
-
-    const memoizedSelectedKeys = useMemo(() => selectedKeys, [selectedKeys]);
+    }, [param.items, selectedKeys]);
 
     // #endregion -- Data Memos --
 
-    // #region -- Style Memos --
-
-    const dropdownStyles: Partial<IDropdownStyles> = useMemo(() => ({
-        dropdown: { width: 300 },
-        callout: {
-            selectors: {
-                [".ms-Callout-main"]: {
-                    padding: "4px 4px 12px 12px",
-                    overflow: "auto" as "auto",
-                },
-            },
-        },
-        dropdownItemHeader: { fontSize: 12, color: modeColors.text },
-        dropdownOptionText: { fontSize: 12 },
-        dropdownItem: {
-            minHeight: 18,
-            lineHeight: 18,
-            selectors: {
-                [".ms-Checkbox-checkbox"]: {
-                    width: 14,
-                    height: 14,
-                    fontSize: 11,
-                },
-            },
-        },
-        dropdownItemSelected: {
-            minHeight: 18,
-            lineHeight: 18,
-            backgroundColor: "inherit",
-            selectors: {
-                [".ms-Checkbox-checkbox"]: {
-                    width: 14,
-                    height: 14,
-                    fontSize: 11,
-                },
-            },
-        },
-    }), [modeColors.text]);
-
-    // #endregion -- Style Memos --
+    const filteredOptions = doptions.filter((option) => {
+        if (option.itemType === DropdownMenuItemType.Header) return true;
+        if (!filterText) return true;
+        return option.text.toLowerCase().includes(filterText.toLowerCase());
+    });
 
     // #region -- Dropdown Callbacks --
 
-    const renderItem = useCallback((item, defaultRender) => {
-        let castedOption = item as IDropdownOption<BuildHealthOptionData>;
-        if (!castedOption || !defaultRender) return null;
+    const onRenderOption = useCallback(
+        (option?: ISelectableOption) => {
+            if (!option) return null;
 
-        if (!castedOption.data) {
-            return defaultRender(item);
-        }
-        return (
-            <TooltipHost key={`tooltip_${castedOption.text}`} content={
-                <div>
-                    <Stack>
-                        <Label>{castedOption.data?.group}</Label>
-                    </Stack>
-                    <Stack>
-                        {castedOption.text}
-                    </Stack>
-                </div>
+            if (option.itemType === DropdownMenuItemType.Header) {
+                return <div style={{ fontWeight: "bold", padding: "4px 8px" }}>{option.text}</div>;
             }
-                directionalHint={DirectionalHint.rightCenter}
-                calloutProps={{
-                    styles: {
-                        root: { maxWidth: 600 }
-                    }
-                }}
-            >
-                {defaultRender(item)}
-            </TooltipHost>
-        );
-    }, []);
-
-    const handleChange = useCallback((_ev, option) => {
-        const castedOption = option as IDropdownOption<BuildHealthOptionData>;
-        if (!castedOption) return;
-        onChange(castedOption);
+
+            return (
+                <TooltipHost key={`tooltip_${option.text}`} content={
+                    <div>
+                        <Stack>
+                            <Label>{(option as IDropdownOption<BuildHealthOptionData>).data?.group}</Label>
+                        </Stack>
+                        <Stack>
+                            {option.text}
+                        </Stack>
+                    </div>
+                }
+                    directionalHint={DirectionalHint.rightCenter}
+                    calloutProps={{
+                        styles: {
+                            root: { maxWidth: 800 }
+                        }
+                    }}
+                >
+                    <div style={{ padding: "4px 8px" }}>{option.text}</div>
+                </TooltipHost>
+            );
+        }, []);
+
+    const handleDropdownChange = useCallback((_ev, option) => {
+        if (option && option.itemType !== DropdownMenuItemType.Header) {
+            onChange(option);
+        }
+
     }, [onChange]);
 
-    const handleDismiss = useCallback(() => {
+    const handleMenuDismiss = useCallback(() => {
+        setFilterText("");
         handler.requestHierarchicalRefresh(request);
         buildHealthOptions.synchronizeDerivedKeys();
-    }, [handler, request, buildHealthOptions]);
+    }, [request, buildHealthOptions]);
 
     // #endregion -- Dropdown Callbacks --
 
     return (
-        <Dropdown
+        <SearchableDropdown
             id={changeVersion.toString() + key}
             key={key}
-            label={param.label}
-            placeholder="Select options"
-            options={doptions}
-            selectedKeys={memoizedSelectedKeys}
-            dropdownWidth={400}
-            calloutProps={{ directionalHint: DirectionalHint.rightCenter }}
-            styles={dropdownStyles}
-            onRenderItem={renderItem}
-            onChange={handleChange}
-            onDismiss={handleDismiss}
+            placeholder={param.label}
+            options={filteredOptions}
             multiSelect
-        />);
+            selectedKeys={selectedKeys}
+            onChange={handleDropdownChange}
+            onDismiss={handleMenuDismiss}
+            onSearchValueChanged={setFilterText}
+            onRenderOption={onRenderOption}
+            styles={{
+                dropdown: { width: 300, maxWidth: 900 },
+            }}
+            disabled={disabled}
+        />
+    );
 };
 
 // #endregion -- Base Dropdown Components --
@@ -243,7 +269,7 @@ const MultiListParameter: React.FC<{ handler: BuildHealthDataHandler, buildHealt
 
 const ProjectDropdownSingle: React.FC<{ handler: BuildHealthDataHandler, buildHealthOptions: BuildHealthOptionsController }> = observer(function ConstructProjectDropdownSingle({ handler, buildHealthOptions }) {
     const projectDropdownParams: BuildHealthOptionList = {
-        label: "Select Project",
+        label: "",
         items: [
         ],
         toolTip: "Select an option from the dropdown."
@@ -264,6 +290,7 @@ const ProjectDropdownSingle: React.FC<{ handler: BuildHealthDataHandler, buildHe
 
     return (
         <Stack>
+            <Label>Project</Label>
             <BasicListParameter handler={handler} buildHealthOptions={buildHealthOptions} param={projectDropdownParams} selectedKey={selectedProjectKey} request={request} onChange={(option: IDropdownOption<BuildHealthOptionData>) => {
                 const key = option.key as string;
                 buildHealthOptions.toggleSingleProject(key, option.text);
@@ -272,68 +299,6 @@ const ProjectDropdownSingle: React.FC<{ handler: BuildHealthDataHandler, buildHe
     );
 });
 
-const ProjectDropdown: React.FC<{ handler: BuildHealthDataHandler, buildHealthOptions: BuildHealthOptionsController }> = observer(function ConstructProjectDropdown({ handler, buildHealthOptions }) {
-    const projectDropdownParams: BuildHealthOptionList = {
-        label: "Select Project",
-        items: [
-        ],
-        toolTip: "Select an option from the dropdown."
-    };
-
-    projectDropdownParams.items = handler.projectsData.map((project: GetProjectResponse) => {
-        let listParam: BuildHealthOptionData = {
-            id: `${PARAMETER_KEY_PREFIX}${KEY_SEPARATOR}${projectDropdownParams.label}${KEY_SEPARATOR}${project.id}`,
-            text: project.name,
-        };
-
-        return listParam;
-    }).sort((a, b) => a.text.localeCompare(b.text));
-
-    let selectedProjectKeys: string[] = Object.keys(buildHealthOptions.state.enabledProjects);
-    let request: DataHandlerRefreshRequest = expandRefreshRequest({ streams: true });
-
-    return (
-        <Stack>
-            <MultiListParameter handler={handler} buildHealthOptions={buildHealthOptions} param={projectDropdownParams} selectedKeys={selectedProjectKeys} request={request} onChange={(option: IDropdownOption<BuildHealthOptionData>) => {
-                const key = option.key as string;
-                const isSelected = !!option.selected;
-                buildHealthOptions.toggleProject(key, option.text, isSelected);
-            }} />
-        </Stack>
-    );
-});
-
-const StreamDropdownSingle: React.FC<{ handler: BuildHealthDataHandler, buildHealthOptions: BuildHealthOptionsController }> = observer(function ConstructProjectDropdownSingle({ handler, buildHealthOptions }) {
-    const streamDropdownParams: BuildHealthOptionList = {
-        label: "Select Stream",
-        items: [
-        ],
-        toolTip: "Select an option from the dropdown."
-    };
-
-    streamDropdownParams.items = handler.streamsData.map((stream: GetStreamResponse) => {
-        let listParam: BuildHealthOptionData = {
-            id: stream.id,
-            text: stream.name,
-            group: stream.projectId,
-        };
-        return listParam;
-    }).sort((a, b) => a.text.localeCompare(b.text));
-
-    let selectedStreamKey: string = Object.keys(buildHealthOptions.state.enabledStreams)[0]!;
-    let request: DataHandlerRefreshRequest = expandRefreshRequest({ streams: true });
-    {
-        return streamDropdownParams.items.length > 0 && (
-            <Stack >
-                <BasicListParameter handler={handler} buildHealthOptions={buildHealthOptions} param={streamDropdownParams} selectedKey={selectedStreamKey} request={request} onChange={(option: IDropdownOption<BuildHealthOptionData>) => {
-                    const key = option.key as string;
-                    buildHealthOptions.toggleSingleStream(key, option.text);
-                }} />
-            </Stack>
-        );
-    }
-});
-
 const StreamDropdown: React.FC<{ handler: BuildHealthDataHandler, buildHealthOptions: BuildHealthOptionsController }> = observer(function ConstructProjectDropdown({ handler, buildHealthOptions }) {
     const streamDropdownParams: BuildHealthOptionList = {
         label: "Select Stream(s)",
@@ -353,18 +318,18 @@ const StreamDropdown: React.FC<{ handler: BuildHealthDataHandler, buildHealthOpt
 
     let selectedStreamKeys: string[] = Object.keys(buildHealthOptions.state.enabledStreams);
     let request: DataHandlerRefreshRequest = expandRefreshRequest({ streams: true });
-    {
-        return streamDropdownParams.items.length > 0 && (
-            <Stack>
-                <MultiListParameter handler={handler} buildHealthOptions={buildHealthOptions} param={streamDropdownParams} selectedKeys={selectedStreamKeys} request={request} onChange={(option: IDropdownOption<BuildHealthOptionData>) => {
-                    const key = option.key as string;
-                    const isSelected = !!option.selected;
-
-                    buildHealthOptions.toggleStream(key, option.text, isSelected);
-                }} />
-            </Stack>
-        );
-    }
+
+    return streamDropdownParams.items.length > 0 && (
+        <Stack>
+            <Label>Stream(s)</Label>
+            <MultiListParameter handler={handler} buildHealthOptions={buildHealthOptions} param={streamDropdownParams} selectedKeys={selectedStreamKeys} request={request} onChange={(option: IDropdownOption<BuildHealthOptionData>) => {
+                const key = option.key as string;
+                const isSelected = !!option.selected;
+
+                buildHealthOptions.toggleStream(key, option.text, isSelected);
+            }} />
+        </Stack>
+    );
 });
 
 const TemplateDropdown: React.FC<{ handler: BuildHealthDataHandler, buildHealthOptions: BuildHealthOptionsController }> = observer(function ConstructProjectDropdown({ handler, buildHealthOptions }) {
@@ -386,18 +351,18 @@ const TemplateDropdown: React.FC<{ handler: BuildHealthDataHandler, buildHealthO
 
     let selectedTemplateKeys: string[] = Object.keys(buildHealthOptions.state.enabledTemplates);
     let request: DataHandlerRefreshRequest = expandRefreshRequest({ jobs: true });
-    {
-        return templateDropdownParams.items.length > 0 && (
-            <Stack id={`job_dropdown`}>
-                <MultiListParameter handler={handler} buildHealthOptions={buildHealthOptions} param={templateDropdownParams} selectedKeys={selectedTemplateKeys} request={request} onChange={(option: IDropdownOption<BuildHealthOptionData>) => {
-                    const key = option.key as string;
-                    const isSelected = !!option.selected;
-
-                    buildHealthOptions.toggleTemplate(key, option.text, isSelected);
-                }} />
-            </Stack>
-        );
-    }
+
+    return templateDropdownParams.items.length > 0 && (
+        <Stack id={`job_dropdown`}>
+            <Label>Template(s)</Label>
+            <MultiListParameter handler={handler} buildHealthOptions={buildHealthOptions} param={templateDropdownParams} selectedKeys={selectedTemplateKeys} request={request} onChange={(option: IDropdownOption<BuildHealthOptionData>) => {
+                const key = option.key as string;
+                const isSelected = !!option.selected;
+
+                buildHealthOptions.toggleTemplate(key, option.text, isSelected);
+            }} />
+        </Stack>
+    );
 });
 
 const StepDropdown: React.FC<{ handler: BuildHealthDataHandler, buildHealthOptions: BuildHealthOptionsController }> = observer(function ConstructProjectDropdown({ handler, buildHealthOptions }) {
@@ -419,18 +384,18 @@ const StepDropdown: React.FC<{ handler: BuildHealthDataHandler, buildHealthOptio
 
     let selectedStepKeys: string[] = Object.keys(buildHealthOptions.state.enabledSteps);
     let request: DataHandlerRefreshRequest = { steps: true };
-    {
-        return stepDropdownParams.items.length > 0 && (
-            <Stack id={`step_dropdown`}>
-                <MultiListParameter handler={handler} buildHealthOptions={buildHealthOptions} param={stepDropdownParams} selectedKeys={selectedStepKeys} request={request} onChange={(option: IDropdownOption<BuildHealthOptionData>) => {
-                    const key = option.key as string;
-                    const isSelected = !!option.selected;
-
-                    buildHealthOptions.toggleStep(key, option.text, isSelected);
-                }} />
-            </Stack>
-        );
-    }
+
+    return stepDropdownParams.items.length > 0 && (
+        <Stack id={`step_dropdown`}>
+            <Label>Step(s)</Label>
+            <MultiListParameter handler={handler} buildHealthOptions={buildHealthOptions} param={stepDropdownParams} selectedKeys={selectedStepKeys} request={request} onChange={(option: IDropdownOption<BuildHealthOptionData>) => {
+                const key = option.key as string;
+                const isSelected = !!option.selected;
+
+                buildHealthOptions.toggleStep(key, option.text, isSelected);
+            }} />
+        </Stack>
+    );
 });
 
 // #endregion -- Option Specific Dropdown Components --
@@ -446,18 +411,21 @@ const StepDropdown: React.FC<{ handler: BuildHealthDataHandler, buildHealthOptio
 const BuildHealthUrlSync: React.FC<{ buildHealthController: BuildHealthOptionsController }> = function ConstrucBuildHealthSidebar({ buildHealthController }) {
     const location = useLocation();
     const navigate = useNavigate();
+    const controllerRef = useRef(buildHealthController);
 
     useEffect(() => {
         const initializeFromUrl = async () => {
             const params = parseBuildHealthQueryParams(location);
-            await loadBuildHealthOptionsFromParams(params, buildHealthController, handler);
+            await loadBuildHealthOptionsFromParams(params, controllerRef.current, handler);
         };
 
         initializeFromUrl();
-    }, []);
+    },
+    // eslint-disable-next-line react-hooks/exhaustive-deps -- location is a special global in this case.
+    []);
 
     useEffect(() => {
-        const desiredSearch = buildHealthController.toNavigationQuery().toString();
+        const desiredSearch = controllerRef.current.toNavigationQuery().toString();
         if (desiredSearch && location.search !== `?${desiredSearch}`) {
             navigate(`${location.pathname}?${desiredSearch}`, { replace: true });
         }
@@ -537,7 +505,7 @@ const BuildHealthSidebar: React.FC<{ handler: BuildHealthDataHandler, buildHealt
 
     // #endregion -- Property Use Effects --
 
-    let timeComboWidth = 180;
+    let timeComboWidth = 150;
 
     return (
         <Stack style={{ minWidth: 300, maxWidth: 450, paddingRight: 18 }}>
@@ -673,7 +641,9 @@ export const BuildHealthView: React.FC = observer(function ConstructBuildHealthV
     const jobHistorySpan = useMemo(() => ({
         start: buildHealthState.startDate,
         end: buildHealthState.endDate
-    }), [buildHealthState.startDate, buildHealthState.endDate]);
+    }), 
+    // eslint-disable-next-line react-hooks/exhaustive-deps -- All necessary dependencies included for MobX observables
+    [buildHealthState.startDate, buildHealthState.endDate]);
 
     const filter = useMemo(() => ({
         streams: buildHealthState.stepOutcomeEnabledStreamKeys,
@@ -681,7 +651,9 @@ export const BuildHealthView: React.FC = observer(function ConstructBuildHealthV
         steps: buildHealthState.stepOutcomeEnabledStepKeys,
         jobHistorySpan,
         includePreflights: buildHealthState.includePreflight,
-    }), [
+    }), 
+    // eslint-disable-next-line react-hooks/exhaustive-deps -- All necessary dependencies included for MobX observables
+    [
         buildHealthState.stepOutcomeEnabledStreamKeys,
         buildHealthState.stepOutcomeEnabledJobKeys,
         buildHealthState.stepOutcomeEnabledStepKeys,
@@ -693,7 +665,9 @@ export const BuildHealthView: React.FC = observer(function ConstructBuildHealthV
         includeDateAnchors: buildHealthUIState.includeDateAnchors,
         skippedAsRed: buildHealthUIState.skippedIsFailure,
         warningsAsSummarySuccess: buildHealthUIState.warningsAsSummarySuccess
-    }), [
+    }), 
+    // eslint-disable-next-line react-hooks/exhaustive-deps -- All necessary dependencies included for MobX observables
+    [
         buildHealthUIState.includeDateAnchors,
         buildHealthUIState.skippedIsFailure,
         buildHealthUIState.warningsAsSummarySuccess
