Horde: Updating schema documentation

#rnx

diff --git a/Engine/Source/Programs/Horde/Docs/Deployment/AgentSettings.md b/Engine/Source/Programs/Horde/Docs/Deployment/AgentSettings.md
--- a/Engine/Source/Programs/Horde/Docs/Deployment/AgentSettings.md
+++ b/Engine/Source/Programs/Horde/Docs/Deployment/AgentSettings.md
@@ -32,6 +32,7 @@ Name | Description
 `properties` | `string` `->` `string`<br>Key/value properties in addition to those set internally by the agent
 `adminEndpoints` | `string[]`<br>Listen addresses for the built-in HTTP admin/management server. Disabled when empty. If activated, it's recommended to bind only to localhost for security reasons. Example: localhost:7008 to listen on localhost, port 7008
 `healthCheckEndpoints` | `string[]`<br>Listen addresses for the built-in HTTP health check server. Disabled when empty. If activated, it's recommended to bind only to localhost for security reasons. Example: *:7009 to listen on all interfaces/IPs, port 7009 If all interfaces are bound with *, make sure to run process as administrator.
+`tempStorageMaxBatchSize` | `integer`<br>Maximum number of temp storage exports to write in a single request
 
 ## ServerProfile
 
