Horde: Reduce cookie size by removing tokens for OIDC and Okta auth

diff --git a/Engine/Source/Programs/Horde/HordeServer/Authentication/OidcAuthHandler.cs b/Engine/Source/Programs/Horde/HordeServer/Authentication/OidcAuthHandler.cs
--- a/Engine/Source/Programs/Horde/HordeServer/Authentication/OidcAuthHandler.cs
+++ b/Engine/Source/Programs/Horde/HordeServer/Authentication/OidcAuthHandler.cs
@@ -131,7 +131,7 @@ namespace HordeServer.Authentication
 		{
 			options.ResponseType = OpenIdConnectResponseType.Code;
 			options.GetClaimsFromUserInfoEndpoint = true;
-			options.SaveTokens = true;
+			options.SaveTokens = false;
 			options.TokenValidationParameters.NameClaimType = "name";
 			options.ClaimActions.Add(new MapRolesClaimAction(settings));
 			options.RequireHttpsMetadata = !settings.OidcDebugMode;
diff --git a/Engine/Source/Programs/Horde/HordeServer/Authentication/OktaAuthHandler.cs b/Engine/Source/Programs/Horde/HordeServer/Authentication/OktaAuthHandler.cs
--- a/Engine/Source/Programs/Horde/HordeServer/Authentication/OktaAuthHandler.cs
+++ b/Engine/Source/Programs/Horde/HordeServer/Authentication/OktaAuthHandler.cs
@@ -133,7 +133,7 @@ namespace HordeServer.Authentication
 			options.Scope.Add("email");
 			options.ResponseType = OpenIdConnectResponseType.Code;
 			options.GetClaimsFromUserInfoEndpoint = true;
-			options.SaveTokens = true;
+			options.SaveTokens = false;
 			options.TokenValidationParameters.NameClaimType = "name";
 			options.ClaimActions.Add(new MapRolesClaimAction(settings));
 
