[Backout] - CL45881300
#rnx
Original CL Desc
-----------------------------------------------------------------
Added feature to Horde's issues to be grouped based on customizable rules

diff --git a/Engine/Source/Programs/Horde/Plugins/Build/HordeServer.Build.Tests/Issues/IssueGroupingRulesTests.cs b/Engine/Source/Programs/Horde/Plugins/Build/HordeServer.Build.Tests/Issues/IssueGroupingRulesTests.cs
deleted file mode 100644
diff --git a/Engine/Source/Programs/Horde/Plugins/Build/HordeServer.Build/Issues/IIssueCollection.cs b/Engine/Source/Programs/Horde/Plugins/Build/HordeServer.Build/Issues/IIssueCollection.cs
--- a/Engine/Source/Programs/Horde/Plugins/Build/HordeServer.Build/Issues/IIssueCollection.cs
+++ b/Engine/Source/Programs/Horde/Plugins/Build/HordeServer.Build/Issues/IIssueCollection.cs
@@ -363,22 +363,17 @@ namespace HordeServer.Issues
 		/// <inheritdoc/>
 		public string ChangeFilter { get; set; }
 
-		/// <inheritdoc/>
-		public IIssueGroupingRules GroupingRules { get; set; }
-
 		/// <summary>
 		/// Constructor
 		/// </summary>
 		/// <param name="type">The type of issue</param>
 		/// <param name="summaryTemplate">Template for the summary string to display for the issue</param>
 		/// <param name="changeFilter">Filter for changes covered by this issue</param>
-		/// <param name="groupingRules">Grouping rules used by this issue</param>
-		public NewIssueFingerprint(string type, string summaryTemplate, string changeFilter, IIssueGroupingRules groupingRules)
+		public NewIssueFingerprint(string type, string summaryTemplate, string changeFilter)
 		{
 			Type = type;
 			SummaryTemplate = summaryTemplate;
 			ChangeFilter = changeFilter;
-			GroupingRules = groupingRules;
 		}
 
 		/// <summary>
@@ -390,8 +385,7 @@ namespace HordeServer.Issues
 		/// <param name="rejectKeys">Keys which should not match with this issue</param>
 		/// <param name="metadata">Additional metadata added by the issue handler</param>
 		/// <param name="changeFilter">Filter for changes covered by this issue</param>
-		/// <param name="groupingRules">Grouping rules used by this issue</param>
-		public NewIssueFingerprint(string type, string summaryTemplate, IEnumerable<IssueKey> keys, IEnumerable<IssueKey>? rejectKeys, IEnumerable<IssueMetadata>? metadata, string changeFilter, IIssueGroupingRules groupingRules)
+		public NewIssueFingerprint(string type, string summaryTemplate, IEnumerable<IssueKey> keys, IEnumerable<IssueKey>? rejectKeys, IEnumerable<IssueMetadata>? metadata, string changeFilter)
 		{
 			Type = type;
 			SummaryTemplate = summaryTemplate;
@@ -407,8 +401,6 @@ namespace HordeServer.Issues
 			}
 
 			ChangeFilter = changeFilter;
-
-			GroupingRules = groupingRules;
 		}
 
 		/// <summary>
@@ -416,7 +408,7 @@ namespace HordeServer.Issues
 		/// </summary>
 		/// <param name="other">The fingerprint to copy from</param>
 		public NewIssueFingerprint(IIssueFingerprint other)
-			: this(other.Type, other.SummaryTemplate, other.Keys, other.RejectKeys, other.Metadata, other.ChangeFilter, other.GroupingRules)
+			: this(other.Type, other.SummaryTemplate, other.Keys, other.RejectKeys, other.Metadata, other.ChangeFilter)
 		{
 		}
 
diff --git a/Engine/Source/Programs/Horde/Plugins/Build/HordeServer.Build/Issues/IssueCollection.cs b/Engine/Source/Programs/Horde/Plugins/Build/HordeServer.Build/Issues/IssueCollection.cs
--- a/Engine/Source/Programs/Horde/Plugins/Build/HordeServer.Build/Issues/IssueCollection.cs
+++ b/Engine/Source/Programs/Horde/Plugins/Build/HordeServer.Build/Issues/IssueCollection.cs
@@ -29,7 +29,6 @@ using Microsoft.Extensions.Options;
 using MongoDB.Bson;
 using MongoDB.Bson.Serialization;
 using MongoDB.Bson.Serialization.Attributes;
-using MongoDB.Bson.Serialization.Serializers;
 using MongoDB.Driver;
 using OpenTelemetry.Trace;
 
@@ -262,7 +261,7 @@ namespace HordeServer.Issues
 			}
 		}
 
-		internal class IssueFingerprintDocument : IIssueFingerprint
+		class IssueFingerprintDocument : IIssueFingerprint
 		{
 			public string Type { get; set; }
 
@@ -290,9 +289,6 @@ namespace HordeServer.Issues
 			[BsonElement("fl2")]
 			public string? ChangeFilter { get; set; }
 
-			[BsonElement("gr")]
-			public IIssueGroupingRules GroupingRules { get; set; } = DefaultIssueGroupingRules.Instance;
-
 			[Obsolete("Use ChangeFilter instead")]
 			[BsonElement("flt"), BsonIgnoreIfNull]
 			public List<string>? LegacyChangeFilter
@@ -307,7 +303,6 @@ namespace HordeServer.Issues
 			private IssueFingerprintDocument()
 			{
 				Type = String.Empty;
-				GroupingRules = DefaultIssueGroupingRules.Instance;
 			}
 
 			public IssueFingerprintDocument(IIssueFingerprint fingerprint)
@@ -326,8 +321,6 @@ namespace HordeServer.Issues
 				}
 
 				ChangeFilter = fingerprint.ChangeFilter;
-
-				GroupingRules = fingerprint.GroupingRules;
 			}
 
 			#region Legacy
@@ -694,28 +687,6 @@ namespace HordeServer.Issues
 				cm.MapProperty(x => x.Key).SetElementName("k");
 				cm.MapProperty(x => x.Value).SetElementName("v");
 			});
-
-			Func<Type, bool> issueGroupingRulesSerializer = delegate (Type type)
-			{
-				if (type == typeof(DefaultIssueGroupingRules) || type == typeof(CustomizableIssueGroupingRules))
-				{
-					return true;
-				}
-				return false;
-			};
-			BsonSerializer.RegisterSerializer(new ObjectSerializer(issueGroupingRulesSerializer));
-
-			BsonClassMap.RegisterClassMap<IssueFingerprintDocument>(cm =>
-			{
-				cm.AutoMap();
-				cm.GetMemberMap(c => c.GroupingRules).SetIgnoreIfDefault(true).SetDefaultValue(DefaultIssueGroupingRules.Instance);
-			});
-
-			BsonClassMap.RegisterClassMap<CustomizableIssueGroupingRules>(cm =>
-			{
-				cm.AutoMap();
-				cm.SetDiscriminator("c");
-			});
 		}
 
 		public IssueCollection(IMongoService mongoService, IRedisService redisService, IUserCollection userCollection, ICommitService commitService, IAuditLogFactory<int> auditLogFactory, ITelemetryWriter telemetryWriter, IOptionsMonitor<BuildConfig> buildConfig, Tracer tracer, ILogger<IssueCollection> logger)
diff --git a/Engine/Source/Programs/Horde/Plugins/Build/HordeServer.Build/Issues/IssueService.cs b/Engine/Source/Programs/Horde/Plugins/Build/HordeServer.Build/Issues/IssueService.cs
--- a/Engine/Source/Programs/Horde/Plugins/Build/HordeServer.Build/Issues/IssueService.cs
+++ b/Engine/Source/Programs/Horde/Plugins/Build/HordeServer.Build/Issues/IssueService.cs
@@ -823,7 +823,7 @@ namespace HordeServer.Issues
 			HashSet<IssueEventGroupInternal> eventGroups = new HashSet<IssueEventGroupInternal>();
 			foreach (IssueEventGroup issue in issues)
 			{
-				NewIssueFingerprint fingerprint = new NewIssueFingerprint(issue.Type, issue.SummaryTemplate, issue.ChangeFilter, issue.GroupingRules);
+				NewIssueFingerprint fingerprint = new NewIssueFingerprint(issue.Type, issue.SummaryTemplate, issue.ChangeFilter);
 				fingerprint.Keys.UnionWith(issue.Keys);
 				fingerprint.Metadata.UnionWith(issue.Metadata);
 
diff --git a/Engine/Source/Programs/Shared/EpicGames.Horde/Issues/Handlers/ExternalIssueHandler.cs b/Engine/Source/Programs/Shared/EpicGames.Horde/Issues/Handlers/ExternalIssueHandler.cs
--- a/Engine/Source/Programs/Shared/EpicGames.Horde/Issues/Handlers/ExternalIssueHandler.cs
+++ b/Engine/Source/Programs/Shared/EpicGames.Horde/Issues/Handlers/ExternalIssueHandler.cs
@@ -17,6 +17,7 @@ namespace EpicGames.Horde.Issues.Handlers
 	public class ExternalIssueHandler : IssueHandler
 	{
 		static readonly JsonNodeOptions s_jsonNodeOptions = new JsonNodeOptions { PropertyNameCaseInsensitive = true };
+		static readonly JsonSerializerOptions? s_jsonSerializerOptions = new JsonSerializerOptions { PropertyNameCaseInsensitive = true };
 
 		readonly List<IssueEventGroup> _issues = new List<IssueEventGroup>();
 
@@ -65,7 +66,7 @@ namespace EpicGames.Horde.Issues.Handlers
 					JsonObject? issue = properties["$issue"] as JsonObject;
 					if (issue != null)
 					{
-						IssueFingerprint? fingerprint = JsonSerializer.Deserialize<IssueFingerprint>(issue, IssueFingerprint.JsonSerializerOptions);
+						IssueFingerprint? fingerprint = JsonSerializer.Deserialize<IssueFingerprint>(issue, s_jsonSerializerOptions);
 						if (fingerprint != null)
 						{
 							return fingerprint;
diff --git a/Engine/Source/Programs/Shared/EpicGames.Horde/Issues/Handlers/LocalizationIssueHandler.cs b/Engine/Source/Programs/Shared/EpicGames.Horde/Issues/Handlers/LocalizationIssueHandler.cs
--- a/Engine/Source/Programs/Shared/EpicGames.Horde/Issues/Handlers/LocalizationIssueHandler.cs
+++ b/Engine/Source/Programs/Shared/EpicGames.Horde/Issues/Handlers/LocalizationIssueHandler.cs
@@ -42,21 +42,6 @@ namespace EpicGames.Horde.Issues.Handlers
 		/// <inheritdoc/>
 		public override int Priority => 10;
 
-
-		private CustomizableIssueGroupingRules? _localizationIssueGroupingRules = null;
-		private IIssueGroupingRules LocalizationIssueGroupingRules
-		{
-			get
-			{
-				if(_localizationIssueGroupingRules == null)
-				{
-					_localizationIssueGroupingRules = new CustomizableIssueGroupingRules();
-					_localizationIssueGroupingRules.SpanRule = CustomizableIssueGroupingRules.FingerprintMatchingRule.FullFingerprint;
-				}
-				return _localizationIssueGroupingRules;
-			}
-		}
-
 		/// <summary>
 		/// Constructor
 		/// </summary>
@@ -282,8 +267,6 @@ namespace EpicGames.Horde.Issues.Handlers
 						issueGroup = _issuesWithoutClearSuspects;
 					}
 
-					issueGroup.GroupingRules = LocalizationIssueGroupingRules;
-
 					// Add current event to the chosen issueGroup
 					issueGroup.Events.Add(issueEvent);
 
diff --git a/Engine/Source/Programs/Shared/EpicGames.Horde/Issues/IIssueFingerprint.cs b/Engine/Source/Programs/Shared/EpicGames.Horde/Issues/IIssueFingerprint.cs
--- a/Engine/Source/Programs/Shared/EpicGames.Horde/Issues/IIssueFingerprint.cs
+++ b/Engine/Source/Programs/Shared/EpicGames.Horde/Issues/IIssueFingerprint.cs
@@ -40,11 +40,6 @@ namespace EpicGames.Horde.Issues
 		/// Filter for changes that should be included in this issue
 		/// </summary>
 		public string ChangeFilter { get; }
-
-		/// <summary>
-		/// Issue grouping rules
-		/// </summary>
-		public IIssueGroupingRules GroupingRules { get; }
 	}
 
 	/// <summary>
@@ -60,11 +55,33 @@ namespace EpicGames.Horde.Issues
 		/// <returns>True is the fingerprints match</returns>
 		public static bool IsMatch(this IIssueFingerprint fingerprint, IIssueFingerprint other)
 		{
-			if (fingerprint.GroupingRules != null)
+			if (!fingerprint.Type.Equals(other.Type, StringComparison.Ordinal))
+			{
+				return false;
+			}
+			if (!fingerprint.SummaryTemplate.Equals(other.SummaryTemplate, StringComparison.Ordinal))
+			{
+				return false;
+			}
+			if (fingerprint.Keys.Count == 0)
 			{
-				return fingerprint.GroupingRules.IsMatch(fingerprint, other);
+				if (other.Keys.Count > 0)
+				{
+					return false;
+				}
 			}
-			return DefaultIssueGroupingRules.Instance.IsMatch(fingerprint, other);
+			else
+			{
+				if (!fingerprint.Keys.Any(x => other.Keys.Contains(x)))
+				{
+					return false;
+				}
+			}
+			if (fingerprint.RejectKeys != null && fingerprint.RejectKeys.Any(x => other.Keys.Contains(x)))
+			{
+				return false;
+			}
+			return true;
 		}
 
 		/// <summary>
@@ -75,11 +92,15 @@ namespace EpicGames.Horde.Issues
 		/// <returns>True is the fingerprints match</returns>
 		public static bool IsMatchForNewSpan(this IIssueFingerprint fingerprint, IIssueFingerprint other)
 		{
-			if(fingerprint.GroupingRules != null)
+			if (!fingerprint.Type.Equals(other.Type, StringComparison.Ordinal))
+			{
+				return false;
+			}
+			if (fingerprint.RejectKeys != null && fingerprint.RejectKeys.Any(x => other.Keys.Contains(x)))
 			{
-				return fingerprint.GroupingRules.IsMatchForNewSpan(fingerprint, other);
+				return false;
 			}
-			return DefaultIssueGroupingRules.Instance.IsMatchForNewSpan(fingerprint, other);
+			return true;
 		}
 	}
 }
diff --git a/Engine/Source/Programs/Shared/EpicGames.Horde/Issues/IssueEvent.cs b/Engine/Source/Programs/Shared/EpicGames.Horde/Issues/IssueEvent.cs
--- a/Engine/Source/Programs/Shared/EpicGames.Horde/Issues/IssueEvent.cs
+++ b/Engine/Source/Programs/Shared/EpicGames.Horde/Issues/IssueEvent.cs
@@ -89,11 +89,6 @@ namespace EpicGames.Horde.Issues
 		/// </summary>
 		public string ChangeFilter { get; set; }
 
-		/// <summary>
-		/// Issue grouping rules
-		/// </summary>
-		public IIssueGroupingRules GroupingRules { get; set; }
-
 		/// <summary>
 		/// Individual log events
 		/// </summary>
@@ -110,7 +105,6 @@ namespace EpicGames.Horde.Issues
 			Type = type;
 			SummaryTemplate = summaryTemplate;
 			ChangeFilter = changeFilter;
-			GroupingRules = DefaultIssueGroupingRules.Instance;
 		}
 
 		/// <summary>
diff --git a/Engine/Source/Programs/Shared/EpicGames.Horde/Issues/IssueFingerprint.cs b/Engine/Source/Programs/Shared/EpicGames.Horde/Issues/IssueFingerprint.cs
--- a/Engine/Source/Programs/Shared/EpicGames.Horde/Issues/IssueFingerprint.cs
+++ b/Engine/Source/Programs/Shared/EpicGames.Horde/Issues/IssueFingerprint.cs
@@ -2,9 +2,7 @@
 
 using System;
 using System.Collections.Generic;
-using System.Text.Json;
 using System.Text.Json.Serialization;
-using System.Text.Json.Serialization.Metadata;
 
 #pragma warning disable CA2227 // Change 'x' to be read-only by removing the property setter
 
@@ -54,12 +52,6 @@ namespace EpicGames.Horde.Issues
 		[JsonPropertyName("filter")]
 		public string ChangeFilter { get; set; } = "...";
 
-		/// <summary>
-		/// Issue grouping rules
-		/// </summary>
-		[JsonPropertyName("groupingRules")]
-		public IIssueGroupingRules GroupingRules { get; set; } = DefaultIssueGroupingRules.Instance;
-
 		/// <summary>
 		/// Keys set which is null when empty
 		/// </summary>
@@ -97,36 +89,5 @@ namespace EpicGames.Horde.Issues
 			Type = type;
 			SummaryTemplate = summaryTemplate;
 		}
-
-		private static readonly JsonSerializerOptions? s_jsonSerializerOptions = new JsonSerializerOptions()
-		{
-			TypeInfoResolver = new DefaultJsonTypeInfoResolver
-			{
-				Modifiers =
-					{
-						typeInfo =>
-						{
-							if(typeInfo.Type == typeof(IIssueGroupingRules))
-							{
-								typeInfo.PolymorphismOptions = new JsonPolymorphismOptions
-								{
-									TypeDiscriminatorPropertyName = "$type",
-									IgnoreUnrecognizedTypeDiscriminators = false,
-									UnknownDerivedTypeHandling = JsonUnknownDerivedTypeHandling.FailSerialization
-								};
-
-								typeInfo.PolymorphismOptions.DerivedTypes.Add(new JsonDerivedType(typeof(DefaultIssueGroupingRules), "d"));
-								typeInfo.PolymorphismOptions.DerivedTypes.Add(new JsonDerivedType(typeof(CustomizableIssueGroupingRules), "c"));
-							}
-						}
-					}
-			},
-			PropertyNameCaseInsensitive = true
-		};
-
-		/// <summary>
-		/// Always use those JsonSerializerOptions when serializing/deserializing a JSON IssueFingerprint
-		/// </summary>
-		public static JsonSerializerOptions? JsonSerializerOptions => s_jsonSerializerOptions;
 	}
 }
diff --git a/Engine/Source/Programs/Shared/EpicGames.Horde/Issues/IssueGroupingRules.cs b/Engine/Source/Programs/Shared/EpicGames.Horde/Issues/IssueGroupingRules.cs
deleted file mode 100644
