Horde: Add Example Horde Server Plugin, Tests, and Dashboard

diff --git a/Engine/Source/Programs/Horde/HordeDashboard/plugins/exampleplugin/examplePluginView.tsx b/Engine/Source/Programs/Horde/HordeDashboard/plugins/exampleplugin/examplePluginView.tsx
new file mode 100644
--- /dev/null
+++ b/Engine/Source/Programs/Horde/HordeDashboard/plugins/exampleplugin/examplePluginView.tsx
@@ -0,0 +1,30 @@
+// Copyright Epic Games, Inc. All Rights Reserved.
+
+import { Stack, Text } from "@fluentui/react"
+import { getGreeting } from "./api"
+import { useState } from "react";
+import { Breadcrumbs } from "horde/components/Breadcrumbs";
+import { TopNav } from "horde/components/TopNav";
+import { getHordeStyling } from "horde/styles/Styles";
+
+
+export const ExamplePluginView: React.FC = () => {
+
+    const [greeting, setGreeting] = useState("Getting Greeting");
+
+    getGreeting().then(greeting => {
+        setGreeting(greeting);
+    })
+
+    const hordeClasses = getHordeStyling();
+
+    return <Stack className={hordeClasses.horde} key="key_metrics_graph_test">
+            <TopNav />
+            <Breadcrumbs items={[{ text: 'Example Plugin' }]} />
+            <Stack horizontalAlign="center" style={{ paddingTop: 24 }}>
+                <Text variant="mediumPlus">{greeting}</Text>
+            </Stack>
+        </Stack>
+
+
+}
\ No newline at end of file
diff --git a/Engine/Source/Programs/Horde/HordeDashboard/plugins/exampleplugin/index.tsx b/Engine/Source/Programs/Horde/HordeDashboard/plugins/exampleplugin/index.tsx
new file mode 100644
--- /dev/null
+++ b/Engine/Source/Programs/Horde/HordeDashboard/plugins/exampleplugin/index.tsx
@@ -0,0 +1,15 @@
+// Copyright Epic Games, Inc. All Rights Reserved.
+
+import { registerHordePlugin, MountType } from "..";
+import { ExamplePluginView } from "./examplePluginView";
+
+registerHordePlugin({
+   id: "exampleplugin",
+   routes: [{ path: "exampleplugin", element: <ExamplePluginView /> }],
+   mount: {
+      type: MountType.TopNav,
+      context: "Tools",
+      text: "ExamplePlugin",
+      route: `/exampleplugin`
+   }
+})
\ No newline at end of file
diff --git a/Engine/Source/Programs/Horde/HordeDashboard/plugins/registry.ts b/Engine/Source/Programs/Horde/HordeDashboard/plugins/registry.ts
--- a/Engine/Source/Programs/Horde/HordeDashboard/plugins/registry.ts
+++ b/Engine/Source/Programs/Horde/HordeDashboard/plugins/registry.ts
@@ -1,5 +1,5 @@
 // Copyright Epic Games, Inc. All Rights Reserved.
 
+import "./exampleplugin"
 import "./analytics"
-
 import "./test-automation"
diff --git a/Engine/Source/Programs/Horde/HordeServer/Defaults/default.global.json b/Engine/Source/Programs/Horde/HordeServer/Defaults/default.global.json
--- a/Engine/Source/Programs/Horde/HordeServer/Defaults/default.global.json
+++ b/Engine/Source/Programs/Horde/HordeServer/Defaults/default.global.json
@@ -129,6 +129,9 @@
 					}
 				]
 			}
+		},
+		"exampleplugin": {
+			"greetingMessage": "Hello from the Horde Example Plugin"
 		}
 	},
 
diff --git a/Engine/Source/Programs/Horde/HordeServer/HordeServer.csproj b/Engine/Source/Programs/Horde/HordeServer/HordeServer.csproj
--- a/Engine/Source/Programs/Horde/HordeServer/HordeServer.csproj
+++ b/Engine/Source/Programs/Horde/HordeServer/HordeServer.csproj
@@ -114,6 +114,7 @@
     <ProjectReference Include="..\Plugins\Build\HordeServer.Build\HordeServer.Build.csproj" />
     <ProjectReference Include="..\Plugins\Compute\HordeServer.Compute\HordeServer.Compute.csproj" />
     <ProjectReference Include="..\Plugins\Ddc\HordeServer.Ddc\HordeServer.Ddc.csproj" />
+    <ProjectReference Include="..\Plugins\ExamplePlugin\HordeServer.ExamplePlugin\HordeServer.ExamplePlugin.csproj" />
     <ProjectReference Include="..\Plugins\Experimental\HordeServer.Experimental\HordeServer.Experimental.csproj" />
     <ProjectReference Include="..\Plugins\Secrets\HordeServer.Secrets\HordeServer.Secrets.csproj" />
     <ProjectReference Include="..\Plugins\Storage\HordeServer.Storage\HordeServer.Storage.csproj" />
diff --git a/Engine/Source/Programs/Horde/HordeServer/appsettings.Local.json b/Engine/Source/Programs/Horde/HordeServer/appsettings.Local.json
--- a/Engine/Source/Programs/Horde/HordeServer/appsettings.Local.json
+++ b/Engine/Source/Programs/Horde/HordeServer/appsettings.Local.json
@@ -25,8 +25,11 @@
 			"ReplicateContent": false
 		},
 
-		// Enable the experimental plugin for local development
+		// Enable the Example/experimental plugin for local development
 		"Plugins": {
+			"ExamplePlugin": {
+				"Enabled": true				
+			},
 			"Experimental": {
 				"Enabled": true
 			}
diff --git a/Engine/Source/Programs/Horde/Plugins/ExamplePlugin/HordeServer.ExamplePlugin.Tests/ExamplePluginTestSetup.cs b/Engine/Source/Programs/Horde/Plugins/ExamplePlugin/HordeServer.ExamplePlugin.Tests/ExamplePluginTestSetup.cs
new file mode 100644
--- /dev/null
+++ b/Engine/Source/Programs/Horde/Plugins/ExamplePlugin/HordeServer.ExamplePlugin.Tests/ExamplePluginTestSetup.cs
@@ -0,0 +1,23 @@
+// Copyright Epic Games, Inc. All Rights Reserved.
+
+using HordeServer.Tests;
+using Microsoft.Extensions.DependencyInjection;
+namespace HordeServer.Example.Plugin.Tests
+{
+	/// <summary>
+	/// Handles set up of collections, services, fixtures etc during testing
+	/// </summary>
+	public class ExamplePluginTestSetup : BuildTestSetup
+	{
+		public ExamplePluginTestSetup()
+		{
+			AddPlugin<ExamplePlugin>();
+		}
+
+		protected override void ConfigureServices(IServiceCollection services)
+		{
+			base.ConfigureServices(services);
+
+		}
+	}
+}
\ No newline at end of file
diff --git a/Engine/Source/Programs/Horde/Plugins/ExamplePlugin/HordeServer.ExamplePlugin.Tests/ExamplePluginTests.cs b/Engine/Source/Programs/Horde/Plugins/ExamplePlugin/HordeServer.ExamplePlugin.Tests/ExamplePluginTests.cs
new file mode 100644
--- /dev/null
+++ b/Engine/Source/Programs/Horde/Plugins/ExamplePlugin/HordeServer.ExamplePlugin.Tests/ExamplePluginTests.cs
@@ -0,0 +1,60 @@
+// Copyright Epic Games, Inc. All Rights Reserved.
+
+using System.Security.Claims;
+using HordeServer.Tests;
+using HordeServer.Utilities;
+using Microsoft.AspNetCore.Http;
+using Microsoft.AspNetCore.Mvc;
+using Microsoft.Extensions.DependencyInjection;
+using HordeServer.Server;
+
+namespace HordeServer.Example.Plugin.Tests
+{
+
+	[TestClass]
+	public class ExamplePluginTests : BuildTestSetup
+	{
+		private readonly string _greetingMessage = "Test Greeting";
+
+		private ExamplePluginController GetExamplePluginController()
+		{
+			ExamplePluginController controller = ActivatorUtilities.CreateInstance<ExamplePluginController>(ServiceProvider);
+			controller.ControllerContext = GetControllerContext();
+			return controller;
+		}
+
+		private static ControllerContext GetControllerContext()
+		{
+			ControllerContext controllerContext = new ControllerContext();
+			controllerContext.HttpContext = new DefaultHttpContext();
+			controllerContext.HttpContext.User = new ClaimsPrincipal(new ClaimsIdentity(
+				new List<Claim> { HordeClaims.AdminClaim.ToClaim() }, "TestAuthType"));
+			return controllerContext;
+		}
+
+		public ExamplePluginTests()
+		{
+			AddPlugin<ExamplePlugin>();
+
+			ExamplePluginConfig examplePluginConfig = new ExamplePluginConfig();
+			examplePluginConfig.GreetingMessage = _greetingMessage;
+
+			GlobalConfig globalConfig = new GlobalConfig();
+			globalConfig.Plugins.AddExamplePluginConfig(examplePluginConfig);
+
+			SetConfigAsync(globalConfig).Wait();
+
+		}
+
+		[TestMethod]
+		public void TestGreeting()
+		{
+			ExamplePluginController controller = GetExamplePluginController();
+
+			ContentResult greeting = (controller.GetGreeting() as ContentResult)!;
+
+			Assert.AreEqual(greeting.Content, _greetingMessage);
+		}
+	}
+}
+
diff --git a/Engine/Source/Programs/Horde/Plugins/ExamplePlugin/HordeServer.ExamplePlugin.Tests/HordeServer.ExamplePlugin.Tests.csproj b/Engine/Source/Programs/Horde/Plugins/ExamplePlugin/HordeServer.ExamplePlugin.Tests/HordeServer.ExamplePlugin.Tests.csproj
new file mode 100644
--- /dev/null
+++ b/Engine/Source/Programs/Horde/Plugins/ExamplePlugin/HordeServer.ExamplePlugin.Tests/HordeServer.ExamplePlugin.Tests.csproj
@@ -0,0 +1,46 @@
+﻿<Project Sdk="Microsoft.NET.Sdk">
+
+  <PropertyGroup>
+    <TargetFramework>net8.0</TargetFramework>
+    <ImplicitUsings>enable</ImplicitUsings>
+    <Nullable>enable</Nullable>
+    <IsPackable>false</IsPackable>
+    <IsTestProject>true</IsTestProject>
+    <RootNamespace>HordeServer.Tests</RootNamespace>
+    <RunAnalyzersDuringBuild>False</RunAnalyzersDuringBuild>
+    <Configurations>Debug;Release;Analyze</Configurations>
+    <NoWarn>1701;1702;1591</NoWarn>
+  </PropertyGroup>
+
+  <PropertyGroup Condition=" '$(Configuration)|$(Platform)' == 'Analyze|AnyCPU' ">
+    <RunAnalyzersDuringBuild>True</RunAnalyzersDuringBuild>
+    <TreatWarningsAsErrors>false</TreatWarningsAsErrors>
+    <EnforceCodeStyleInBuild>True</EnforceCodeStyleInBuild>
+    <GenerateDocumentationFile>True</GenerateDocumentationFile>
+  </PropertyGroup>
+
+  <ItemGroup>
+    <PackageReference Include="coverlet.collector" Version="6.0.2">
+      <PrivateAssets>all</PrivateAssets>
+      <IncludeAssets>runtime; build; native; contentfiles; analyzers; buildtransitive</IncludeAssets>
+    </PackageReference>
+    <PackageReference Include="Microsoft.NET.Test.Sdk" Version="17.11.1" />
+    <PackageReference Include="MSTest.TestAdapter" Version="3.6.2" />
+    <PackageReference Include="MSTest.TestFramework" Version="3.6.2" />
+    <PackageReference Include="Microsoft.AspNetCore.Mvc.Testing" Version="8.0.10" />
+    <PackageReference Include="Moq" Version="4.20.70" />
+    <PackageReference Include="Microsoft.Extensions.Caching.Memory" Version="8.0.1" />
+  </ItemGroup>
+
+  <ItemGroup>
+    <ProjectReference Include="..\..\Build\HordeServer.Build\HordeServer.Build.csproj" />
+    <ProjectReference Include="..\..\Build\HordeServer.Build.Tests\HordeServer.Build.Tests.csproj" />
+    <ProjectReference Include="..\..\..\HordeServer.Tests.Shared\HordeServer.Tests.Shared.csproj" />
+    <ProjectReference Include="..\HordeServer.ExamplePlugin\HordeServer.ExamplePlugin.csproj" />
+  </ItemGroup>
+
+  <ItemGroup>
+    <Using Include="Microsoft.VisualStudio.TestTools.UnitTesting" />
+  </ItemGroup>
+
+</Project>
diff --git a/Engine/Source/Programs/Horde/Plugins/ExamplePlugin/HordeServer.ExamplePlugin/ExamplePlugin.cs b/Engine/Source/Programs/Horde/Plugins/ExamplePlugin/HordeServer.ExamplePlugin/ExamplePlugin.cs
new file mode 100644
--- /dev/null
+++ b/Engine/Source/Programs/Horde/Plugins/ExamplePlugin/HordeServer.ExamplePlugin/ExamplePlugin.cs
@@ -0,0 +1,51 @@
+// Copyright Epic Games, Inc. All Rights Reserved.
+
+using HordeServer.Plugins;
+using Microsoft.AspNetCore.Builder;
+using Microsoft.Extensions.DependencyInjection;
+
+namespace HordeServer
+{
+	/// <summary>
+	/// Entry point for the ExamplePlugin
+	/// </summary>
+	[Plugin("ExamplePlugin", EnabledByDefault = false, GlobalConfigType = typeof(ExamplePluginConfig), ServerConfigType = typeof(ExamplePluginServerConfig))]
+	public class ExamplePlugin : IPluginStartup
+	{
+		/// <summary>
+		/// Constructor
+		/// </summary>
+		public ExamplePlugin()
+		{
+		}
+
+		/// <inheritdoc/>
+		public void Configure(IApplicationBuilder app)
+		{
+		}
+
+		/// <inheritdoc/>
+		public void ConfigureServices(IServiceCollection serviceCollection)
+		{
+			
+		}
+	}
+
+	/// <summary>
+	/// Helper methods for ExamplePlugin plugin config
+	/// </summary>
+	public static class ExamplePluginExtensions
+	{
+		/// <summary>
+		/// Configures the ExamplePlugin plugin
+		/// </summary>
+		public static void AddExamplePluginConfig(this IDictionary<PluginName, IPluginConfig> dictionary, ExamplePluginConfig examplePluginConfig)
+			=> dictionary[new PluginName("ExamplePlugin")] = examplePluginConfig;
+
+		/// <summary>
+		/// Gets configuration for the ExamplePlugin plugin
+		/// </summary>
+		public static ExamplePluginConfig GetExamplePluginConfig(this IDictionary<PluginName, IPluginConfig> dictionary)
+			=> (ExamplePluginConfig)dictionary[new PluginName("ExamplePlugin")];
+	}
+}
diff --git a/Engine/Source/Programs/Horde/Plugins/ExamplePlugin/HordeServer.ExamplePlugin/ExamplePluginConfig.cs b/Engine/Source/Programs/Horde/Plugins/ExamplePlugin/HordeServer.ExamplePlugin/ExamplePluginConfig.cs
new file mode 100644
--- /dev/null
+++ b/Engine/Source/Programs/Horde/Plugins/ExamplePlugin/HordeServer.ExamplePlugin/ExamplePluginConfig.cs
@@ -0,0 +1,22 @@
+// Copyright Epic Games, Inc. All Rights Reserved.
+
+using HordeServer.Plugins;
+
+namespace HordeServer
+{
+	/// <summary>
+	/// Configuration for the ExamplePlugin plugin
+	/// </summary>
+	public class ExamplePluginConfig : IPluginConfig
+	{
+		/// <summary>
+		/// The Greeting provided by the example plugin
+		/// </summary>
+		public string GreetingMessage { get; set; } = "Default Greeting";
+
+		/// <inheritdoc/>
+		public void PostLoad(PluginConfigOptions configOptions)
+		{			
+		}
+	}
+}
diff --git a/Engine/Source/Programs/Horde/Plugins/ExamplePlugin/HordeServer.ExamplePlugin/ExamplePluginServerConfig.cs b/Engine/Source/Programs/Horde/Plugins/ExamplePlugin/HordeServer.ExamplePlugin/ExamplePluginServerConfig.cs
new file mode 100644
--- /dev/null
+++ b/Engine/Source/Programs/Horde/Plugins/ExamplePlugin/HordeServer.ExamplePlugin/ExamplePluginServerConfig.cs
@@ -0,0 +1,13 @@
+// Copyright Epic Games, Inc. All Rights Reserved.
+
+using HordeServer.Plugins;
+
+namespace HordeServer
+{
+	/// <summary>
+	/// Server configuration for the ExamplePlugin 
+	/// </summary>
+	public class ExamplePluginServerConfig : PluginServerConfig
+	{
+	}
+}
diff --git a/Engine/Source/Programs/Horde/Plugins/ExamplePlugin/HordeServer.ExamplePlugin/HordeServer.ExamplePlugin.csproj b/Engine/Source/Programs/Horde/Plugins/ExamplePlugin/HordeServer.ExamplePlugin/HordeServer.ExamplePlugin.csproj
new file mode 100644
--- /dev/null
+++ b/Engine/Source/Programs/Horde/Plugins/ExamplePlugin/HordeServer.ExamplePlugin/HordeServer.ExamplePlugin.csproj
@@ -0,0 +1,46 @@
+﻿<Project Sdk="Microsoft.NET.Sdk">
+
+  <PropertyGroup>
+    <TargetFramework>net8.0</TargetFramework>
+    <ImplicitUsings>enable</ImplicitUsings>
+    <Nullable>enable</Nullable>
+    <RootNamespace>HordeServer</RootNamespace>
+    <RunAnalyzersDuringBuild>False</RunAnalyzersDuringBuild>
+    <Configurations>Debug;Release;Analyze</Configurations>
+  </PropertyGroup>
+
+  <PropertyGroup Condition=" '$(Configuration)|$(Platform)' == 'Analyze|AnyCPU' ">
+    <RunAnalyzersDuringBuild>True</RunAnalyzersDuringBuild>
+    <TreatWarningsAsErrors>false</TreatWarningsAsErrors>
+    <EnforceCodeStyleInBuild>True</EnforceCodeStyleInBuild>
+    <GenerateDocumentationFile>True</GenerateDocumentationFile>
+    <DocumentationFile></DocumentationFile>
+  </PropertyGroup>
+
+  <ItemGroup>
+    <PackageReference Include="Grpc.AspNetCore" Version="2.59.0" />
+    <PackageReference Include="Grpc.AspNetCore.Server.Reflection" Version="2.59.0" />
+    <PackageReference Include="Grpc.Net.Client" Version="2.59.0" />
+    <PackageReference Include="Grpc.Tools" Version="2.59.0">
+      <PrivateAssets>all</PrivateAssets>
+      <IncludeAssets>runtime; build; native; contentfiles; analyzers; buildtransitive</IncludeAssets>
+    </PackageReference>
+    <PackageReference Include="Polly" Version="8.2.0" />
+    <PackageReference Include="TimeZoneConverter" Version="6.1.0" />
+  </ItemGroup>
+
+  <ItemGroup>
+    <ProjectReference Include="..\..\..\..\Shared\EpicGames.Core\EpicGames.Core.csproj" />
+    <ProjectReference Include="..\..\..\..\Shared\EpicGames.Horde\EpicGames.Horde.csproj" />
+    <ProjectReference Include="..\..\..\..\Shared\EpicGames.IoHash\EpicGames.IoHash.csproj" />
+    <ProjectReference Include="..\..\..\..\Shared\EpicGames.Perforce.Managed\EpicGames.Perforce.Managed.csproj" />
+    <ProjectReference Include="..\..\..\..\Shared\EpicGames.Perforce\EpicGames.Perforce.csproj" />
+    <ProjectReference Include="..\..\..\..\Shared\EpicGames.Redis\EpicGames.Redis.csproj" />
+    <ProjectReference Include="..\..\..\..\Shared\EpicGames.Slack\EpicGames.Slack.csproj" />
+    <ProjectReference Include="..\..\..\HordeServer.Shared\HordeServer.Shared.csproj" />
+    <ProjectReference Include="..\..\Build\HordeServer.Build\HordeServer.Build.csproj" />
+    <ProjectReference Include="..\..\Compute\HordeServer.Compute\HordeServer.Compute.csproj" />
+    <ProjectReference Include="..\..\Tools\HordeServer.Tools\HordeServer.Tools.csproj" />
+  </ItemGroup>
+
+</Project>
