Horde: Add example backend api calls for dashboard

diff --git a/Engine/Source/Programs/Horde/HordeDashboard/plugins/exampleplugin/api.ts b/Engine/Source/Programs/Horde/HordeDashboard/plugins/exampleplugin/api.ts
new file mode 100644
--- /dev/null
+++ b/Engine/Source/Programs/Horde/HordeDashboard/plugins/exampleplugin/api.ts
@@ -0,0 +1,11 @@
+// Copyright Epic Games, Inc. All Rights Reserved.
+
+import backend from "horde/backend";
+
+export async function getGreeting(): Promise<string> {
+    return new Promise<string>((resolve, reject) => {
+        backend.fetch.get(`/api/v1/exampleplugin/greeting`).then((response) => {
+            resolve(response.data as string);
+        }).catch(reason => { reject(reason); });
+    });
+}
\ No newline at end of file
