Horde: Minor fix to missing relay error message

Avoids stacking error prefixes when reporting, for example in UBT/UBA sessions.

diff --git a/Engine/Source/Programs/Horde/Plugins/Compute/HordeServer.Compute/Compute/ComputeService.cs b/Engine/Source/Programs/Horde/Plugins/Compute/HordeServer.Compute/Compute/ComputeService.cs
--- a/Engine/Source/Programs/Horde/Plugins/Compute/HordeServer.Compute/Compute/ComputeService.cs
+++ b/Engine/Source/Programs/Horde/Plugins/Compute/HordeServer.Compute/Compute/ComputeService.cs
@@ -511,7 +511,7 @@ namespace HordeServer.Compute
 			}
 			catch (AgentRelayException are)
 			{
-				throw new ComputeServiceException("Unable to allocate compute resource: " + are.Message, are) { ShowToUser = true };
+				throw new ComputeServiceException(are.Message, are) { ShowToUser = true };
 			}
 		}
 		
