Horde: Update release notes

#rnx

diff --git a/Engine/Source/Programs/Horde/Docs/ReleaseNotes.md b/Engine/Source/Programs/Horde/Docs/ReleaseNotes.md
--- a/Engine/Source/Programs/Horde/Docs/ReleaseNotes.md
+++ b/Engine/Source/Programs/Horde/Docs/ReleaseNotes.md
@@ -1,5 +1,22 @@
 # Release Notes
 
+## 2025-09-18
+
+* Add example backend api calls for dashboard (45941624)
+* Add Example Horde Server Plugin, Tests, and Dashboard (45940896)
+* Text filtering & combo box conversion (45895296)
+* Styling pass V3 & general UI improvements:
+  - Styling for aborted
+  - Styling for skipped
+  - Sticky summary columns
+  - Min width reduced for smaller screens
+  - Skipped as red (+ view/ui options segregation refactor)
+  - Warnings as summary success
+  - Patch changelist UTC skew in UI
+  - Fall back to job start when no changelist data is available from server
+  - Clearing project cascades correctly (45849344)
+* Export to csv button. (45763412)
+
 ## 2025-09-11
 
 * Summary column & hours label for title & improving the console.log for change summary request count so we can see the *actual* request (45714224)
