Horde: Add missing ExamplePluginController

diff --git a/Engine/Source/Programs/Horde/Plugins/ExamplePlugin/HordeServer.ExamplePlugin/ExamplePluginController.cs b/Engine/Source/Programs/Horde/Plugins/ExamplePlugin/HordeServer.ExamplePlugin/ExamplePluginController.cs
new file mode 100644
--- /dev/null
+++ b/Engine/Source/Programs/Horde/Plugins/ExamplePlugin/HordeServer.ExamplePlugin/ExamplePluginController.cs
@@ -0,0 +1,43 @@
+// Copyright Epic Games, Inc. All Rights Reserved.
+
+using HordeServer.Utilities;
+using Microsoft.AspNetCore.Authorization;
+using Microsoft.AspNetCore.Mvc;
+using Microsoft.Extensions.Logging;
+using Microsoft.Extensions.Options;
+
+namespace HordeServer
+{
+	/// <summary>
+	/// Controller example for ExamplePlugun
+	/// </summary>
+	
+	[Authorize]
+	[ApiController]	
+	public class ExamplePluginController : HordeControllerBase
+	{
+		readonly IOptionsSnapshot<ExamplePluginConfig> _examplePluginConfig;
+		readonly ILogger _logger;
+
+		/// <summary>
+		/// Constructor
+		/// </summary>
+		public ExamplePluginController(IOptionsSnapshot<ExamplePluginConfig> examplePluginConfig, ILogger<ExamplePluginController> logger)
+		{
+			_examplePluginConfig = examplePluginConfig;
+			_logger = logger;
+		}
+
+		/// <summary>
+		/// Example plugin endpoint which returns message provided in 
+		/// </summary>
+		[HttpGet]
+		[Route("/api/v1/exampleplugin/greeting")]
+		public ActionResult GetGreeting()
+		{
+			string greeting = _examplePluginConfig.Value.GreetingMessage;
+			_logger.LogInformation("Getting Greeting: {Greeting}", greeting);
+			return Content(greeting, "text/plain");
+		}
+	}
+}
