Horde: Forward schedule-based claims when scheduled jobs are created

diff --git a/Engine/Source/Programs/Horde/Plugins/Build/HordeServer.Build.Tests/Jobs/Schedules/SchedulerTests.cs b/Engine/Source/Programs/Horde/Plugins/Build/HordeServer.Build.Tests/Jobs/Schedules/SchedulerTests.cs
--- a/Engine/Source/Programs/Horde/Plugins/Build/HordeServer.Build.Tests/Jobs/Schedules/SchedulerTests.cs
+++ b/Engine/Source/Programs/Horde/Plugins/Build/HordeServer.Build.Tests/Jobs/Schedules/SchedulerTests.cs
@@ -595,6 +595,31 @@ namespace HordeServer.Tests.Jobs.Schedules
 			Assert.AreEqual(1233, jobs3[1].CommitId.GetPerforceChange());
 		}
 
+		[TestMethod]
+		public async Task ClaimsAsync()
+		{
+			DateTime startTime = new (2021, 1, 1, 12, 0, 0, DateTimeKind.Local); // Friday Jan 1, 2021 
+			Clock.UtcNow = startTime;
+			Clock.TimeZone = TimeZoneInfo.Local;
+
+			ScheduleConfig schedule = new() {
+				Enabled = true,
+				Claims = [new ScheduleClaimConfig { Type = "MyType", Value = "MyValue"}]
+			};
+			schedule.Patterns.Add(new SchedulePatternConfig { MinTime = ScheduleTimeOfDay.Parse("13:00"), MaxTime = ScheduleTimeOfDay.Parse("14:00"), Interval = ScheduleInterval.Parse("15m") });
+			_ = await SetScheduleAsync(schedule);
+			
+			// Trigger a job
+			await Clock.AdvanceAsync(TimeSpan.FromHours(1.25));
+			await ScheduleService.TickForTestingAsync();
+			
+			List<IJob> jobs = await GetNewJobsAsync();
+			Assert.AreEqual(1, jobs.Count);
+			Assert.AreEqual(1, jobs[0].Claims.Count);
+			Assert.AreEqual("MyType", jobs[0].Claims[0].Type);
+			Assert.AreEqual("MyValue", jobs[0].Claims[0].Value);
+		}
+
 		[TestMethod]
 		public async Task UpdateConfigAsync()
 		{
diff --git a/Engine/Source/Programs/Horde/Plugins/Build/HordeServer.Build/Jobs/Schedules/ScheduleService.cs b/Engine/Source/Programs/Horde/Plugins/Build/HordeServer.Build/Jobs/Schedules/ScheduleService.cs
--- a/Engine/Source/Programs/Horde/Plugins/Build/HordeServer.Build/Jobs/Schedules/ScheduleService.cs
+++ b/Engine/Source/Programs/Horde/Plugins/Build/HordeServer.Build/Jobs/Schedules/ScheduleService.cs
@@ -12,6 +12,7 @@ using EpicGames.Horde.Streams;
 using EpicGames.Redis;
 using EpicGames.Redis.Utility;
 using EpicGames.Serialization;
+using HordeServer.Acls;
 using HordeServer.Auditing;
 using HordeServer.Commits;
 using HordeServer.Jobs.Graphs;
@@ -530,8 +531,9 @@ namespace HordeServer.Jobs.Schedules
 			{
 				cancellationToken.ThrowIfCancellationRequested();
 
-				CreateJobOptions options = new CreateJobOptions(templateRefConfig);
+				CreateJobOptions options = new (templateRefConfig);
 				options.Priority = template.Priority;
+				options.Claims = schedule.Claims?.Select(x => new AclClaimConfig(x.Type, x.Value)).ToList() ?? [];
 				template.GetDefaultParameters(options.Parameters, true);
 				template.GetArgumentsForParameters(options.Parameters, options.Arguments);
 
