Fixing configurations not being found due to streams being specified along with streamtags

#rnx

diff --git a/Engine/Source/Programs/Horde/Plugins/Experimental/HordeServer.Experimental/ExperimentalConfig.cs b/Engine/Source/Programs/Horde/Plugins/Experimental/HordeServer.Experimental/ExperimentalConfig.cs
--- a/Engine/Source/Programs/Horde/Plugins/Experimental/HordeServer.Experimental/ExperimentalConfig.cs
+++ b/Engine/Source/Programs/Horde/Plugins/Experimental/HordeServer.Experimental/ExperimentalConfig.cs
@@ -27,7 +27,7 @@ namespace HordeServer
 			foreach (NotificationConfig config in Notifications)
 			{
 				config.PostLoad();
-				_notificationConfigLookup.Add(config.Id, config);
+				_notificationConfigLookup.TryAdd(config.Id, config);
 			}
 		}
 
@@ -47,8 +47,15 @@ namespace HordeServer
 			foreach (NotificationConfig config in Notifications)
 			{
 				HashSet<JobNotificationConfig>? localJobNotificationConfigs;
-				if (config.TryGetJobNofications(streamId, templateId, out localJobNotificationConfigs, includeDisabledConfigs) ||
-					(streamTags is not null && config.TryGetJobNofications(streamTags, templateId, out localJobNotificationConfigs, includeDisabledConfigs)))
+
+				// A notification configuration can be setup with a stream, stream tag, or have configurations that use both.
+				// We need to search through all of our lookups to make sure that nothing has been missed
+				if (config.TryGetJobNofications(streamId, templateId, out localJobNotificationConfigs, includeDisabledConfigs))
+				{
+					jobNotificationConfigs.UnionWith(localJobNotificationConfigs);
+				}
+				
+				if (streamTags is not null && config.TryGetJobNofications(streamTags, templateId, out localJobNotificationConfigs, includeDisabledConfigs))
 				{
 					jobNotificationConfigs.UnionWith(localJobNotificationConfigs);
 				}
diff --git a/Engine/Source/Programs/Horde/Plugins/Experimental/HordeServer.Experimental/Notifications/NotificationConfig.cs b/Engine/Source/Programs/Horde/Plugins/Experimental/HordeServer.Experimental/Notifications/NotificationConfig.cs
--- a/Engine/Source/Programs/Horde/Plugins/Experimental/HordeServer.Experimental/Notifications/NotificationConfig.cs
+++ b/Engine/Source/Programs/Horde/Plugins/Experimental/HordeServer.Experimental/Notifications/NotificationConfig.cs
@@ -180,6 +180,7 @@ namespace HordeServer
 		public void PostLoad()
 		{
 			_streamLookup.Clear();
+			_streamTagLookup.Clear();
 			foreach (NotificationStreamConfig notificationStream in NotificationStreams)
 			{
 				notificationStream.PostLoad();
@@ -256,6 +257,11 @@ namespace HordeServer
 		{
 			jobNotificationConfigs = new HashSet<JobNotificationConfig>();
 
+			if (_streamTagLookup.Count == 0)
+			{
+				return false;
+			}
+
 			HashSet<NotificationStreamConfig>? streamConfigList;
 			foreach (StreamTag streamTag in streamTags)
 			{
diff --git a/Engine/Source/Programs/Horde/Plugins/Experimental/HordeServer.Experimental/Notifications/SlackNotificationProcessor.cs b/Engine/Source/Programs/Horde/Plugins/Experimental/HordeServer.Experimental/Notifications/SlackNotificationProcessor.cs
--- a/Engine/Source/Programs/Horde/Plugins/Experimental/HordeServer.Experimental/Notifications/SlackNotificationProcessor.cs
+++ b/Engine/Source/Programs/Horde/Plugins/Experimental/HordeServer.Experimental/Notifications/SlackNotificationProcessor.cs
@@ -359,7 +359,7 @@ namespace HordeServer.Experimental.Notifications
 			string streamInformation = $"'{parentJob.StreamId}'";
 			if (streamConfig is not null && streamConfig.StreamTags.Count > 0)
 			{
-				streamInformation += $" or with stream tags {String.Join(',', streamConfig.StreamTags)}";
+				streamInformation += $" or with stream tags {String.Join(',', streamConfig.StreamTags.Where(x => x.Enabled).Select(x => x.Name))}";
 			}
 
 			HashSet<JobNotificationConfig>? jobNotificationConfigs;
@@ -404,7 +404,7 @@ namespace HordeServer.Experimental.Notifications
 			string streamInformation = $"'{job.StreamId}'";
 			if (streamConfig is not null && streamConfig.StreamTags.Count > 0)
 			{
-				streamInformation += $" or with stream tags {String.Join(',', streamConfig.StreamTags)}";
+				streamInformation += $" or with stream tags {String.Join(',', streamConfig.StreamTags.Where(x => x.Enabled).Select(x => x.Name))}";
 			}
 
 			HashSet<JobNotificationConfig>? jobNotificationConfigs;
