[Horde][BuildHealth][Dashboard] Quick UI Feedback:
- Label renames (and updated underlying symbols for sanity)
- Dates included in tooltip
- Tooltip now inclludes conditional information - job outcome IF STEP IS ABORTED
- Outcome summaries now exist for each job - these are used to provide adjacent step context when on a particular cell of interest (at the moment, Aborted).
- Improved styling
-- Shape reflects the step State : Filled colour is for a completed, pending or running  job. Outline for a skipped a skipped job. Outline and Cross/Strikethrough for a Cancelled Job
-- Colour reflects the step outcome: Red failed, Yellow warnings, Green Success, Gray Pending, Cyan Running

diff --git a/Engine/Source/Programs/Horde/HordeDashboard/src/components/buildhealth/BuildHealthOptions.ts b/Engine/Source/Programs/Horde/HordeDashboard/src/components/buildhealth/BuildHealthOptions.ts
--- a/Engine/Source/Programs/Horde/HordeDashboard/src/components/buildhealth/BuildHealthOptions.ts
+++ b/Engine/Source/Programs/Horde/HordeDashboard/src/components/buildhealth/BuildHealthOptions.ts
@@ -193,9 +193,8 @@ export async function loadBuildHealthOptionsFromParams(
  * Build Health View Options.
  */
 export class BuildHealthUIOptionsState {
-    @observable includeDateAnchors: boolean = true;
-    @observable skippedIsFailure: boolean = false;
-    @observable warningsAsSummarySuccess: boolean = false;
+    @observable includeDateAnchors: boolean = false;
+    @observable warningsAsSummaryFailure: boolean = false;
 
     constructor() {
         makeObservable(this);
@@ -286,7 +285,7 @@ export class BuildHealthOptionsController {
     @action
     toggleSingleProject(id: string, name: string) {
         this.state.enabledProjects = { [id]: name };
-         this.clearStreams(true);
+        this.clearStreams(true);
         this.setOptionsChanged();
     }
 
@@ -400,22 +399,13 @@ export class BuildHealthOptionsController {
         this.uiState.includeDateAnchors = includeDateAnchors;
     }
 
-    /**
-     * Sets whether entries in the table that are skipped state should be rendered as failures.
-     * @param skippedIsFailure True to render as failures, false otherwise.
-     */
-    @action
-    setSkippedIsFailure(skippedIsFailure: boolean) {
-        this.uiState.skippedIsFailure = skippedIsFailure;
-    }
-
     /**
  * Sets whether summary rows should consider warnings as a success, or failure.
- * @param warningsAsSummarySuccess True to include warnings in the numerator of success ratio, false otherwise.
+ * @param warningsAsSummaryFailure False to include warnings in the numerator of success ratio, true otherwise.
  */
     @action
-    setWarningsAsSummarySuccess(warningsAsSummarySuccess: boolean) {
-        this.uiState.warningsAsSummarySuccess = warningsAsSummarySuccess;
+    setWarningsAsSummaryFailure(warningsAsSummaryFailure: boolean) {
+        this.uiState.warningsAsSummaryFailure = warningsAsSummaryFailure;
     }
 
     /**
diff --git a/Engine/Source/Programs/Horde/HordeDashboard/src/components/buildhealth/BuildhealthView.tsx b/Engine/Source/Programs/Horde/HordeDashboard/src/components/buildhealth/BuildhealthView.tsx
--- a/Engine/Source/Programs/Horde/HordeDashboard/src/components/buildhealth/BuildhealthView.tsx
+++ b/Engine/Source/Programs/Horde/HordeDashboard/src/components/buildhealth/BuildhealthView.tsx
@@ -97,7 +97,7 @@ const SearchableDropdown: React.FC<ISearchableDropdownProps> = (props) => {
  * React Component that is a single select drop down.
  * @returns React Component.
  */
-const BasicListParameter: React.FC<{ handler: BuildHealthDataHandler, buildHealthOptions: BuildHealthOptionsController, param: BuildHealthOptionList, disabled?: boolean, selectedKey?: string, request: DataHandlerRefreshRequest, onChange: (option: IDropdownOption<BuildHealthOptionData>) => void }> = function ConstructBasicListParamater({ param, buildHealthOptions, selectedKey, request, onChange }) {
+const BasicListParameter: React.FC<{ handler: BuildHealthDataHandler, buildHealthOptions: BuildHealthOptionsController, param: BuildHealthOptionList, disabled?: boolean, selectedKey?: string, request: DataHandlerRefreshRequest, onChange: (option: IDropdownOption<BuildHealthOptionData>) => void }> = function ConstructBasicListParamater({ param, handler, buildHealthOptions, selectedKey, request, onChange }) {
     const key = `${PARAMETER_KEY_PREFIX}${KEY_SEPARATOR}${param.label}`;
     const changeVersion = buildHealthOptions.optionsChangeVersion;
     const doptions: IDropdownOption<BuildHealthOptionData>[] = [];
@@ -138,7 +138,7 @@ const BasicListParameter: React.FC<{ handler: BuildHealthDataHandler, buildHealt
  * React Component that is a multi select drop down.
  * @returns React Component.
  */
-const MultiListParameter: React.FC<{ handler: BuildHealthDataHandler, buildHealthOptions: BuildHealthOptionsController, param: BuildHealthOptionList, disabled?: boolean, selectedKeys: string[], request: DataHandlerRefreshRequest, onChange: (option: IDropdownOption<BuildHealthOptionData>) => void }> = function ConstructMultilistParamter({ param, buildHealthOptions, disabled, selectedKeys, request, onChange }) {
+const MultiListParameter: React.FC<{ handler: BuildHealthDataHandler, buildHealthOptions: BuildHealthOptionsController, param: BuildHealthOptionList, disabled?: boolean, selectedKeys: string[], request: DataHandlerRefreshRequest, onChange: (option: IDropdownOption<BuildHealthOptionData>) => void }> = function ConstructMultilistParamter({ param, handler, buildHealthOptions, disabled, selectedKeys, request, onChange }) {
     const [filterText, setFilterText] = useState("");
     const changeVersion = buildHealthOptions.optionsChangeVersion;
     const key = `${PARAMETER_KEY_PREFIX}${KEY_SEPARATOR}${param.label}`;
@@ -408,24 +408,21 @@ const StepDropdown: React.FC<{ handler: BuildHealthDataHandler, buildHealthOptio
  * React Component used to manage the URL query params, and refresh it based on build health options. This is separate from the main component to scope rerenders to solely this component.
  * @returns React Component.
  */
-const BuildHealthUrlSync: React.FC<{ buildHealthController: BuildHealthOptionsController }> = function ConstrucBuildHealthSidebar({ buildHealthController }) {
+const BuildHealthUrlSync: React.FC<{ buildHealthController: BuildHealthOptionsController, handler: BuildHealthDataHandler }> = function ConstrucBuildHealthSidebar({ buildHealthController, handler }) {
     const location = useLocation();
     const navigate = useNavigate();
-    const controllerRef = useRef(buildHealthController);
 
     useEffect(() => {
         const initializeFromUrl = async () => {
             const params = parseBuildHealthQueryParams(location);
-            await loadBuildHealthOptionsFromParams(params, controllerRef.current, handler);
+            await loadBuildHealthOptionsFromParams(params, buildHealthController, handler);
         };
 
         initializeFromUrl();
-    },
-    // eslint-disable-next-line react-hooks/exhaustive-deps -- location is a special global in this case.
-    []);
+    }, []);
 
     useEffect(() => {
-        const desiredSearch = controllerRef.current.toNavigationQuery().toString();
+        const desiredSearch = buildHealthController.toNavigationQuery().toString();
         if (desiredSearch && location.search !== `?${desiredSearch}`) {
             navigate(`${location.pathname}?${desiredSearch}`, { replace: true });
         }
@@ -446,8 +443,7 @@ const BuildHealthSidebar: React.FC<{ handler: BuildHealthDataHandler, buildHealt
     // #region -- Visual Options -- 
 
     const [hideDateAnchors, setHideDateAnchors] = useState(() => buildHealthOptions.uiState.includeDateAnchors)
-    const [skippedAsRed, setSkippedAsRed] = useState(() => buildHealthOptions.uiState.skippedIsFailure)
-    const [warningsAsSummarySuccess, setWarningsAsSummarySuccess] = useState(() => buildHealthOptions.uiState.warningsAsSummarySuccess)
+    const [warningsAsSummaryFailure, setWarningsAsSummaryFailure] = useState(() => buildHealthOptions.uiState.warningsAsSummaryFailure)
 
     // #ednregion -- Visual Options --
 
@@ -469,14 +465,9 @@ const BuildHealthSidebar: React.FC<{ handler: BuildHealthDataHandler, buildHealt
         buildHealthOptions.setHideDateAnchors(isChecked);
     }
 
-    const updateSkippedIsFailure = (isChecked) => {
-        setSkippedAsRed(isChecked);
-        buildHealthOptions.setSkippedIsFailure(isChecked);
-    }
-
-    const updateWarningsAsSummarySuccess = (isChecked) => {
-        setWarningsAsSummarySuccess(isChecked);
-        buildHealthOptions.setWarningsAsSummarySuccess(isChecked);
+    const updateWarningsAsSummaryFailure = (isChecked) => {
+        setWarningsAsSummaryFailure(isChecked);
+        buildHealthOptions.setWarningsAsSummaryFailure(isChecked);
     }
 
     // #endregion -- Property Callbacks --
@@ -496,12 +487,8 @@ const BuildHealthSidebar: React.FC<{ handler: BuildHealthDataHandler, buildHealt
     }, [buildHealthOptions.uiState.includeDateAnchors]);
 
     useEffect(() => {
-        setSkippedAsRed(buildHealthOptions.uiState.skippedIsFailure);
-    }, [buildHealthOptions.uiState.skippedIsFailure]);
-
-    useEffect(() => {
-        setWarningsAsSummarySuccess(buildHealthOptions.uiState.warningsAsSummarySuccess);
-    }, [buildHealthOptions.uiState.warningsAsSummarySuccess]);
+        setWarningsAsSummaryFailure(buildHealthOptions.uiState.warningsAsSummaryFailure);
+    }, [buildHealthOptions.uiState.warningsAsSummaryFailure]);
 
     // #endregion -- Property Use Effects --
 
@@ -585,28 +572,11 @@ const BuildHealthSidebar: React.FC<{ handler: BuildHealthDataHandler, buildHealt
                         }}
                     >
                         <Checkbox
-                            label="Use date anchors"
+                            label="Show Date Anchors"
                             checked={hideDateAnchors}
                             onChange={(_, isChecked) => updateIncludeDateAnchors(!!isChecked)}
                         />
                     </TooltipHost>
-                    <TooltipHost content={
-                        <div style={{ maxWidth: 300 }}>
-                            <div>
-                                Skipped steps (due to dependency failures) render the same as failures.
-                            </div>
-                        </div>
-                    }
-                        calloutProps={{
-                            directionalHint: DirectionalHint.rightCenter
-                        }}
-                    >
-                        <Checkbox
-                            label="Skipped is failure"
-                            checked={skippedAsRed}
-                            onChange={(_, isChecked) => updateSkippedIsFailure(!!isChecked)}
-                        />
-                    </TooltipHost>
                     <TooltipHost content={
                         <div style={{ maxWidth: 300 }}>
                             <div>
@@ -619,9 +589,9 @@ const BuildHealthSidebar: React.FC<{ handler: BuildHealthDataHandler, buildHealt
                         }}
                     >
                         <Checkbox
-                            label="Warnings as summary success"
-                            checked={warningsAsSummarySuccess}
-                            onChange={(_, isChecked) => updateWarningsAsSummarySuccess(!!isChecked)}
+                            label="Warnings as Errors"
+                            checked={warningsAsSummaryFailure}
+                            onChange={(_, isChecked) => updateWarningsAsSummaryFailure(!!isChecked)}
                         />
                     </TooltipHost>
                 </Stack>
@@ -641,9 +611,9 @@ export const BuildHealthView: React.FC = observer(function ConstructBuildHealthV
     const jobHistorySpan = useMemo(() => ({
         start: buildHealthState.startDate,
         end: buildHealthState.endDate
-    }), 
-    // eslint-disable-next-line react-hooks/exhaustive-deps -- All necessary dependencies included for MobX observables
-    [buildHealthState.startDate, buildHealthState.endDate]);
+    }),
+        // eslint-disable-next-line react-hooks/exhaustive-deps -- All necessary dependencies included for MobX observables
+        [buildHealthState.startDate, buildHealthState.endDate]);
 
     const filter = useMemo(() => ({
         streams: buildHealthState.stepOutcomeEnabledStreamKeys,
@@ -651,31 +621,29 @@ export const BuildHealthView: React.FC = observer(function ConstructBuildHealthV
         steps: buildHealthState.stepOutcomeEnabledStepKeys,
         jobHistorySpan,
         includePreflights: buildHealthState.includePreflight,
-    }), 
-    // eslint-disable-next-line react-hooks/exhaustive-deps -- All necessary dependencies included for MobX observables
-    [
-        buildHealthState.stepOutcomeEnabledStreamKeys,
-        buildHealthState.stepOutcomeEnabledJobKeys,
-        buildHealthState.stepOutcomeEnabledStepKeys,
-        jobHistorySpan,
-        buildHealthState.includePreflight
-    ]);
+    }),
+        // eslint-disable-next-line react-hooks/exhaustive-deps -- All necessary dependencies included for MobX observables
+        [
+            buildHealthState.stepOutcomeEnabledStreamKeys,
+            buildHealthState.stepOutcomeEnabledJobKeys,
+            buildHealthState.stepOutcomeEnabledStepKeys,
+            jobHistorySpan,
+            buildHealthState.includePreflight
+        ]);
 
     const uiOptions = useMemo(() => ({
         includeDateAnchors: buildHealthUIState.includeDateAnchors,
-        skippedAsRed: buildHealthUIState.skippedIsFailure,
-        warningsAsSummarySuccess: buildHealthUIState.warningsAsSummarySuccess
-    }), 
-    // eslint-disable-next-line react-hooks/exhaustive-deps -- All necessary dependencies included for MobX observables
-    [
-        buildHealthUIState.includeDateAnchors,
-        buildHealthUIState.skippedIsFailure,
-        buildHealthUIState.warningsAsSummarySuccess
-    ]);
+        warningsAsSummaryFailure: buildHealthUIState.warningsAsSummaryFailure
+    }),
+        // eslint-disable-next-line react-hooks/exhaustive-deps -- All necessary dependencies included for MobX observables
+        [
+            buildHealthUIState.includeDateAnchors,
+            buildHealthUIState.warningsAsSummaryFailure
+        ]);
 
     return (
         <Stack className={hordeClasses.horde}>
-            <BuildHealthUrlSync buildHealthController={buildHealthController} />
+            <BuildHealthUrlSync buildHealthController={buildHealthController} handler={handler} />
             <TopNav />
             <Breadcrumbs items={[{ text: 'Build Health' }]} />
             <Stack horizontal>
diff --git a/Engine/Source/Programs/Horde/HordeDashboard/src/components/buildhealth/stepoutcome/StepOutcome.tsx b/Engine/Source/Programs/Horde/HordeDashboard/src/components/buildhealth/stepoutcome/StepOutcome.tsx
--- a/Engine/Source/Programs/Horde/HordeDashboard/src/components/buildhealth/stepoutcome/StepOutcome.tsx
+++ b/Engine/Source/Programs/Horde/HordeDashboard/src/components/buildhealth/stepoutcome/StepOutcome.tsx
@@ -6,7 +6,7 @@ import { observer } from "mobx-react-lite";
 import React, { useEffect, useState } from 'react';
 import { StepOutcomeD3Table } from "./StepOutcomeD3Component";
 import { prettyPrintStepOutcomeFilters, prettyPrintStepOutcomeFiltersSummary, StepOutcomeDataHandler, StepOutcomeFilters, StepOutcomeViewOptions } from "./StepOutcomeDataHandler";
-import { StepOutcomeTableEntry } from "./StepOutcomeDataTypes";
+import { AbortedStr, StepOutcomeTableEntry } from "./StepOutcomeDataTypes";
 import { sharedBorderStyle, stepOutcomeTableClasses } from "./StepOutcomeSharedUIComponents";
 import { isTerminalState, toDate } from "./StepOutcomeUtilities";
 
@@ -61,7 +61,8 @@ const StepOutcomeFocusModal: React.FC<StepOutcomeFocusModalProps> = ({ item, onD
         ['Step', item.stepResponse.name],
         ['Duration', item.getDurationString()],
         ['State', item.stepResponse.state],
-        isTerminalState(item) ? ['Outcome', item.stepResponse.outcome] : null,
+        // For simplicity, we will compeltely override the outcome of this step with the job's outcome, as the step was aborted.
+        item.stepResponse.state === AbortedStr ? ['Job Outcome', handler.getJobResponseData(item.jobId)?.summarizedOutcome ?? "-"] : isTerminalState(item) ? ['Outcome', item.stepResponse.outcome] : null,
         ['Job Link', <a href={`/job/${item.jobId}`} target="_blank" rel="noopener noreferrer">View Job</a>],
         ['Step Link', <a href={`/job/${item.jobId}?step=${item.stepResponse.id}`} target="_blank" rel="noopener noreferrer">View Step</a>],
     ].filter((d): d is (string | JSX.Element)[] => Boolean(d));
diff --git a/Engine/Source/Programs/Horde/HordeDashboard/src/components/buildhealth/stepoutcome/StepOutcomeD3Component.tsx b/Engine/Source/Programs/Horde/HordeDashboard/src/components/buildhealth/stepoutcome/StepOutcomeD3Component.tsx
--- a/Engine/Source/Programs/Horde/HordeDashboard/src/components/buildhealth/stepoutcome/StepOutcomeD3Component.tsx
+++ b/Engine/Source/Programs/Horde/HordeDashboard/src/components/buildhealth/stepoutcome/StepOutcomeD3Component.tsx
@@ -5,10 +5,11 @@ import * as d3 from "d3";
 import { getHordeStyling, preloadFonts } from "horde/styles/Styles";
 import { useEffect, useRef } from "react";
 import { StepOutcomeDataHandler } from "./StepOutcomeDataHandler";
-import { ChangeHeader, ColumnHeader, encodeStepNameFromStepNameHeader, isChangeHeader, isDateHeader, isStepOutcome, isSummary, isSummaryHeader, StepNameHeader, StepOutcomeTable, TableEntry } from "./StepOutcomeDataTypes";
+import { AbortedStr, ChangeHeader, ColumnHeader, encodeStepNameFromStepNameHeader, isChangeHeader, isDateHeader, isStepOutcome, isSummary, isSummaryHeader, StepNameHeader, StepOutcomeTable, TableEntry } from "./StepOutcomeDataTypes";
 import { getCellAccents, getCellClasses, getCellStyle, getColorRecords, IncrementalRefreshIndicatorPanel, resetCellClasses, sharedBorderStyle } from "./StepOutcomeSharedUIComponents";
 import { getHordeTheme } from "horde/styles/theme";
 import { observer } from "mobx-react-lite";
+import { isTerminalState, toLocalISOString } from "./StepOutcomeUtilities";
 
 // #region -- Styles --
 
@@ -326,7 +327,7 @@ function GenerateD3Table(stepOutcomeTable: StepOutcomeTable, handler: StepOutcom
                     tooltipContent.appendChild(dateUTCDiv);
 
                     const dateLocalDiv = document.createElement("div");
-                    dateLocalDiv.innerHTML = `<strong>Date (Local):</strong><p class="${tooltipClasses.tooltipLine}">${corrected.toString()}</p>`;
+                    dateLocalDiv.innerHTML = `<strong>Date (Local):</strong><p class="${tooltipClasses.tooltipLine}">${toLocalISOString(corrected)}</p>`;
                     tooltipContent.appendChild(dateLocalDiv);
 
                     attachTargettedToolTip(this as HTMLElement, scaffold.current!.tooltip, tooltipContent);
@@ -503,27 +504,48 @@ function GenerateD3Table(stepOutcomeTable: StepOutcomeTable, handler: StepOutcom
                         return;
                     }
                     if (isStepOutcome(entry)) { // StepOutcomeTableEntry
+                        // UE-324856 - This is a temporary fix until the horde server is fixed.
+                        const original = new Date(entry.changeDate!);
+                        const corrected = new Date(original.getTime() - 4 * 60 * 60 * 1000);
                         div.text(""),
                             div.style("cursor", "pointer")
                                 .style("background-color", getCellStyle(entry, stepOutcomeTable.viewModelOptions).bgColor)
-                                .style("box-shadow", getCellAccents(entry).boxShadow)
+                                .style("box-shadow", getCellAccents(entry, handler.getJobResponseData(entry.jobId)).boxShadow)
                                 .on("click", () => onCellClick?.(entry))
                                 .on("mouseover", function () {
                                     const tooltipContent = document.createElement("div");
                                     tooltipContent.style.paddingLeft = "4px";
 
                                     const changeDiv = document.createElement("div");
-                                    changeDiv.innerHTML = `<strong>Change:</strong><p class="${tooltipClasses.tooltipLineWithBorder}">${entry.change ?? "-"}</p>`;
+                                    changeDiv.innerHTML = `<strong>Change:</strong><p class="${tooltipClasses.tooltipLine}">${entry.change ?? "-"}</p>`;
                                     tooltipContent.appendChild(changeDiv);
 
+                                    const dateDiv = document.createElement("div");
+                                    dateDiv.innerHTML = `<strong>Date:</strong><p class="${tooltipClasses.tooltipLine}">(UTC) ${corrected.toISOString() ?? "-"}</p>`;
+                                    dateDiv.innerHTML += `<p class="${tooltipClasses.tooltipLineWithBorder}">(Local) ${toLocalISOString(corrected) ?? "-"}</p>`
+                                    tooltipContent.appendChild(dateDiv);
+
                                     const durationDiv = document.createElement("div");
                                     durationDiv.innerHTML = `<strong>Duration:</strong><p class="${tooltipClasses.tooltipLine}">${entry.getDurationString?.() ?? "-"}</p>`;
                                     tooltipContent.appendChild(durationDiv);
 
+
                                     const stateDiv = document.createElement("div");
-                                    stateDiv.innerHTML = `<strong>State:</strong><p class="${tooltipClasses.tooltipLineWithBorder}">${entry.stepResponse?.state ?? "-"}</p>`;
+                                    stateDiv.innerHTML = `<strong>State:</strong><p class="${tooltipClasses.tooltipLine}">${entry.stepResponse?.state ?? "-"}</p>`;
                                     tooltipContent.appendChild(stateDiv);
 
+                                    if (isTerminalState(entry)) {
+                                        const outcomeDiv = document.createElement("div");
+                                        outcomeDiv.innerHTML = `<strong>Outcome:</strong><p class="${tooltipClasses.tooltipLine}">${entry.stepResponse?.outcome ?? "-"}</p>`;
+
+                                        // For simplicity, we will compeltely override the outcome of this step with the job's outcome, as the step was aborted.
+                                        if (entry.stepResponse.state === AbortedStr) {
+                                            outcomeDiv.innerHTML = `<strong>Job Outcome:</strong><p class="${tooltipClasses.tooltipLine}">${handler.getJobResponseData(entry.jobId)?.summarizedOutcome ?? "-"}</p>`;
+                                        }
+
+                                        tooltipContent.appendChild(outcomeDiv);
+                                    }
+
                                     const streamDiv = document.createElement("div");
                                     streamDiv.innerHTML = `<strong>Stream:</strong><p class="${tooltipClasses.tooltipLine}">${entry.streamId}</p>`;
                                     tooltipContent.appendChild(streamDiv);
@@ -541,14 +563,14 @@ function GenerateD3Table(stepOutcomeTable: StepOutcomeTable, handler: StepOutcom
                                 .on("mouseout", () => {
                                     scaffold.current!.tooltip.style("opacity", 0);
                                 })
-                        getCellClasses(entry).forEach(cls => div.classed(cls, true));
+                        getCellClasses(entry, handler.getJobResponseData(entry.jobId)).forEach(cls => div.classed(cls, true));
                     } else if (isSummary(entry)) { // SummaryTableEntry
                         let ratio: number = 0;
 
                         if (entry.steps > 0) {
                             let numerator: number = entry.stepsPass
 
-                            if (stepOutcomeTable.viewModelOptions.warningsAsSummarySuccess) {
+                            if (!stepOutcomeTable.viewModelOptions.warningsAsSummaryFailure) {
                                 numerator += entry.stepsWarning;
                             }
 
@@ -685,7 +707,7 @@ function GenerateD3Table(stepOutcomeTable: StepOutcomeTable, handler: StepOutcom
                 container.scrollLeft = container.scrollWidth;
             });
         }
-    }, [handler.StepOutcomeTableData, stepOutcomeTable.viewModelOptions]);
+    }, [handler.StepOutcomeTableData, stepOutcomeTable.viewModelOptions?.includeDateAnchors]);
 
     const updateMaxHeight = () => {
         if (containerRef.current) {
diff --git a/Engine/Source/Programs/Horde/HordeDashboard/src/components/buildhealth/stepoutcome/StepOutcomeDataHandler.ts b/Engine/Source/Programs/Horde/HordeDashboard/src/components/buildhealth/stepoutcome/StepOutcomeDataHandler.ts
--- a/Engine/Source/Programs/Horde/HordeDashboard/src/components/buildhealth/stepoutcome/StepOutcomeDataHandler.ts
+++ b/Engine/Source/Programs/Horde/HordeDashboard/src/components/buildhealth/stepoutcome/StepOutcomeDataHandler.ts
@@ -2,10 +2,11 @@
 
 import { GetJobResponse, IssueQuery, JobStreamQuery } from "horde/backend/Api";
 import { PollBase } from "horde/backend/PollBase";
-import { StepOutcomeTable, StepOutcomeTableEntry, UpdatedChangeSummaryData } from "./StepOutcomeDataTypes";
+import { JobResponseData, StepOutcomeTable, StepOutcomeTableEntry, UpdatedChangeSummaryData } from "./StepOutcomeDataTypes";
 import backend from "horde/backend";
 import { decodeStepKey, decodeTemplateKey, encodeStepNameFromStrings } from "../BuildHealthUtilities";
 import { action, computed, makeObservable, observable, runInAction } from "mobx";
+import { summarizeJob } from "./StepOutcomeUtilities";
 
 // #region -- API Data Types & Utilities --
 
@@ -82,8 +83,7 @@ export function prettyPrintStepOutcomeFiltersSummary(filters: StepOutcomeFilters
  */
 export interface StepOutcomeViewOptions {
     includeDateAnchors: boolean;
-    skippedAsRed: boolean;
-    warningsAsSummarySuccess: boolean;
+    warningsAsSummaryFailure: boolean;
 }
 
 /**
@@ -121,7 +121,7 @@ export class StepOutcomeDataHandler extends PollBase {
     private activeRefresh: boolean = false;
 
     private activeFilter: StepOutcomeFilters;
-    private jobData: GetJobResponse[] = [];
+    private jobData: Map<string, JobResponseData> = new Map<string, JobResponseData>();
     private changeData: Map<number, UpdatedChangeSummaryData> = new Map<number, UpdatedChangeSummaryData>();
     private stepOutcomeTableData: StepOutcomeTable = new StepOutcomeTable();
 
@@ -235,8 +235,16 @@ export class StepOutcomeDataHandler extends PollBase {
             }
         });
 
-        // Flatten all jobData
-        this.jobData = results.flat();
+        // Flatten all jobData & memoize the summarized outcome
+        this.jobData = new Map<string, JobResponseData>(results.flat().map((job: GetJobResponse): [string, JobResponseData] => (
+            [
+                job.id,
+                {
+                    ...job,
+                    summarizedOutcome: summarizeJob(job)
+                }
+            ]
+        )));
 
         if (jobHistoryCount.size > 0) {
             // Parallelize backend calls for change summaries
@@ -273,10 +281,7 @@ export class StepOutcomeDataHandler extends PollBase {
                             templateStepSet.size === 0 ||
                             templateStepSet.has(fullyQualifiedStepName)
                         ) {
-
-                            this.stepOutcomeTableData.addEntry(
-                                new StepOutcomeTableEntry(step, jobResponse.id, jobResponse.name, jobResponse.streamId, jobResponse.createTime.toString(), jobResponse.change), this.changeData.get(jobResponse.change!)
-                            );
+                            this.stepOutcomeTableData.addEntry(new StepOutcomeTableEntry(step, jobResponse.id, jobResponse.name, jobResponse.streamId, jobResponse.createTime.toString(), jobResponse.change, this.changeData.get(jobResponse.change!)?.dateUtc));
                         }
                     }
                 }
@@ -287,7 +292,7 @@ export class StepOutcomeDataHandler extends PollBase {
 
         this.StepOutcomeTableData.summarize();
 
-        console.info(`Retrieved: ${this.jobData.length} job instance results for StepOutcome query.`);
+        console.info(`Retrieved: ${this.jobData.size} job instance results for StepOutcome query.`);
 
         this.stepOutcomeTableData.orderTableDataByChange(true);
 
@@ -394,5 +399,14 @@ export class StepOutcomeDataHandler extends PollBase {
         this.pollTime = newRefreshTime;
     }
 
+    /**
+     * Obtains the corresponding JobResponseData given the job id.
+     * @param jobId The job id for the corresponding JobResponseData to retrieve.
+     * @returns The JobResponseData if it is present in the handler, undefined otherwise.
+     */
+    getJobResponseData(jobId: string): JobResponseData | undefined {
+        return this.jobData.get(jobId);
+    }
+
     // #endregion -- Public API  --
 }
\ No newline at end of file
diff --git a/Engine/Source/Programs/Horde/HordeDashboard/src/components/buildhealth/stepoutcome/StepOutcomeDataTypes.ts b/Engine/Source/Programs/Horde/HordeDashboard/src/components/buildhealth/stepoutcome/StepOutcomeDataTypes.ts
--- a/Engine/Source/Programs/Horde/HordeDashboard/src/components/buildhealth/stepoutcome/StepOutcomeDataTypes.ts
+++ b/Engine/Source/Programs/Horde/HordeDashboard/src/components/buildhealth/stepoutcome/StepOutcomeDataTypes.ts
@@ -1,6 +1,6 @@
 // Copyright Epic Games, Inc. All Rights Reserved.
 
-import { ChangeSummaryData, GetStepResponse, IssueData } from "horde/backend/Api";
+import { ChangeSummaryData, GetJobResponse, GetStepResponse, IssueData, JobStepOutcome } from "horde/backend/Api";
 import { computeDuration, DURATION_NOT_SET, formatDurationFromMs, toDate } from "./StepOutcomeUtilities";
 import { StepOutcomeViewOptions } from "./StepOutcomeDataHandler";
 import { action, makeObservable, observable } from "mobx";
@@ -11,6 +11,7 @@ export const SkippedStr: string = "Skipped";
 export const FailureStr: string = "Failure";
 export const WaitingStr: string = "Waiting";
 export const RunningStr: string = "Running";
+export const ReadyStr:   string = "Ready";
 export const AbortedStr: string = "Aborted";
 
 /**
@@ -63,6 +64,13 @@ export type StepReponseData = GetStepResponse & {
     stepDurationMs: number;
 }
 
+/*
+* Convenience type to fold in extra metadata relevant for the JobResponse.
+*/
+export type JobResponseData = GetJobResponse & {
+    summarizedOutcome : JobStepOutcome;
+}
+
 // #region -- Table Entry Types --
 
 /**
@@ -107,6 +115,7 @@ export class StepOutcomeTableEntry {
     // #region -- Public Members --
 
     change?: number;
+    changeDate: Date | string;
     stepResponse: StepReponseData;
     jobId: string;
     jobName: string;
@@ -126,12 +135,22 @@ export class StepOutcomeTableEntry {
      * @param streamId The owning streams's Id.
      * @param jobCreateTime The owning job's create time.
      * @param change The change associated with the stepResponse.
+     * @param changeDate The change submit date associated with the stepResponse.
      */
-    constructor(stepResponse: GetStepResponse, jobId: string, jobName: string, streamId: string, jobCreateTime: string, change?: number) {
+    constructor(stepResponse: GetStepResponse, jobId: string, jobName: string, streamId: string, jobCreateTime: string, change?: number, changeDate? : Date | string) {
         if (change !== undefined) {
             this.change = change;
         }
 
+        if(changeDate !== undefined)
+        {
+            this.changeDate = changeDate;
+        }
+        else
+        {
+            this.changeDate = jobCreateTime;
+        }
+
         this.stepResponse = { ...stepResponse, stepDurationMs: computeDuration(stepResponse.startTime, stepResponse.finishTime) };
         this.jobId = jobId;
         this.jobName = jobName;
@@ -370,10 +389,9 @@ export class StepOutcomeTable {
     /**
      * Adds a new entry to the table data.
      * @param tableEntry The table entry to add to the data table.
-     * @param changeSummaryData The change summary data corresponding to the @see tableEntry.
      * @returns If tableEntry parameter is invalid, returns.
      */
-    addEntry(tableEntry: StepOutcomeTableEntry, changeSummaryData: UpdatedChangeSummaryData | undefined): void {
+    addEntry(tableEntry: StepOutcomeTableEntry): void {
         if (tableEntry.change === undefined) {
             return;
         }
@@ -394,7 +412,7 @@ export class StepOutcomeTable {
         if (insertionCol === undefined) {
             insertionCol = this.changeColHeaders.length;
             this.changeColLookup.set(tableEntry.change, insertionCol);
-            let changeDate = changeSummaryData?.dateUtc ? changeSummaryData?.dateUtc : tableEntry.jobCreateTime;
+            let changeDate = tableEntry.changeDate;
             this.changeColHeaders.push({ type: "change", change: tableEntry.change, date: toDate(changeDate)?.toISOString() });
 
             for (const row of this.tableEntries) {
diff --git a/Engine/Source/Programs/Horde/HordeDashboard/src/components/buildhealth/stepoutcome/StepOutcomeSharedUIComponents.tsx b/Engine/Source/Programs/Horde/HordeDashboard/src/components/buildhealth/stepoutcome/StepOutcomeSharedUIComponents.tsx
--- a/Engine/Source/Programs/Horde/HordeDashboard/src/components/buildhealth/stepoutcome/StepOutcomeSharedUIComponents.tsx
+++ b/Engine/Source/Programs/Horde/HordeDashboard/src/components/buildhealth/stepoutcome/StepOutcomeSharedUIComponents.tsx
@@ -4,7 +4,8 @@ import { mergeStyleSets, Spinner, SpinnerSize } from "@fluentui/react";
 import dashboard, { StatusColor } from "horde/backend/Dashboard";
 import { observer } from "mobx-react-lite";
 import { StepOutcomeDataHandler, StepOutcomeViewOptions } from "./StepOutcomeDataHandler";
-import { AbortedStr, FailureStr, RunningStr, SkippedStr, StepOutcomeTableEntry, SuccessStr, WarningStr } from "./StepOutcomeDataTypes";
+import { AbortedStr, FailureStr, JobResponseData, RunningStr, SkippedStr, StepOutcomeTableEntry, SuccessStr, WarningStr } from "./StepOutcomeDataTypes";
+import { isTerminalState } from "./StepOutcomeUtilities";
 
 // #region -- Constants --
 
@@ -20,37 +21,52 @@ export const sharedBorderStyle = {
     }
 };
 
+// Example: your function returns something like
+// { redX: "red", greenX: "green", blueX: "blue" }
+const colors: Record<string, string> = getColorRecords();
+
+const baseX = {
+    position: "relative",
+    overflow: "visible",
+    selectors: {
+        "&::before": {
+            content: "''",
+            position: "absolute",
+            left: "50%",
+            top: "50%",
+            width: "calc(80% * 1.414)",
+            height: "4px",
+            transform: "translate(-50%, -50%) rotate(45deg)",
+            transformOrigin: "center",
+            pointerEvents: "none",
+        },
+    },
+};
+
 export const stepOutcomeTableClasses = mergeStyleSets({
     root: {
         maxWidth: "min(85vw, 2400px)",
     },
-    redX: {
-        position: "relative",
-        overflow: "visible",
-        selectors: {
-            "&::before, &::after": {
-                content: "''",
-                position: "absolute",
-                left: "50%",
-                top: "50%",
-                width: "140%",
-                height: "2px",
-                background: "red",
-                transformOrigin: "center",
-                pointerEvents: "none",
-            },
-            "&::before": {
-                transform: "translate(-50%, -50%) rotate(45deg)",
+    ...Object.fromEntries(
+        Object.entries(colors).map(([className, color]) => [
+            className, // whatever key your getColorRecords() provides
+            {
+                ...baseX,
+                selectors: {
+                    "&::before": {
+                        ...baseX.selectors["&::before"],
+                        background: color,
+                    },
+                },
             },
-            "&::after": {
-                transform: "translate(-50%, -50%) rotate(-45deg)",
-            },
-        },
-    },
+        ])
+    ),
 });
 
 // #endregion -- Styles --
 
+// #region -- Style Helpers --
+
 /**
  * Gets the color records based off of the string representations of @see StatusColor - which maps to the status string.
  * @returns A record that maps sates to their corresponding colours. 
@@ -76,44 +92,48 @@ export function getColorRecords(): Record<string, string> {
  * @returns The list of all conditional classes to remove from the cell.
  */
 export function resetCellClasses(): string[] {
-    const allClasses = [
-        stepOutcomeTableClasses.redX,
-    ];
-
-    return allClasses;
+    return Object.values(stepOutcomeTableClasses) as string[];
 }
 
 /**
- * Gets the classes for a step outcome table entry cell.
+ * Gets the classes for a step outcome table entry cell. These are used to further highlight steps that have not run to completion (Aborted OR Skipped), and are not running.
  * @param entry The step outcome table entry of which to base the classes on.
+ * @param jobState The job response data which will be used to style applicable table entries.
  * @returns The classes to apply to the cell.
  */
-export function getCellClasses(entry: StepOutcomeTableEntry | null): string[] {
+export function getCellClasses(entry: StepOutcomeTableEntry | null,  jobState: JobResponseData | undefined): string[] {
     const classes: string[] = [];
 
-    if (!entry) return classes;
+    if (!entry || jobState?.summarizedOutcome === undefined) {
+        return classes;
+    }
 
     if (entry.stepResponse.state === AbortedStr) {
-        classes.push(stepOutcomeTableClasses.redX);
+        classes.push(stepOutcomeTableClasses[jobState?.summarizedOutcome]);
     }
 
     return classes;
 }
 
 /**
- * Gets the accents for a step outcome table entry cell.
+ * Gets the accents for a step outcome table entry cell. These are used to highlight steps that are currently in flight, or have not run to completion (Aborted OR Skipped).
  * @param entry The step outcome table entry of which to base the accents on.
+ * @param jobState The job response data which will be used to style applicable table entries.
  * @returns The accents to apply to the cell.
  */
-export function getCellAccents(entry: StepOutcomeTableEntry | null): { boxShadow: string } {
+export function getCellAccents(entry: StepOutcomeTableEntry | null, jobState: JobResponseData | undefined): { boxShadow: string } {
     const standard = { boxShadow: "none" };
-
-    if (!entry) {
+    if (!entry || jobState?.summarizedOutcome === undefined) {
         return standard;
     }
+    const colors: Record<string, string> = getColorRecords();
 
+    // If we have skipped or aborted the step, we use the summarized outcome to help with coloring.
     if (entry.stepResponse.state == SkippedStr || entry.stepResponse.state == AbortedStr) {
-        return { boxShadow: "inset 0 0 0 4px black" };
+        return { boxShadow: `inset 0 0 0 4px ${colors[jobState?.summarizedOutcome]}` };
+    }
+    else if (!isTerminalState(entry)) {
+        return { boxShadow: `inset 0 0 0 4px ${colors[entry.stepResponse.state]}` };
     }
 
     return standard;
@@ -136,11 +156,17 @@ export function getCellStyle(entry: StepOutcomeTableEntry | null, viewModelOptio
     const { stepResponse } = entry;
     const { outcome, state } = stepResponse;
 
+    // We do not fill the cells for Aborted, Skipped, or those about to run (or currently running).
+    if (state === AbortedStr || state === SkippedStr || !isTerminalState(entry)) {
+        return { bgColor: resultColor };
+    }
+
     // There are some interesting mappings going on internally with respect to state & outcome. Outcome & State have overlapping "result view" space, and as such we need to 
     // narrowly select the edge cases (e.g. outcome=success && state=completed uniquely signifies a completed run with no errors, otherwise outcome=success it could mean ready, waiting, or running).
     // Then we can fall back to the state as the truth.
+
     if (outcome === FailureStr) {
-        resultColor = (state === SkippedStr && !viewModelOptions.skippedAsRed) ? colors[SkippedStr] : colors[FailureStr];
+        resultColor = colors[FailureStr];
     }
     // A step can be emitting warnings whilst still running, so we narrowly select this colouring.
     else if (outcome === WarningStr) {
@@ -158,6 +184,8 @@ export function getCellStyle(entry: StepOutcomeTableEntry | null, viewModelOptio
     return { bgColor: resultColor };
 }
 
+// #endregion -- Style Helpers --
+
 /**
  * React Component that provides refresh spinner on incremental refresh. 
  * Note: This is a separate component to prevent rerender of primary table @see GenerateFluentUITable upon refresh spinner.
diff --git a/Engine/Source/Programs/Horde/HordeDashboard/src/components/buildhealth/stepoutcome/StepOutcomeUtilities.ts b/Engine/Source/Programs/Horde/HordeDashboard/src/components/buildhealth/stepoutcome/StepOutcomeUtilities.ts
--- a/Engine/Source/Programs/Horde/HordeDashboard/src/components/buildhealth/stepoutcome/StepOutcomeUtilities.ts
+++ b/Engine/Source/Programs/Horde/HordeDashboard/src/components/buildhealth/stepoutcome/StepOutcomeUtilities.ts
@@ -1,6 +1,7 @@
 // Copyright Epic Games, Inc. All Rights Reserved.
 
-import { RunningStr, StepOutcomeTableEntry, WaitingStr } from "./StepOutcomeDataTypes";
+import { GetJobResponse, JobStepOutcome, JobStepState } from "horde/backend/Api";
+import { ReadyStr, RunningStr, StepOutcomeTableEntry, WaitingStr } from "./StepOutcomeDataTypes";
 
 export const DURATION_NOT_SET: number = -1;
 
@@ -10,14 +11,60 @@ export const DURATION_NOT_SET: number = -1;
  * @returns True if the step has run to completion, false otherwise.
  */
 export function isTerminalState(stepOutcomeEntry: StepOutcomeTableEntry): boolean {
-    if(stepOutcomeEntry.stepResponse.state === RunningStr || stepOutcomeEntry.stepResponse.state === WaitingStr)
-    {
+    if (stepOutcomeEntry.stepResponse.state === ReadyStr || stepOutcomeEntry.stepResponse.state === RunningStr || stepOutcomeEntry.stepResponse.state === WaitingStr) {
         return false;
     }
 
     return true;
 }
 
+/**
+  * This utility method will summarize the provided job, with respect to the most severe step outcome encountered thus far.
+  * @param job The job to summarize.
+  * @returns The most severe job step outcome. This is ordered as follows: @see JobStepOutcome.Unspecified , @see JobStepOutcome.Success , @see JobStepOutcome.Warnings , @see JobStepOutcome.Failre .
+  */
+export function summarizeJob(job: GetJobResponse): JobStepOutcome {
+    let accumulatorOutcome: JobStepOutcome = JobStepOutcome.Unspecified;
+
+    if (job.batches === undefined) {
+        return accumulatorOutcome;
+    }
+    accumulatorOutcome = JobStepOutcome.Success;
+    for (let response of job.batches) {
+        for (let step of response.steps) {
+            if (step.state !== JobStepState.Aborted && step.outcome === JobStepOutcome.Failure) {
+                accumulatorOutcome = JobStepOutcome.Failure;
+
+                // We cannot become less severe than Failure; exit.
+                return accumulatorOutcome;
+            } else if (step.outcome === JobStepOutcome.Warnings) {
+                accumulatorOutcome = JobStepOutcome.Warnings;
+            }
+        }
+    }
+
+    return accumulatorOutcome;
+}
+
+/**
+ * Produces a ISO compliant string, in the provided locale.
+ * @param date The date to convert.
+ * @returns The string representation of the date, in ISO standard.
+ */
+export function toLocalISOString(date: Date): string {
+    const pad = (n: number) => n.toString().padStart(2, "0");
+
+    return (
+        date.getFullYear() +
+        "-" + pad(date.getMonth() + 1) +
+        "-" + pad(date.getDate()) +
+        "T" + pad(date.getHours()) +
+        ":" + pad(date.getMinutes()) +
+        ":" + pad(date.getSeconds()) +
+        "." + date.getMilliseconds().toString().padStart(3, "0")
+    );
+}
+
 /**
  * Normalizes an ambiguous input (Date | string) to a Date object. 
  * @param input the input string union type.
