[Horde][Analytics][Server] Fixing telemetry to use a string version of the commit ID to avoid complex objects.

diff --git a/Engine/Source/Programs/Horde/Plugins/Build/HordeServer.Build/Jobs/JobStepRefCollection.cs b/Engine/Source/Programs/Horde/Plugins/Build/HordeServer.Build/Jobs/JobStepRefCollection.cs
--- a/Engine/Source/Programs/Horde/Plugins/Build/HordeServer.Build/Jobs/JobStepRefCollection.cs
+++ b/Engine/Source/Programs/Horde/Plugins/Build/HordeServer.Build/Jobs/JobStepRefCollection.cs
@@ -172,7 +172,7 @@ namespace HordeServer.Jobs
 					AgentId = agentId,
 					BatchInitTime = initTime,
 					BatchWaitTime = waitTime,
-					Change = commitId,
+					Change = commitId?.Name,
 					FinishTime = finishTimeUtc,
 					JobName = jobName,
 					JobStartTime = jobStartTimeUtc,
diff --git a/Engine/Source/Programs/Horde/Plugins/Build/HordeServer.Build/Jobs/JobTelemetry.cs b/Engine/Source/Programs/Horde/Plugins/Build/HordeServer.Build/Jobs/JobTelemetry.cs
--- a/Engine/Source/Programs/Horde/Plugins/Build/HordeServer.Build/Jobs/JobTelemetry.cs
+++ b/Engine/Source/Programs/Horde/Plugins/Build/HordeServer.Build/Jobs/JobTelemetry.cs
@@ -1,6 +1,5 @@
 // Copyright Epic Games, Inc. All Rights Reserved.
 
-using EpicGames.Horde.Commits;
 using EpicGames.Horde.Jobs;
 using EpicGames.Horde.Jobs.Templates;
 using EpicGames.Horde.Streams;
@@ -35,12 +34,12 @@ namespace HordeServer.Jobs
 		/// <summary>
 		/// The commit of the job. <see cref="IJob.CommitId"/>.
 		/// </summary>
-		public CommitIdWithOrder CommitId { get; init; }
+		public string CommitId { get; init; }
 
 		/// <summary>
 		/// The code commit of the job. <see cref="IJob.CodeCommitId"/>.
 		/// </summary>
-		public CommitIdWithOrder CodeCommitId { get; init; }
+		public string CodeCommitId { get; init; }
 
 		/// <summary>
 		/// The created time of the job.
@@ -123,8 +122,8 @@ namespace HordeServer.Jobs
 			JobId = job.Id;
 			TemplateId = job.TemplateId;
 			StreamId = job.StreamId;
-			CommitId = job.CommitId;
-			CodeCommitId = job.CodeCommitId ?? CommitIdWithOrder.Empty;
+			CommitId = job.CommitId?.Name ?? "0";
+			CodeCommitId = job.CodeCommitId?.Name ?? "0";
 			CreateTimeUtc = job.CreateTimeUtc;
 			FinishTimeUtc = finishTimeUtc;
 			JobWallTime = (float)(FinishTimeUtc - CreateTimeUtc).TotalSeconds;
