Handle large number of Devices in phase event component
+ prune identical device information

#rnx

diff --git a/Engine/Source/Programs/Horde/HordeDashboard/plugins/test-automation/testAutomationEvents.tsx b/Engine/Source/Programs/Horde/HordeDashboard/plugins/test-automation/testAutomationEvents.tsx
--- a/Engine/Source/Programs/Horde/HordeDashboard/plugins/test-automation/testAutomationEvents.tsx
+++ b/Engine/Source/Programs/Horde/HordeDashboard/plugins/test-automation/testAutomationEvents.tsx
@@ -3,7 +3,7 @@
 import { Checkbox, DefaultButton, DirectionalHint, FontIcon, HoverCard, HoverCardType, Image, IPlainCardProps, Link, mergeStyles, mergeStyleSets, Spinner, SpinnerSize, Stack, Text } from "@fluentui/react";
 import { Artifacts, EmbeddedHTML, EventEntry, EventLevel, ImageCompare, LogPropertyType, URLLink } from "./testEventsModel";
 import { TestPhaseOutcome } from "./api";
-import { memo, useCallback, useEffect, useState } from "react";
+import { memo, useCallback, useEffect, useMemo, useState } from "react";
 import { ArtifactFactory, DeviceRef, TestPhaseStatus } from "./testData";
 import dashboard, { StatusColor } from "horde/backend/Dashboard";
 import { getShortNiceTime, msecToElapsed } from "horde/base/utilities/timeUtils";
@@ -161,7 +161,7 @@ const OnDeviceCard = (device: DeviceRef): JSX.Element => {
         </Stack>;
 }
 
-const DeviceItem = (device: DeviceRef): JSX.Element => {
+const DeviceItem = (device: DeviceRef, index: number): JSX.Element => {
     const deviceCardProps: IPlainCardProps = {
         renderData: device,
         onRenderPlainCard: OnDeviceCard,
@@ -170,8 +170,8 @@ const DeviceItem = (device: DeviceRef): JSX.Element => {
         calloutProps: { isBeakVisible: true }
     };
 
-    return <HoverCard plainCardProps={deviceCardProps} type={HoverCardType.plain} cardOpenDelay={20} key={device.name}>
-                <Text style={{fontSize: 12, cursor: 'pointer' }}>{device.name}</Text>
+    return <HoverCard plainCardProps={deviceCardProps} type={HoverCardType.plain} cardOpenDelay={20} key={`${device.name}#${index}`}>
+                <Text style={{fontSize: 12, cursor: 'pointer', paddingLeft: 3, paddingRight: 3}} styles={{root:{selectors: {':hover': {backgroundColor: "rgba(125, 125, 125, 0.3)"}}}}}>{device.name}</Text>
             </HoverCard>
 }
 
@@ -259,6 +259,10 @@ export const PhaseEventsPane: React.FC<{phase: TestPhaseStatus, onEventsLoaded?:
         return false;
     }, [filterError, filterWarning, filterInfo]);
 
+    const uniqueDevices = useMemo(() => {
+        return phase.devices.filter((device, index) => phase.devices.findIndex(other => device.isSame(other)) === index);
+    },[phase]);
+
     const borderColor = dashboard.darktheme ? "#4D4C4B" : "#6D6C6B";
 
     return <Stack style={{borderWidth: 1, borderColor: borderColor, borderStyle: 'solid', padding: 4, overflow: 'hidden'}} grow>
@@ -275,10 +279,13 @@ export const PhaseEventsPane: React.FC<{phase: TestPhaseStatus, onEventsLoaded?:
                         <Stack style={{borderWidth: '0px 0px 1px 0px', borderColor: borderColor, borderStyle: 'solid', paddingBottom: 2}} horizontal tokens={{childrenGap: 8}} disableShrink>
                             <Text><span style={styles.labelSmall}>Run on: </span><span>{getShortNiceTime(phase.start, true, true, true)}</span></Text>
                             <Text><span style={styles.labelSmall}> For: </span><span>{msecToElapsed(phase.duration * 1000, true, true)}</span></Text>
-                            {!!phase.devices.length &&
+                            {!!uniqueDevices.length &&
                                 <Stack horizontal tokens={{childrenGap: 4}} verticalAlign="end">
                                     <Text style={styles.labelSmall}> Devices: </Text>
-                                    {phase.devices.map(d => DeviceItem(d))}
+                                    {uniqueDevices.splice(0, 5).map((d, index) => DeviceItem(d, index))}
+                                    {uniqueDevices.length > 5 &&
+                                        <Text style={{userSelect: 'none'}}>...</Text>
+                                    }
                                 </Stack>
                             }
                             <Stack verticalAlign="center" horizontal tokens={{ childrenGap: 5 }} style={{borderWidth: '0px 1px', paddingLeft: 9, paddingRight: 9, borderStyle: 'solid', borderColor: borderColor}}>
diff --git a/Engine/Source/Programs/Horde/HordeDashboard/plugins/test-automation/testData.ts b/Engine/Source/Programs/Horde/HordeDashboard/plugins/test-automation/testData.ts
--- a/Engine/Source/Programs/Horde/HordeDashboard/plugins/test-automation/testData.ts
+++ b/Engine/Source/Programs/Horde/HordeDashboard/plugins/test-automation/testData.ts
@@ -247,7 +247,23 @@ export class DeviceRef {
         this.metadata = metadata ?? {};
     }
 
-    name: MetaId;
+    isSame(other: DeviceRef): boolean {
+
+        if (this.name === other.name)
+        {
+            const metaKeys = new Set(Object.keys(this.metadata));
+            const otherMetaKeys = new Set(Object.keys(other.metadata));
+            if (metaKeys.size === otherMetaKeys.size && metaKeys.difference(otherMetaKeys).size === 0) {
+                const metaValues = new Set(Object.values(this.metadata));
+                const otherMetaValues = new Set(Object.values(other.metadata));
+                return metaValues.difference(otherMetaValues).size === 0;
+            }
+        }
+
+        return false;
+    }
+
+    name: string;
     version: number;
     logPath?: string;
     metadata: Metadata;
