Horde Dashboard: Add ellipsis to job operation menu items that seem risky

diff --git a/Engine/Source/Programs/Horde/HordeDashboard/src/components/JobOperationsContextMenu.tsx b/Engine/Source/Programs/Horde/HordeDashboard/src/components/JobOperationsContextMenu.tsx
--- a/Engine/Source/Programs/Horde/HordeDashboard/src/components/JobOperationsContextMenu.tsx
+++ b/Engine/Source/Programs/Horde/HordeDashboard/src/components/JobOperationsContextMenu.tsx
@@ -97,7 +97,7 @@ export const JobOperationsContextMenu: React.FC<{ job: GetJobResponse, children?
    if (!abortDisabled) {
       menuItems.push({
          key: 'abort_job',
-         text: 'Cancel Job',
+         text: 'Cancel Job...',
          onClick: () => setAbortShown(true),
       });
    }
diff --git a/Engine/Source/Programs/Horde/HordeDashboard/src/components/jobDetailsV2/JobOperationsBar.tsx b/Engine/Source/Programs/Horde/HordeDashboard/src/components/jobDetailsV2/JobOperationsBar.tsx
--- a/Engine/Source/Programs/Horde/HordeDashboard/src/components/jobDetailsV2/JobOperationsBar.tsx
+++ b/Engine/Source/Programs/Horde/HordeDashboard/src/components/jobDetailsV2/JobOperationsBar.tsx
@@ -130,7 +130,7 @@ export const JobOperations: React.FC<{ jobDetails: JobDetailsV2 }> = observer(({
    }
    opsList.push({
       key: 'jobops_edit',
-      text: "Edit",
+      text: "Edit...",
       disabled: !jobData.canUpdate,
       iconProps: { iconName: "Edit" },
       onClick: () => { setEditShown(true); }
@@ -149,7 +149,7 @@ export const JobOperations: React.FC<{ jobDetails: JobDetailsV2 }> = observer(({
 
    opsList.push({
       key: 'jobops_abort',
-      text: "Cancel Job",
+      text: "Cancel Job...",
       disabled: abortDisabled,
       iconProps: { iconName: "Delete" },
       onClick: () => { setAbortShown(true); }
@@ -157,7 +157,7 @@ export const JobOperations: React.FC<{ jobDetails: JobDetailsV2 }> = observer(({
 
    opsList.push({
       key: 'jobops_runagain',
-      text: "Run Again",
+      text: "Run Again...",
       iconProps: { iconName: "Duplicate" },
       disabled: !canRun,
       onClick: () => { setParametersState(ParameterState.Clone); }
@@ -435,7 +435,7 @@ const StepOperations: React.FC<{ jobDetails: JobDetailsV2, stepId: string }> = o
 
    opsList.push({
       key: 'stepops_abort',
-      text: "Cancel Step",
+      text: "Cancel Step...",
       iconProps: { iconName: "Cross" },
       disabled: !!step.finishTime || !jobData.canUpdate,
       onClick: () => { setShown({ abortShown: true }); }
