Horde: Respect Horde-based views in PerforceMaterializer

diff --git a/Engine/Source/Programs/Horde/Drivers/JobDriver.Tests/PerforceMaterializerTests.cs b/Engine/Source/Programs/Horde/Drivers/JobDriver.Tests/PerforceMaterializerTests.cs
--- a/Engine/Source/Programs/Horde/Drivers/JobDriver.Tests/PerforceMaterializerTests.cs
+++ b/Engine/Source/Programs/Horde/Drivers/JobDriver.Tests/PerforceMaterializerTests.cs
@@ -47,6 +47,36 @@ public class PerforceMaterializerTests : BasePerforceFixtureTest
 		PerforceMaterializer.State? state = await pm.LoadStateForTestAsync(CancellationToken.None);
 		Assert.AreEqual(7, state!.Changelist);
 	}
+	
+	[TestMethod]
+	public async Task SyncWithViewAsync()
+	{
+		using PerforceMaterializer pm = CreatePm(Fixture.StreamFooMain, view: ["-/Data/...", "/Data/moredata.txt"]);
+		
+		await pm.SyncAsync(7, -1, new SyncOptions(), CancellationToken.None);
+		ChangelistFixture cl = Fixture.StreamFooMain.GetChangelist(7);
+		ChangelistFixture clWithView = new (cl.Number, cl.Description, cl.StreamFiles.Where(x => x.ClientFile != "Data/data.txt").ToList());
+		clWithView.AssertDepotFiles(pm.DirectoryPath.FullName);
+		await pm.FinalizeAsync(CancellationToken.None);
+		
+		PerforceMaterializer.State? state = await pm.LoadStateForTestAsync(CancellationToken.None);
+		Assert.AreEqual(7, state!.Changelist);
+	}
+	
+	private const string Stream = "//Foo/Main";
+	private const string ClientName = "MyComputer";
+	
+	[TestMethod]
+	[DataRow("Platforms/...", "//Foo/Main/Platforms/... //MyComputer/Platforms/...")]
+	[DataRow(" Platforms/...", "//Foo/Main/Platforms/... //MyComputer/Platforms/...")]
+	[DataRow("/Engine/Build/...", "//Foo/Main/Engine/Build/... //MyComputer/Engine/Build/...")]
+	[DataRow("-/Engine/Docs/...", "-//Foo/Main/Engine/Docs/... //MyComputer/Engine/Docs/...")]
+	[DataRow("   -/Engine/Docs/...", "-//Foo/Main/Engine/Docs/... //MyComputer/Engine/Docs/...")]
+	public void ConvertHordeViews(string hordeView, string expected)
+	{
+		string actual = PerforceMaterializer.ConvertHordeViewToClientView(hordeView, Stream, ClientName);
+		Assert.AreEqual(expected, actual);
+	}
 
 	[TestMethod]
 	public async Task TestAllStatesAsync()
@@ -146,7 +176,7 @@ public class PerforceMaterializerTests : BasePerforceFixtureTest
 		Assert.AreEqual(PerforceMaterializer.TransactionStatus.Clean, pmState.Status);
 	}
 
-	private PerforceMaterializer CreatePm(StreamFixture stream, string? dataDir = null, string? identifier = "pmTest")
+	private PerforceMaterializer CreatePm(StreamFixture stream, string? dataDir = null, string? identifier = "pmTest", string[]? view = null)
 	{
 		ILogger logger = LoggerFactory.CreateLogger("PM");
 		RpcAgentWorkspace raw = new ()
@@ -157,6 +187,7 @@ public class PerforceMaterializerTests : BasePerforceFixtureTest
 			Identifier = $"{TestGuid}-{identifier}",
 			Stream = stream.Root,
 		};
+		raw.View.AddRange(view ?? []);
 		PerforceMaterializerOptions options = new(dataDir ?? TempDir.FullName, raw);
 		return new PerforceMaterializer(options, NoOpTracer, logger);
 	}
diff --git a/Engine/Source/Programs/Horde/Drivers/JobDriver/Execution/PerforceMaterializer.cs b/Engine/Source/Programs/Horde/Drivers/JobDriver/Execution/PerforceMaterializer.cs
--- a/Engine/Source/Programs/Horde/Drivers/JobDriver/Execution/PerforceMaterializer.cs
+++ b/Engine/Source/Programs/Horde/Drivers/JobDriver/Execution/PerforceMaterializer.cs
@@ -274,13 +274,25 @@ public sealed class PerforceMaterializer : IWorkspaceMaterializer
 		string clientName = _perforceSettings.ClientName!;
 		string streamName = _options.AgentWorkspace.Stream;
 		
+		// Use a placeholder name to avoid conflicting with any existing client name
+		const string PlaceholderName = "horde-placeholder-name-pm";
+		
+		// Assume latest stream version is wanted. Should possibly be tied to current CL being synced
+		ClientRecord streamView = await _perforceWithoutClient.GetStreamViewAsync(PlaceholderName, streamName, -1, cancellationToken);
+		
+		List<string> hordeView = _options.AgentWorkspace.View.Select(x => ConvertHordeViewToClientView(x, streamName, clientName)).ToList();
+		List<string> view = [..streamView.View.Select(x => x.Replace(PlaceholderName, clientName, StringComparison.Ordinal)), ..hordeView];
+		
+		// Create a classic client while still applying the view from stream together with extra Horde-defined view entries.
+		// A stream-based client won't allow changing the view
 		ClientRecord newClient = new (clientName, _perforceSettings.UserName, _workspaceDir.FullName)
 		{
 			Host = _options.HostName ?? Environment.MachineName,
-			Stream = streamName,
 			Type = "partitioned",
 			Root = _workspaceDir.FullName,
+			Options = ClientOptions.Clobber,
 		};
+		newClient.View.AddRange(view);
 		Directory.CreateDirectory(newClient.Root);
 		
 		using ILoggerProgress status = _logger.BeginProgressScope(forceCreate ? "Force creating client ..." : "Updating client ...");
@@ -306,6 +318,21 @@ public sealed class PerforceMaterializer : IWorkspaceMaterializer
 		status.Progress = $"({timer.Elapsed.TotalSeconds:0.0}s)";
 		return newClient;
 	}
+	
+	internal static string ConvertHordeViewToClientView(string hordeView, string stream, string clientName)
+	{
+		string exclude = "";
+		string path = hordeView.TrimStart();
+		bool isExclude = path[0] == '-';
+		if (isExclude)
+		{
+			exclude = "-";
+			path = path[1..]; // Strip minus
+		}
+
+		path = path[0] == '/' ? path : '/' + path;
+		return $"{exclude}{stream}{path} //{clientName}{path}";
+	}
 
 	private Task SaveStateAsync(TransactionStatus status, string client ,int changelist, int shelvedChangelist, CancellationToken cancellationToken)
 	{
