Getting Horde Agent to work with an "Unreal root" hierarchy in a subfolder of the workspace root, instead of being the same as the workspace root. Leverages the "EnginePath" setting to consider the parent of the Engine folder as the Unreal root.

diff --git a/Engine/Source/Programs/Horde/Drivers/JobDriver/Execution/JobExecutor.cs b/Engine/Source/Programs/Horde/Drivers/JobDriver/Execution/JobExecutor.cs
--- a/Engine/Source/Programs/Horde/Drivers/JobDriver/Execution/JobExecutor.cs
+++ b/Engine/Source/Programs/Horde/Drivers/JobDriver/Execution/JobExecutor.cs
@@ -1192,6 +1192,7 @@ namespace JobDriver.Execution
 		private async Task<bool> ExecuteWithTempStorageAsync(JobStepInfo step, DirectoryReference workspaceDir, string arguments, bool? useP4, ILogger logger, CancellationToken cancellationToken)
 		{
 			DirectoryReference manifestDir = DirectoryReference.Combine(workspaceDir, $"{EnginePath}/Saved/BuildGraph");
+			DirectoryReference unrealRootDir = GetUnrealRootDirectoryFromWorkspace(workspaceDir);
 
 			// Create the mapping of tag names to file sets
 			Dictionary<string, HashSet<FileReference>> tagNameToFileSet = new Dictionary<string, HashSet<FileReference>>();
@@ -1211,7 +1212,7 @@ namespace JobDriver.Execution
 				string tagName = input.Substring(slashIdx + 1);
 
 				TempStorageTagManifest fileList = await TempStorage.RetrieveTagAsync(HordeClient, JobId, step.StepId, nodeName, tagName, manifestDir, logger, cancellationToken);
-				tagNameToFileSet[tagName] = fileList.ToFileSet(workspaceDir);
+				tagNameToFileSet[tagName] = fileList.ToFileSet(unrealRootDir);
 				inputStorageBlocks.UnionWith(fileList.Blocks);
 			}
 
@@ -1267,7 +1268,7 @@ namespace JobDriver.Execution
 				logger.LogInformation("Reading local file list from {File}", tagFileListLocation.FullName);
 
 				TempStorageTagManifest fileList = TempStorageTagManifest.Load(tagFileListLocation);
-				tagNameToFileSet[tagName] = fileList.ToFileSet(workspaceDir);
+				tagNameToFileSet[tagName] = fileList.ToFileSet(unrealRootDir);
 			}
 
 			// Check that none of the inputs have been clobbered
@@ -1275,7 +1276,7 @@ namespace JobDriver.Execution
 			foreach (TempStorageFile file in inputManifests.Values.SelectMany(x => x.Files))
 			{
 				string? message;
-				if (!modifiedFiles.ContainsKey(file.RelativePath) && !file.Compare(workspaceDir, out message))
+				if (!modifiedFiles.ContainsKey(file.RelativePath) && !file.Compare(unrealRootDir, out message))
 				{
 					modifiedFiles.Add(file.RelativePath, message);
 				}
@@ -1318,7 +1319,7 @@ namespace JobDriver.Execution
 				HashSet<FileReference> files = tagNameToFileSet[tagName];
 				foreach (FileReference file in files)
 				{
-					if (file.IsUnderDirectory(workspaceDir))
+					if (file.IsUnderDirectory(unrealRootDir))
 					{
 						if (isDefaultOutput)
 						{
@@ -1397,7 +1398,7 @@ namespace JobDriver.Execution
 						}
 						if (pair.Value.Any(x => referencedOutputFiles.Contains(x)))
 						{
-							outputNode.AddDirectory(await TempStorage.ArchiveBlockAsync(manifestDir, step.Name, pair.Key, workspaceDir, pair.Value.ToArray(), blobWriter, logger, cancellationToken));
+							outputNode.AddDirectory(await TempStorage.ArchiveBlockAsync(manifestDir, step.Name, pair.Key, unrealRootDir, pair.Value.ToArray(), blobWriter, logger, cancellationToken));
 						}
 					}
 
@@ -1416,7 +1417,7 @@ namespace JobDriver.Execution
 							}
 						}
 
-						outputNode.AddFile(await TempStorage.ArchiveTagAsync(manifestDir, step.Name, outputName, workspaceDir, files, storageBlocks.ToArray(), blobWriter, logger, cancellationToken));
+						outputNode.AddFile(await TempStorage.ArchiveTagAsync(manifestDir, step.Name, outputName, unrealRootDir, files, storageBlocks.ToArray(), blobWriter, logger, cancellationToken));
 					}
 
 					outputNodeRef = await blobWriter.WriteBlobAsync(outputNode, cancellationToken: cancellationToken);
@@ -1803,9 +1804,11 @@ namespace JobDriver.Execution
 
 		async Task<int> ExecuteCommandAsync(JobStepInfo step, DirectoryReference workspaceDir, string fileName, string arguments, TelemetrySpan buildGraphSpan, ILogger jobLogger, CancellationToken cancellationToken)
 		{
+			DirectoryReference unrealRootDir = GetUnrealRootDirectoryFromWorkspace(workspaceDir);
+
 			// Method for expanding environment variable properties related to this step
 			Dictionary<string, string> properties = new Dictionary<string, string>(StringComparer.OrdinalIgnoreCase);
-			properties.Add("RootDir", workspaceDir.FullName);
+			properties.Add("RootDir", unrealRootDir.FullName);
 
 			// Combine all the supplied environment variables together
 			Dictionary<string, string> newEnvVars = new Dictionary<string, string>(StringComparer.Ordinal);
@@ -2270,6 +2273,12 @@ namespace JobDriver.Execution
 				}
 			}
 		}
+
+		protected DirectoryReference GetUnrealRootDirectoryFromWorkspace(DirectoryReference workspaceDir)
+		{
+			// The Unreal root directory is the parent of the Engine directory
+			return DirectoryReference.Combine(workspaceDir, EnginePath).ParentDirectory!;
+		}
 	}
 
 	interface IJobExecutorFactory
diff --git a/Engine/Source/Programs/Horde/Drivers/JobDriver/Execution/PerforceExecutor.cs b/Engine/Source/Programs/Horde/Drivers/JobDriver/Execution/PerforceExecutor.cs
--- a/Engine/Source/Programs/Horde/Drivers/JobDriver/Execution/PerforceExecutor.cs
+++ b/Engine/Source/Programs/Horde/Drivers/JobDriver/Execution/PerforceExecutor.cs
@@ -150,7 +150,7 @@ namespace JobDriver.Execution
 
 			// Get all the environment variables for jobs
 			_envVars["IsBuildMachine"] = "1";
-			_envVars["uebp_LOCAL_ROOT"] = _workspace.WorkspaceDir.FullName;
+			_envVars["uebp_LOCAL_ROOT"] = GetUnrealRootDirectoryFromWorkspace(_workspace.WorkspaceDir).FullName;
 			_envVars["uebp_PORT"] = _workspace.ServerAndPort;
 			_envVars["uebp_USER"] = _workspace.UserName;
 			_envVars["uebp_CLIENT"] = _workspace.ClientName;
diff --git a/Engine/Source/Programs/Horde/Drivers/JobDriver/Execution/WorkspaceExecutor.cs b/Engine/Source/Programs/Horde/Drivers/JobDriver/Execution/WorkspaceExecutor.cs
--- a/Engine/Source/Programs/Horde/Drivers/JobDriver/Execution/WorkspaceExecutor.cs
+++ b/Engine/Source/Programs/Horde/Drivers/JobDriver/Execution/WorkspaceExecutor.cs
@@ -77,7 +77,7 @@ namespace JobDriver.Execution
 
 			// Set any non-materializer specific environment variables for jobs
 			_envVars["IsBuildMachine"] = "1";
-			_envVars["uebp_LOCAL_ROOT"] = _workspace.DirectoryPath.FullName;
+			_envVars["uebp_LOCAL_ROOT"] = GetUnrealRootDirectoryFromWorkspace(_workspace.DirectoryPath).FullName;
 			_envVars["uebp_BuildRoot_P4"] = Batch.StreamName;
 			_envVars["uebp_BuildRoot_Escaped"] = Batch.StreamName.Replace('/', '+');
 			_envVars["uebp_CL"] = Batch.Change.ToString();
