[Horde][BuildHealth][Dashboard] Added columns to CSV export to help with diagnosis.

diff --git a/Engine/Source/Programs/Horde/HordeDashboard/src/components/buildhealth/stepoutcome/StepOutcomeD3Component.tsx b/Engine/Source/Programs/Horde/HordeDashboard/src/components/buildhealth/stepoutcome/StepOutcomeD3Component.tsx
--- a/Engine/Source/Programs/Horde/HordeDashboard/src/components/buildhealth/stepoutcome/StepOutcomeD3Component.tsx
+++ b/Engine/Source/Programs/Horde/HordeDashboard/src/components/buildhealth/stepoutcome/StepOutcomeD3Component.tsx
@@ -43,7 +43,7 @@ function saveFile(content: string, filename: string) {
     URL.revokeObjectURL(url);
 }
 
-function exportStepOutcomeDataToCSV(stepOutcomeTable: StepOutcomeTable) {
+function exportStepOutcomeDataToCSV(stepOutcomeTable: StepOutcomeTable, handler : StepOutcomeDataHandler) {
     const defaultName = new Date()
         .toLocaleDateString("en-US", {
             year: "numeric",
@@ -56,8 +56,8 @@ function exportStepOutcomeDataToCSV(stepOutcomeTable: StepOutcomeTable) {
     const filename = window.prompt("Enter name of export:", "ProjectExport");
     if (!filename) return;
 
-    const stepOrderedCsvContent = stepOutcomeTable.toRowOrderedCSV();
-    const changeOrderedCsvContent = stepOutcomeTable.toColOrderedCSV();
+    const stepOrderedCsvContent = stepOutcomeTable.toRowOrderedCSV(handler);
+    const changeOrderedCsvContent = stepOutcomeTable.toColOrderedCSV(handler);
 
     saveFile(stepOrderedCsvContent, `${filename}-StepOrderedData-${defaultName}.csv`);
     saveFile(changeOrderedCsvContent, `${filename}-ChangeOrderedData-${defaultName}.csv`);
@@ -772,7 +772,7 @@ function GenerateD3Table(stepOutcomeTable: StepOutcomeTable, handler: StepOutcom
                                 lineHeight: "15px",
                             },
                         }}
-                        onClick={() => exportStepOutcomeDataToCSV(stepOutcomeTable)}
+                        onClick={() => exportStepOutcomeDataToCSV(stepOutcomeTable, handler)}
                     />
                 </TooltipHost>
             </Stack>
diff --git a/Engine/Source/Programs/Horde/HordeDashboard/src/components/buildhealth/stepoutcome/StepOutcomeDataTypes.ts b/Engine/Source/Programs/Horde/HordeDashboard/src/components/buildhealth/stepoutcome/StepOutcomeDataTypes.ts
--- a/Engine/Source/Programs/Horde/HordeDashboard/src/components/buildhealth/stepoutcome/StepOutcomeDataTypes.ts
+++ b/Engine/Source/Programs/Horde/HordeDashboard/src/components/buildhealth/stepoutcome/StepOutcomeDataTypes.ts
@@ -2,7 +2,7 @@
 
 import { ChangeSummaryData, GetJobResponse, GetStepResponse, IssueData, JobStepOutcome } from "horde/backend/Api";
 import { computeDuration, DURATION_NOT_SET, formatDurationFromMs, toDate } from "./StepOutcomeUtilities";
-import { StepOutcomeViewOptions } from "./StepOutcomeDataHandler";
+import { StepOutcomeDataHandler, StepOutcomeViewOptions } from "./StepOutcomeDataHandler";
 import { action, makeObservable, observable } from "mobx";
 
 export const WarningStr: string = "Warnings";
@@ -11,7 +11,7 @@ export const SkippedStr: string = "Skipped";
 export const FailureStr: string = "Failure";
 export const WaitingStr: string = "Waiting";
 export const RunningStr: string = "Running";
-export const ReadyStr:   string = "Ready";
+export const ReadyStr: string = "Ready";
 export const AbortedStr: string = "Aborted";
 
 /**
@@ -68,7 +68,7 @@ export type StepReponseData = GetStepResponse & {
 * Convenience type to fold in extra metadata relevant for the JobResponse.
 */
 export type JobResponseData = GetJobResponse & {
-    summarizedOutcome : JobStepOutcome;
+    summarizedOutcome: JobStepOutcome;
 }
 
 // #region -- Table Entry Types --
@@ -137,17 +137,15 @@ export class StepOutcomeTableEntry {
      * @param change The change associated with the stepResponse.
      * @param changeDate The change submit date associated with the stepResponse.
      */
-    constructor(stepResponse: GetStepResponse, jobId: string, jobName: string, streamId: string, jobCreateTime: string, change?: number, changeDate? : Date | string) {
+    constructor(stepResponse: GetStepResponse, jobId: string, jobName: string, streamId: string, jobCreateTime: string, change?: number, changeDate?: Date | string) {
         if (change !== undefined) {
             this.change = change;
         }
 
-        if(changeDate !== undefined)
-        {
+        if (changeDate !== undefined) {
             this.changeDate = changeDate;
         }
-        else
-        {
+        else {
             this.changeDate = jobCreateTime;
         }
 
@@ -178,10 +176,14 @@ export class StepOutcomeTableEntry {
 
     /**
      * Gets the CSV representation of the the table entry. Property ordering matches @see getCSVHeader .
+     * @param relatedJobContext The related job's context object.
      * @returns CSV representation of the table entry.
      */
-    toCSVRepresntation(): string {
-        return `${this.change},${this.streamId},${this.jobName},${this?.jobId},${this?.stepResponse.name},${this?.stepResponse.state},${this?.stepResponse.outcome},${this?.stepResponse.stepDurationMs / 1000}\n`;
+    toCSVRepresntation(relatedJobContext: JobResponseData | undefined): string {
+        let baseUrl =  window.location.origin;
+        let jobUrl = `${baseUrl}/job/${this.jobId}`;
+        let stepUrl = `${baseUrl}/job/${this.jobId}?step=${this.stepResponse.id}`;
+        return `${this.change},${this.streamId},${this.jobName},${this.jobId},${this.stepResponse.name},${this.stepResponse.id},${this.stepResponse.state},${this.stepResponse.outcome},${this.stepResponse.stepDurationMs / 1000},${this.jobCreateTime},${this.stepResponse.startTime},${relatedJobContext?.summarizedOutcome},${relatedJobContext?.state},${jobUrl},${stepUrl}\n`;
     }
 
     /**
@@ -189,7 +191,7 @@ export class StepOutcomeTableEntry {
      * @returns The comma separated properties serialization.
      */
     static getCSVHeader() {
-        return "Change,StreamId,JobName,JobId,StepName,StepState,StepOutcome,Duration(Sec)\n";
+        return "Change,StreamId,JobName,JobId,StepName,StepId,StepState,StepOutcome,Duration(Sec),JobStart,StepStart,JobSummaryOutcome,JobState,JobUrl,StepUrl\n";
     }
     // #endregion -- Public API --
 }
@@ -606,9 +608,10 @@ export class StepOutcomeTable {
 
     /**
      * Generates the csv string representation of the step outcome table, in row order.
+     * @param contextHandler The datahandler used to obtain job response data.
      * @returns A string in CSV format that represents the underlying step outcome table, row order.
      */
-    toRowOrderedCSV(): string {
+    toRowOrderedCSV(contextHandler: StepOutcomeDataHandler): string {
         let resultString: string = StepOutcomeTableEntry.getCSVHeader();
         for (let rowIndex = 0; rowIndex < this.tableEntries.length; rowIndex++) {
             const row = this.tableEntries[rowIndex];
@@ -617,7 +620,8 @@ export class StepOutcomeTable {
                     continue;
                 }
                 let tableEntry: StepOutcomeTableEntry = row[i]!;
-                resultString += tableEntry.toCSVRepresntation();
+                let jobResponseData: JobResponseData | undefined = contextHandler.getJobResponseData(tableEntry.jobId);
+                resultString += tableEntry.toCSVRepresntation(jobResponseData);
             }
         }
 
@@ -626,9 +630,10 @@ export class StepOutcomeTable {
 
     /**
      * Generates the csv string representation of the step outcome table, in column order.
+     * @param contextHandler The datahandler used to obtain job response data.
      * @returns A string in CSV format that represents the underlying step outcome table, column order.
      */
-    toColOrderedCSV(): string {
+    toColOrderedCSV(contextHandler: StepOutcomeDataHandler): string {
         let resultString: string = StepOutcomeTableEntry.getCSVHeader();
         for (let changeIdx = 0; changeIdx < this.changeColHeaders.length; changeIdx++) {
             const changeEntries = [...this.getChangeColumn(this.changeColHeaders[changeIdx].change)];
@@ -637,7 +642,8 @@ export class StepOutcomeTable {
                     continue;
                 }
                 let tableEntry: StepOutcomeTableEntry = changeEntries[i]!;
-                resultString += tableEntry.toCSVRepresntation();
+                let jobResponseData: JobResponseData | undefined = contextHandler.getJobResponseData(tableEntry.jobId);
+                resultString += tableEntry.toCSVRepresntation(jobResponseData);
             }
         }
 
