[Horde][BuildHealth][Telemetry] JobLabelSummary telemetry

diff --git a/Engine/Source/Programs/Horde/Plugins/Analytics/HordeServer.Analytics.Integration.Tests/Telemetry/TelemetryIntegrationTests.cs b/Engine/Source/Programs/Horde/Plugins/Analytics/HordeServer.Analytics.Integration.Tests/Telemetry/TelemetryIntegrationTests.cs
--- a/Engine/Source/Programs/Horde/Plugins/Analytics/HordeServer.Analytics.Integration.Tests/Telemetry/TelemetryIntegrationTests.cs
+++ b/Engine/Source/Programs/Horde/Plugins/Analytics/HordeServer.Analytics.Integration.Tests/Telemetry/TelemetryIntegrationTests.cs
@@ -15,6 +15,7 @@ using HordeCommon;
 using HordeServer.Analytics.Tests;
 using HordeServer.Issues;
 using HordeServer.Jobs;
+using HordeServer.Jobs.Graphs;
 using HordeServer.Logs;
 using HordeServer.Projects;
 using HordeServer.Server;
@@ -24,6 +25,7 @@ using HordeServer.Telemetry.Metrics;
 using HordeServer.Tests;
 using HordeServer.Tests.Issues;
 using Microsoft.Extensions.DependencyInjection;
+using Moq;
 
 namespace HordeServer.Analytics.Integration.Tests.Telemetry
 {
@@ -151,7 +153,7 @@ namespace HordeServer.Analytics.Integration.Tests.Telemetry
 
 			// Validate our preconditions
 			Assert.IsNotNull(concreteTelemetryWriter);
-			Assert.IsTrue(!concreteTelemetryWriter.RecordedTelemetry.ContainsKey(TelemetryStoreId.Default));
+			Assert.IsFalse(concreteTelemetryWriter.RecordedTelemetry.ContainsKey(TelemetryStoreId.Default));
 
 			// Generate issue telemetry data
 			{
@@ -229,22 +231,25 @@ namespace HordeServer.Analytics.Integration.Tests.Telemetry
 	[TestClass]
 	public class JobTelemetryIntegrationTests : BuildTestSetup
 	{
-		#region -- Test Constants --
+		#region -- Test Helpers --
 
-		const string EventNameKeyStr = "EventName";
+		#region -- Test Graph Strings --
 
-		#endregion -- Test Constants -- 
+		private const string Win64ToolNodeName = "Compile UnrealHeaderTool Win64";
+		private const string Win64GameNodeName = "Compile ShooterGameEditor Win64";
 
-		protected override void ConfigureServices(IServiceCollection services)
-		{
-			base.ConfigureServices(services);
+		private const string MacToolNodeName = "Compile UnrealHeaderTool Mac";
+		private const string MacGameNodeName = "Compile ShooterGameEditor Mac";
 
-			services.AddSingleton<TestTelemetryWriter>();
-			services.AddSingleton<ITelemetryWriter>(sp => sp.GetRequiredService<TestTelemetryWriter>());
-		}
+		private const string AllWin64Label = "Label_AllWin64Nodes";
+		private const string ToolWin64Label = "Label_Win64UHTOnly";
+		private const string ToolAllLabel = "Label_AllTools";
 
-		[TestMethod]
-		public async Task TestJobSummaryEmissionAsync()
+		private const string Aggregate = "AllNodes";
+
+		#endregion -- Test Graph Strings --
+
+		private async Task<(TemplateId templateRefId1, StreamConfig streamConfig, TestTelemetryWriter concreteTelemetryWriter)> SetupTelemetryTestPreconditionsAsync()
 		{
 			// Configure the metric
 			TelemetryStoreConfig telemetryStoreConfig = new TelemetryStoreConfig();
@@ -279,10 +284,75 @@ namespace HordeServer.Analytics.Integration.Tests.Telemetry
 
 			// Initialize the test telemetry writer
 			TestTelemetryWriter concreteTelemetryWriter = ServiceProvider.GetRequiredService<TestTelemetryWriter>();
+			return (templateRefId1, streamConfig, concreteTelemetryWriter);
+		}
+
+		private static NewLabel CreateLabel(string labelName, List<string> nodes)
+		{
+			NewLabel label = new NewLabel();
+			label.DashboardName = labelName;
+			label.IncludedNodes = nodes;
+			label.RequiredNodes = nodes;
+			return label;
+		}
+
+		private static Task<IGraph> UpdateGraph(IGraphCollection graphCollection, IGraph rootGraph, ITemplate template)
+		{
+			// Nodes
+
+			List<NewNode> nodes1 = new List<NewNode>
+			{
+				new NewNode(Win64ToolNodeName),
+				new NewNode(Win64GameNodeName),
+			};
+
+			List<NewNode> nodes2 = new List<NewNode>
+			{
+				new NewNode(MacToolNodeName),
+				new NewNode(MacGameNodeName),
+			};
+
+			// Labels
+			List<NewLabel> labels = new List<NewLabel>
+			{
+				CreateLabel(AllWin64Label, [Win64ToolNodeName, Win64GameNodeName]),
+				CreateLabel(ToolWin64Label, [Win64ToolNodeName]),
+				CreateLabel(ToolAllLabel,  [Win64ToolNodeName, MacToolNodeName])
+			};
+
+			// Groups
+			NewGroup group1 = new NewGroup("DefaultAgent1", nodes1);
+			NewGroup group2 = new NewGroup("DefaultAgent2", nodes2);
+
+			NewAggregate aggregate = new NewAggregate(Aggregate, nodes1.Union(nodes2).Select(x => x.Name).ToList());
+
+			return graphCollection.AppendAsync(rootGraph, [group1, group2], [aggregate], labels);
+		}
+
+		#endregion -- Test Helpers --
+
+		#region -- Test Constants --
+
+		const string EventNameKeyStr = "EventName";
+
+		#endregion -- Test Constants -- 
+
+		protected override void ConfigureServices(IServiceCollection services)
+		{
+			base.ConfigureServices(services);
+
+			services.AddSingleton<TestTelemetryWriter>();
+			services.AddSingleton<ITelemetryWriter>(sp => sp.GetRequiredService<TestTelemetryWriter>());
+		}
+
+		[TestMethod]
+		public async Task TestJobSummaryEmissionAsync()
+		{
+			(TemplateId templateRefId1, StreamConfig streamConfig, TestTelemetryWriter concreteTelemetryWriter) = await SetupTelemetryTestPreconditionsAsync();
 
 			// Validate our preconditions
 			Assert.IsNotNull(concreteTelemetryWriter);
-			Assert.IsTrue(!concreteTelemetryWriter.RecordedTelemetry.ContainsKey(TelemetryStoreId.Default));
+			Assert.IsFalse(concreteTelemetryWriter.RecordedTelemetry.ContainsKey(TelemetryStoreId.Default));
 
 			{
 				CreateJobOptions options = new CreateJobOptions();
@@ -316,5 +386,206 @@ namespace HordeServer.Analytics.Integration.Tests.Telemetry
 			JobSummaryTelemetry? actualJobSummaryTelemetry = actualStateJobSummaryPayload.MapToObject<JobSummaryTelemetry>();
 			Assert.IsNotNull(actualJobSummaryTelemetry);
 		}
+
+
+		[TestMethod]
+		public async Task TestJobLabelSummaryEmissionAsync()
+		{
+			(TemplateId templateRefId1, StreamConfig streamConfig, TestTelemetryWriter concreteTelemetryWriter) = await SetupTelemetryTestPreconditionsAsync();
+
+			// Validate our preconditions
+			Assert.IsNotNull(concreteTelemetryWriter);
+			Assert.IsFalse(concreteTelemetryWriter.RecordedTelemetry.ContainsKey(TelemetryStoreId.Default));
+
+			// Nominal case; all success
+			{
+				// Reset the writer
+				concreteTelemetryWriter.ClearStoredTelemetry();
+
+				CreateJobOptions options = new CreateJobOptions();
+				options.PreflightCommitId = CommitId.FromPerforceChange(999);
+
+				ITemplate template = await TemplateCollection.GetOrAddAsync(streamConfig.Templates[0]);
+				IGraph graph = await GraphCollection.AddAsync(template);
+				graph = await UpdateGraph(GraphCollection, graph, template);
+
+				options.Arguments.Add(IJob.TargetArgumentPrefix + Aggregate);
+
+				IJob job = await JobService.CreateJobAsync(null, streamConfig, templateRefId1, template.Hash, graph, "TestJobLabelTelemetrySuccess", CommitIdWithOrder.FromPerforceChange(1234), CommitIdWithOrder.FromPerforceChange(1233), options);
+				job = Deref(await JobService.UpdateBatchAsync(job, job.Batches[0].Id, streamConfig, LogIdUtils.GenerateNewId(), JobStepBatchState.Running));
+				job = Deref(await JobService.UpdateBatchAsync(job, job.Batches[1].Id, streamConfig, LogIdUtils.GenerateNewId(), JobStepBatchState.Running));
+				job = Deref(await JobService.UpdateBatchAsync(job, job.Batches[2].Id, streamConfig, LogIdUtils.GenerateNewId(), JobStepBatchState.Running));
+
+				// Batch 0
+				job = Deref(await JobService.UpdateStepAsync(job, job.Batches[0].Id, job.Batches[0].Steps[0].Id, streamConfig, JobStepState.Running));
+
+				// Batch 1
+				job = Deref(await JobService.UpdateStepAsync(job, job.Batches[1].Id, job.Batches[1].Steps[0].Id, streamConfig, JobStepState.Running)); // Tool
+				job = Deref(await JobService.UpdateStepAsync(job, job.Batches[1].Id, job.Batches[1].Steps[1].Id, streamConfig, JobStepState.Running)); // Game
+
+				// Batch 2
+				job = Deref(await JobService.UpdateStepAsync(job, job.Batches[2].Id, job.Batches[2].Steps[0].Id, streamConfig, JobStepState.Running)); // Tool
+				job = Deref(await JobService.UpdateStepAsync(job, job.Batches[2].Id, job.Batches[2].Steps[1].Id, streamConfig, JobStepState.Running)); // Game
+
+				FakeClock clock = ServiceProvider.GetRequiredService<FakeClock>();
+
+				await clock.AdvanceAsync(TimeSpan.FromSeconds(60.0));
+
+				// Batch 0 - Setup Build
+				job = Deref(await JobService.UpdateStepAsync(job, job.Batches[0].Id, job.Batches[0].Steps[0].Id, streamConfig, JobStepState.Completed, JobStepOutcome.Success));
+
+				// Batch 1 - Win64
+				job = Deref(await JobService.UpdateStepAsync(job, job.Batches[1].Id, job.Batches[1].Steps[0].Id, streamConfig, JobStepState.Completed, JobStepOutcome.Success)); // Tool
+				job = Deref(await JobService.UpdateStepAsync(job, job.Batches[1].Id, job.Batches[1].Steps[1].Id, streamConfig, JobStepState.Completed, JobStepOutcome.Success)); // Game
+
+				// Batch 2 - Mac
+				job = Deref(await JobService.UpdateStepAsync(job, job.Batches[2].Id, job.Batches[2].Steps[0].Id, streamConfig, JobStepState.Completed, JobStepOutcome.Success)); // Tool
+				job = Deref(await JobService.UpdateStepAsync(job, job.Batches[2].Id, job.Batches[2].Steps[1].Id, streamConfig, JobStepState.Completed, JobStepOutcome.Success)); // Game
+
+				// Validate we have received a telemetry event.
+				Assert.IsTrue(concreteTelemetryWriter.RecordedTelemetry.ContainsKey(TelemetryStoreId.Default));
+
+				// Verify presence of job summary telemetry
+				IList<Tuple<TelemetryRecordMeta?, Hashtable>> actualTelemetryPayloads = [.. concreteTelemetryWriter.RecordedTelemetry[TelemetryStoreId.Default]];
+				Assert.IsTrue(actualTelemetryPayloads.Any(x => (x.Item2[EventNameKeyStr] as string) == JobLabelSummaryTelemetry.DefaultEventName));
+
+				// Assert we actually have the expected telemetry type.
+				Hashtable actualStateJobSummaryPayload = actualTelemetryPayloads.Where(x => x.Item2[EventNameKeyStr] as string == JobLabelSummaryTelemetry.DefaultEventName).First().Item2;
+				JobLabelSummaryTelemetry? actualJobSummaryTelemetry = actualStateJobSummaryPayload.MapToObject<JobLabelSummaryTelemetry>();
+
+				Assert.IsNotNull(actualJobSummaryTelemetry);
+
+				Assert.IsTrue(actualJobSummaryTelemetry.SuccessfulLabels.Contains(ToolWin64Label));
+				Assert.IsTrue(actualJobSummaryTelemetry.SuccessfulLabels.Contains(AllWin64Label));
+				Assert.IsTrue(actualJobSummaryTelemetry.SuccessfulLabels.Contains(ToolAllLabel));
+				Assert.IsFalse(actualJobSummaryTelemetry.FailedLabels.Any());
+			}
+
+			// Label failures - cross batch
+			{
+				// Reset the writer
+				concreteTelemetryWriter.ClearStoredTelemetry();
+
+				CreateJobOptions options = new CreateJobOptions();
+				options.PreflightCommitId = CommitId.FromPerforceChange(999);
+
+				ITemplate template = await TemplateCollection.GetOrAddAsync(streamConfig.Templates[0]);
+				IGraph graph = await GraphCollection.AddAsync(template);
+				graph = await UpdateGraph(GraphCollection, graph, template);
+
+				options.Arguments.Add(IJob.TargetArgumentPrefix + Aggregate);
+
+				IJob job = await JobService.CreateJobAsync(null, streamConfig, templateRefId1, template.Hash, graph, "TestJobLabelTelemetryFailure", CommitIdWithOrder.FromPerforceChange(1234), CommitIdWithOrder.FromPerforceChange(1233), options);
+				job = Deref(await JobService.UpdateBatchAsync(job, job.Batches[0].Id, streamConfig, LogIdUtils.GenerateNewId(), JobStepBatchState.Running));
+				job = Deref(await JobService.UpdateBatchAsync(job, job.Batches[1].Id, streamConfig, LogIdUtils.GenerateNewId(), JobStepBatchState.Running));
+				job = Deref(await JobService.UpdateBatchAsync(job, job.Batches[2].Id, streamConfig, LogIdUtils.GenerateNewId(), JobStepBatchState.Running));
+
+				// Batch 0
+				job = Deref(await JobService.UpdateStepAsync(job, job.Batches[0].Id, job.Batches[0].Steps[0].Id, streamConfig, JobStepState.Running));
+
+				// Batch 1
+				job = Deref(await JobService.UpdateStepAsync(job, job.Batches[1].Id, job.Batches[1].Steps[0].Id, streamConfig, JobStepState.Running)); // Tool
+				job = Deref(await JobService.UpdateStepAsync(job, job.Batches[1].Id, job.Batches[1].Steps[1].Id, streamConfig, JobStepState.Running)); // Game
+
+				// Batch 2
+				job = Deref(await JobService.UpdateStepAsync(job, job.Batches[2].Id, job.Batches[2].Steps[0].Id, streamConfig, JobStepState.Running)); // Tool
+				job = Deref(await JobService.UpdateStepAsync(job, job.Batches[2].Id, job.Batches[2].Steps[1].Id, streamConfig, JobStepState.Running)); // Game
+
+				FakeClock clock = ServiceProvider.GetRequiredService<FakeClock>();
+
+				await clock.AdvanceAsync(TimeSpan.FromSeconds(60.0));
+
+				// Batch 0 - Setup Build
+				job = Deref(await JobService.UpdateStepAsync(job, job.Batches[0].Id, job.Batches[0].Steps[0].Id, streamConfig, JobStepState.Completed, JobStepOutcome.Success));
+
+				// Batch 1 - Win64
+				job = Deref(await JobService.UpdateStepAsync(job, job.Batches[1].Id, job.Batches[1].Steps[0].Id, streamConfig, JobStepState.Completed, JobStepOutcome.Success)); // Tool
+				job = Deref(await JobService.UpdateStepAsync(job, job.Batches[1].Id, job.Batches[1].Steps[1].Id, streamConfig, JobStepState.Completed, JobStepOutcome.Success)); // Game
+
+				// Batch 2 - Mac
+				job = Deref(await JobService.UpdateStepAsync(job, job.Batches[2].Id, job.Batches[2].Steps[0].Id, streamConfig, JobStepState.Completed, JobStepOutcome.Failure)); // Tool
+				job = Deref(await JobService.UpdateStepAsync(job, job.Batches[2].Id, job.Batches[2].Steps[1].Id, streamConfig, JobStepState.Completed, JobStepOutcome.Success)); // Game
+
+				// Validate we have received a telemetry event.
+				Assert.IsTrue(concreteTelemetryWriter.RecordedTelemetry.ContainsKey(TelemetryStoreId.Default));
+
+				// Verify presence of job summary telemetry
+				IList<Tuple<TelemetryRecordMeta?, Hashtable>> actualTelemetryPayloads = [.. concreteTelemetryWriter.RecordedTelemetry[TelemetryStoreId.Default]];
+				Assert.IsTrue(actualTelemetryPayloads.Any(x => (x.Item2[EventNameKeyStr] as string) == JobLabelSummaryTelemetry.DefaultEventName));
+
+				// Assert we actually have the expected telemetry type.
+				Hashtable actualStateJobSummaryPayload = actualTelemetryPayloads.Where(x => x.Item2[EventNameKeyStr] as string == JobLabelSummaryTelemetry.DefaultEventName).First().Item2;
+				JobLabelSummaryTelemetry? actualJobSummaryTelemetry = actualStateJobSummaryPayload.MapToObject<JobLabelSummaryTelemetry>();
+
+				Assert.IsNotNull(actualJobSummaryTelemetry);
+
+				Assert.IsTrue(actualJobSummaryTelemetry.SuccessfulLabels.Contains(ToolWin64Label));
+				Assert.IsTrue(actualJobSummaryTelemetry.SuccessfulLabels.Contains(AllWin64Label));
+				Assert.IsTrue(actualJobSummaryTelemetry.FailedLabels.Contains(ToolAllLabel));
+			}
+
+			// Label failures - within batch
+			{
+				// Reset the writer
+				concreteTelemetryWriter.ClearStoredTelemetry();
+
+				CreateJobOptions options = new CreateJobOptions();
+				options.PreflightCommitId = CommitId.FromPerforceChange(999);
+
+				ITemplate template = await TemplateCollection.GetOrAddAsync(streamConfig.Templates[0]);
+				IGraph graph = await GraphCollection.AddAsync(template);
+				graph = await UpdateGraph(GraphCollection, graph, template);
+
+				options.Arguments.Add(IJob.TargetArgumentPrefix + Aggregate);
+
+				IJob job = await JobService.CreateJobAsync(null, streamConfig, templateRefId1, template.Hash, graph, "TestJobLabelTelemetryFailure", CommitIdWithOrder.FromPerforceChange(1234), CommitIdWithOrder.FromPerforceChange(1233), options);
+				job = Deref(await JobService.UpdateBatchAsync(job, job.Batches[0].Id, streamConfig, LogIdUtils.GenerateNewId(), JobStepBatchState.Running));
+				job = Deref(await JobService.UpdateBatchAsync(job, job.Batches[1].Id, streamConfig, LogIdUtils.GenerateNewId(), JobStepBatchState.Running));
+				job = Deref(await JobService.UpdateBatchAsync(job, job.Batches[2].Id, streamConfig, LogIdUtils.GenerateNewId(), JobStepBatchState.Running));
+
+				// Batch 0
+				job = Deref(await JobService.UpdateStepAsync(job, job.Batches[0].Id, job.Batches[0].Steps[0].Id, streamConfig, JobStepState.Running));
+
+				// Batch 1
+				job = Deref(await JobService.UpdateStepAsync(job, job.Batches[1].Id, job.Batches[1].Steps[0].Id, streamConfig, JobStepState.Running)); // Tool
+				job = Deref(await JobService.UpdateStepAsync(job, job.Batches[1].Id, job.Batches[1].Steps[1].Id, streamConfig, JobStepState.Running)); // Game
+
+				// Batch 2
+				job = Deref(await JobService.UpdateStepAsync(job, job.Batches[2].Id, job.Batches[2].Steps[0].Id, streamConfig, JobStepState.Running)); // Tool
+				job = Deref(await JobService.UpdateStepAsync(job, job.Batches[2].Id, job.Batches[2].Steps[1].Id, streamConfig, JobStepState.Running)); // Game
+
+				FakeClock clock = ServiceProvider.GetRequiredService<FakeClock>();
+
+				await clock.AdvanceAsync(TimeSpan.FromSeconds(60.0));
+
+				// Batch 0 - Setup Build
+				job = Deref(await JobService.UpdateStepAsync(job, job.Batches[0].Id, job.Batches[0].Steps[0].Id, streamConfig, JobStepState.Completed, JobStepOutcome.Success));
+
+				// Batch 1 - Win64
+				job = Deref(await JobService.UpdateStepAsync(job, job.Batches[1].Id, job.Batches[1].Steps[0].Id, streamConfig, JobStepState.Completed, JobStepOutcome.Failure)); // Tool
+				job = Deref(await JobService.UpdateStepAsync(job, job.Batches[1].Id, job.Batches[1].Steps[1].Id, streamConfig, JobStepState.Completed, JobStepOutcome.Success)); // Game
+
+				// Batch 2 - Mac
+				job = Deref(await JobService.UpdateStepAsync(job, job.Batches[2].Id, job.Batches[2].Steps[0].Id, streamConfig, JobStepState.Completed, JobStepOutcome.Success)); // Tool
+				job = Deref(await JobService.UpdateStepAsync(job, job.Batches[2].Id, job.Batches[2].Steps[1].Id, streamConfig, JobStepState.Completed, JobStepOutcome.Success)); // Game
+
+				// Validate we have received a telemetry event.
+				Assert.IsTrue(concreteTelemetryWriter.RecordedTelemetry.ContainsKey(TelemetryStoreId.Default));
+
+				// Verify presence of job summary telemetry
+				IList<Tuple<TelemetryRecordMeta?, Hashtable>> actualTelemetryPayloads = [.. concreteTelemetryWriter.RecordedTelemetry[TelemetryStoreId.Default]];
+				Assert.IsTrue(actualTelemetryPayloads.Any(x => (x.Item2[EventNameKeyStr] as string) == JobLabelSummaryTelemetry.DefaultEventName));
+
+				// Assert we actually have the expected telemetry type.
+				Hashtable actualStateJobSummaryPayload = actualTelemetryPayloads.Where(x => x.Item2[EventNameKeyStr] as string == JobLabelSummaryTelemetry.DefaultEventName).First().Item2;
+				JobLabelSummaryTelemetry? actualJobSummaryTelemetry = actualStateJobSummaryPayload.MapToObject<JobLabelSummaryTelemetry>();
+
+				Assert.IsNotNull(actualJobSummaryTelemetry);
+
+				Assert.IsTrue(actualJobSummaryTelemetry.FailedLabels.Contains(ToolWin64Label));
+				Assert.IsTrue(actualJobSummaryTelemetry.FailedLabels.Contains(AllWin64Label));
+				Assert.IsTrue(actualJobSummaryTelemetry.FailedLabels.Contains(ToolAllLabel));
+			}
+		}
 	}
 }
\ No newline at end of file
diff --git a/Engine/Source/Programs/Horde/Plugins/Build/HordeServer.Build/Jobs/JobService.cs b/Engine/Source/Programs/Horde/Plugins/Build/HordeServer.Build/Jobs/JobService.cs
--- a/Engine/Source/Programs/Horde/Plugins/Build/HordeServer.Build/Jobs/JobService.cs
+++ b/Engine/Source/Programs/Horde/Plugins/Build/HordeServer.Build/Jobs/JobService.cs
@@ -1052,7 +1052,6 @@ namespace HordeServer.Jobs
 						if (job.GetState() == JobState.Complete)
 						{
 							IJobStep? completingJobStep;
-
 							if (job.TryGetStep(batchId, stepId, out completingJobStep) && completingJobStep.FinishTimeUtc != null)
 							{
 								if (!streamConfig.TelemetryStoreId.IsEmpty)
@@ -1060,8 +1059,9 @@ namespace HordeServer.Jobs
 									try
 									{
 										_telemetryWriter.WriteEvent(streamConfig.TelemetryStoreId, new JobSummaryTelemetry(job, completingJobStep.FinishTimeUtc.Value));
+										_telemetryWriter.WriteEvent(streamConfig.TelemetryStoreId, new JobLabelSummaryTelemetry(job, graph));
 									}
-									catch(Exception ex)
+									catch (Exception ex)
 									{
 										_logger.LogError(ex, "Failed to publish JobSummary telemetry.");
 									}
diff --git a/Engine/Source/Programs/Horde/Plugins/Build/HordeServer.Build/Jobs/JobTelemetry.cs b/Engine/Source/Programs/Horde/Plugins/Build/HordeServer.Build/Jobs/JobTelemetry.cs
--- a/Engine/Source/Programs/Horde/Plugins/Build/HordeServer.Build/Jobs/JobTelemetry.cs
+++ b/Engine/Source/Programs/Horde/Plugins/Build/HordeServer.Build/Jobs/JobTelemetry.cs
@@ -1,6 +1,7 @@
 // Copyright Epic Games, Inc. All Rights Reserved.
 
 using EpicGames.Horde.Jobs;
+using EpicGames.Horde.Jobs.Graphs;
 using EpicGames.Horde.Jobs.Templates;
 using EpicGames.Horde.Streams;
 
@@ -52,7 +53,7 @@ namespace HordeServer.Jobs
 		public DateTime FinishTimeUtc { get; init; }
 
 		/// <summary>
-		/// 
+		/// The summed duration of all steps.
 		/// </summary>
 		public float JobStepsTotalTime { get; init; }
 
@@ -141,4 +142,79 @@ namespace HordeServer.Jobs
 			FailureRatio = jobStepsCompletionInfo.FailureRatio;
 		}
 	}
+
+	/// <summary>
+	/// Record used for Job Label Summary telemetry. Models a <see cref="IJob"/> and <see cref="IGraph"/>.
+	/// </summary>
+	public record JobLabelSummaryTelemetry
+	{
+		/// <summary>
+		/// The telemetry event name.
+		/// </summary>
+		public string EventName { get; init; }
+
+		/// <summary>
+		/// The job Id.
+		/// </summary>
+		public JobId JobId { get; init; }
+
+		/// <summary>
+		/// The graph content hash.
+		/// </summary>
+		public string GraphContentHash { get; init; }
+
+		/// <summary>
+		/// The template of the job.
+		/// </summary>
+		public TemplateId TemplateId { get; init; }
+
+		/// <summary>
+		/// The stream of the job.
+		/// </summary>
+		public StreamId StreamId { get; init; }
+
+		/// <summary>
+		/// The commit of the job. <see cref="IJob.CommitId"/>.
+		/// </summary>
+		public string CommitId { get; init; }
+
+		/// <summary>
+		/// Array representing all labels that were successful in the job.
+		/// </summary>
+		public string?[] SuccessfulLabels { get; init; }
+
+		/// <summary>
+		/// Array representing all labels that were failures in the job.
+		/// </summary>
+		public string?[] FailedLabels { get; init; }
+
+		/// <summary>
+		/// Default event name for the JobLabelSummaryTelemetry.
+		/// </summary>
+		public const string DefaultEventName = "State.JobLabelSummary";
+
+		/// <summary>
+		/// Constructor.
+		/// </summary>
+		/// <param name="job">The job to base the telemetry on.</param>
+		/// <param name="graph">The graph to base the telemetry on.</param>
+		public JobLabelSummaryTelemetry(IJob job, IGraph graph)
+		{
+			EventName = DefaultEventName;
+			SuccessfulLabels = [];
+			FailedLabels = [];
+			EventName = DefaultEventName;
+			JobId = job.Id;
+			GraphContentHash = graph.Id.ToString();
+			TemplateId = job.TemplateId;
+			StreamId = job.StreamId;
+			CommitId = job.CommitId?.Name ?? "0";
+
+			List<GetLabelStateResponse> responses = [];
+			job.GetLabelStateResponses(graph, responses);
+
+			SuccessfulLabels = responses.Where(r => r.State == LabelState.Complete && r.Outcome == LabelOutcome.Success).Select(r => r.DashboardName).ToArray();
+			FailedLabels = responses.Where(r => r.State == LabelState.Complete && r.Outcome == LabelOutcome.Failure).Select(r => r.DashboardName).ToArray();
+		}
+	}
 }
\ No newline at end of file
