Update test nuget packages to latest patch version

#rnx

diff --git a/Engine/Source/Programs/Horde/Plugins/ExamplePlugin/HordeServer.ExamplePlugin.Tests/HordeServer.ExamplePlugin.Tests.csproj b/Engine/Source/Programs/Horde/Plugins/ExamplePlugin/HordeServer.ExamplePlugin.Tests/HordeServer.ExamplePlugin.Tests.csproj
--- a/Engine/Source/Programs/Horde/Plugins/ExamplePlugin/HordeServer.ExamplePlugin.Tests/HordeServer.ExamplePlugin.Tests.csproj
+++ b/Engine/Source/Programs/Horde/Plugins/ExamplePlugin/HordeServer.ExamplePlugin.Tests/HordeServer.ExamplePlugin.Tests.csproj
@@ -20,15 +20,15 @@
   </PropertyGroup>
 
   <ItemGroup>
-    <PackageReference Include="coverlet.collector" Version="6.0.2">
+    <PackageReference Include="coverlet.collector" Version="6.0.4">
       <PrivateAssets>all</PrivateAssets>
       <IncludeAssets>runtime; build; native; contentfiles; analyzers; buildtransitive</IncludeAssets>
     </PackageReference>
     <PackageReference Include="Microsoft.NET.Test.Sdk" Version="17.11.1" />
-    <PackageReference Include="MSTest.TestAdapter" Version="3.6.2" />
-    <PackageReference Include="MSTest.TestFramework" Version="3.6.2" />
+    <PackageReference Include="MSTest.TestAdapter" Version="3.6.4" />
+    <PackageReference Include="MSTest.TestFramework" Version="3.6.4" />
     <PackageReference Include="Microsoft.AspNetCore.Mvc.Testing" Version="8.0.20" />
-    <PackageReference Include="Moq" Version="4.20.70" />
+    <PackageReference Include="Moq" Version="4.20.72" />
     <PackageReference Include="Microsoft.Extensions.Caching.Memory" Version="8.0.1" />
   </ItemGroup>
 
