Change Unspecified color
Change Unspecified text to Catastrophic
Centralize status texts
Fix expanded summary card not presisted when phase view widget is opened
Segragate failed from passed tests
Order test summaries by number of failed + castrophic phase count

#rnx

diff --git a/Engine/Source/Programs/Horde/HordeDashboard/plugins/test-automation/testAutomationCommon.tsx b/Engine/Source/Programs/Horde/HordeDashboard/plugins/test-automation/testAutomationCommon.tsx
--- a/Engine/Source/Programs/Horde/HordeDashboard/plugins/test-automation/testAutomationCommon.tsx
+++ b/Engine/Source/Programs/Horde/HordeDashboard/plugins/test-automation/testAutomationCommon.tsx
@@ -91,23 +91,28 @@ const getStyles = () => {
    return styles;
 }
 
-let _colors: Map<StatusColor, string> | undefined = undefined;
-
 export const getStatusColors = (): Map<StatusColor, string> => {
 
     const dashboardStatusColors = dashboard.getStatusColors();
 
-    const colors = _colors ?? new Map<StatusColor, string>([
+    const colors = new Map<StatusColor, string>([
         [StatusColor.Success, dashboardStatusColors.get(StatusColor.Success)!],
         [StatusColor.Warnings, dashboardStatusColors.get(StatusColor.Warnings)!],
         [StatusColor.Failure, dashboardStatusColors.get(StatusColor.Failure)!],
         [StatusColor.Skipped,  dashboard.darktheme ? "#c5c4c3ff" : "#a1a1a1ff"],
-        [StatusColor.Unspecified, "#515151ff"]
+        [StatusColor.Unspecified, "#8e4747ff"]
     ]);
 
     return colors;
 }
 
+export const statusTexts = new Map<TestOutcome, string>([
+    [TestOutcome.Success, "Passed"],
+    [TestOutcome.Failure, "Failure"],
+    [TestOutcome.Skipped, "Skipped"],
+    [TestOutcome.Unspecified, "Catastrophic"]
+]);
+
 export const MultiOptionChooser: React.FC<{options: IComboBoxOption[], initialSelection: string[], updateKeys: (selectedKeys: string[]) => void, placeholder?: string, style?: any, disabled?: boolean }> = ({ options, initialSelection, updateKeys, placeholder, style, disabled }) => {
     const [selection, setSelection] = useState<string[]>(initialSelection);
     const optionsWithSelectAll = [...options];
@@ -193,14 +198,14 @@ export const SessionStatusBar = (session: TestSessionResult | TestSessionDetails
     const stack: StatusBarStack[] = [
         {
             value: metaUnspecifiedFactor * 100,
-            title: "Unspecified",
+            title: statusTexts.get(TestOutcome.Unspecified),
             titleValue: `${ Math.ceil(session.phasesUnspecifiedCount / (sessionTotalCount || 1) * 100)}%`,
             color: statusColors.get(StatusColor.Unspecified)!,
             stripes: true
         },
         {
             value: metaFailedFactor * 100,
-            title: "Failure",
+            title: statusTexts.get(TestOutcome.Failure),
             titleValue: `${ Math.ceil(session.phasesFailedCount / (sessionTotalCount || 1) * 100)}%`,
             color: statusColors.get(StatusColor.Failure)!,
             stripes: true
@@ -254,17 +259,17 @@ export const SessionStatusPie = (session: TestSessionResult, radius: number): JS
     const stack: StatusBarStack[] = [
         {
             value: metaUnspecifiedFactor * 100,
-            title: "Unspecified",
+            title: statusTexts.get(TestOutcome.Unspecified),
             color: statusColors.get(StatusColor.Unspecified)!
         },
         {
             value: metaFailedFactor * 100,
-            title: "Failure",
+            title: statusTexts.get(TestOutcome.Failure),
             color: statusColors.get(StatusColor.Failure)!
         },
         {
             value: (1 - (metaFailedFactor + metaUnspecifiedFactor)) * 100,
-            title: "Passed",
+            title: statusTexts.get(TestOutcome.Success),
             color: statusColors.get(StatusColor.Success)!
         }
     ];
@@ -302,22 +307,22 @@ export const PhasesStatusPie = (phases: TestPhaseStatus[], radius: number): JSX.
     const stack: StatusBarStack[] = [
         {
             value: unspecifiedFactor * 100,
-            title: "Unspecified",
+            title: statusTexts.get(TestOutcome.Unspecified),
             color: statusColors.get(StatusColor.Unspecified)!
         },
         {
             value: skippedFactor * 100,
-            title: "Skipped",
+            title: statusTexts.get(TestOutcome.Skipped),
             color: statusColors.get(StatusColor.Skipped)!
         },
         {
             value: failedFactor * 100,
-            title: "Failure",
+            title: statusTexts.get(TestOutcome.Failure),
             color: statusColors.get(StatusColor.Failure)!
         },
         {
             value: (1 - (failedFactor + unspecifiedFactor + skippedFactor)) * 100,
-            title: "Passed",
+            title: statusTexts.get(TestOutcome.Success),
             color: statusColors.get(StatusColor.Success)!
         }
     ];
@@ -328,9 +333,9 @@ export const PhasesStatusPie = (phases: TestPhaseStatus[], radius: number): JSX.
 export const SessionValues = (session: TestSessionResult, style?: any): JSX.Element => {
     const statusColors = getStatusColors();
     return <Text style={style}>
-                {!!session.phasesUnspecifiedCount && <span style={{color: statusColors.get(StatusColor.Unspecified)}}> Unspecified: {session.phasesUnspecifiedCount}</span>}
-                {!!session.phasesFailedCount && <span style={{color: statusColors.get(StatusColor.Failure)}}> Failed: {session.phasesFailedCount}</span>}
-                {!!session.phasesSucceededCount && <span style={{color: statusColors.get(StatusColor.Success)}}> Success: {session.phasesSucceededCount}</span>}
+                {!!session.phasesUnspecifiedCount && <span style={{color: statusColors.get(StatusColor.Unspecified)}}> {statusTexts.get(TestOutcome.Unspecified)}: {session.phasesUnspecifiedCount}</span>}
+                {!!session.phasesFailedCount && <span style={{color: statusColors.get(StatusColor.Failure)}}> {statusTexts.get(TestOutcome.Failure)}: {session.phasesFailedCount}</span>}
+                {!!session.phasesSucceededCount && <span style={{color: statusColors.get(StatusColor.Success)}}> {statusTexts.get(TestOutcome.Success)}: {session.phasesSucceededCount}</span>}
             </Text>
 }
 
@@ -467,25 +472,25 @@ export const getPhaseOutcomeOptions = (): ICheckListOption[] => {
     return [
         {
             key: TestOutcome.Unspecified,
-            text: 'Unspecified',
+            text: statusTexts.get(TestOutcome.Unspecified),
             checked: true,
             color: statusColors.get(StatusColor.Unspecified)
         },
         {
             key: TestOutcome.Failure,
-            text: 'Failed',
+            text: statusTexts.get(TestOutcome.Failure),
             checked: true,
             color: statusColors.get(StatusColor.Failure)
         },
         {
             key: TestOutcome.Success,
-            text: 'Passed',
+            text: statusTexts.get(TestOutcome.Success),
             checked: true,
             color: statusColors.get(StatusColor.Success)
         },
         {
             key: TestOutcome.Skipped,
-            text: 'Skipped',
+            text: statusTexts.get(TestOutcome.Skipped),
             checked: true,
             color: statusColors.get(StatusColor.Skipped)
         },
diff --git a/Engine/Source/Programs/Horde/HordeDashboard/plugins/test-automation/testAutomationSummary.tsx b/Engine/Source/Programs/Horde/HordeDashboard/plugins/test-automation/testAutomationSummary.tsx
--- a/Engine/Source/Programs/Horde/HordeDashboard/plugins/test-automation/testAutomationSummary.tsx
+++ b/Engine/Source/Programs/Horde/HordeDashboard/plugins/test-automation/testAutomationSummary.tsx
@@ -9,7 +9,7 @@ import dashboard, { StatusColor } from "horde/backend/Dashboard";
 import { getShortNiceTime } from "horde/base/utilities/timeUtils";
 import { getHordeStyling } from "horde/styles/Styles";
 import { projectStore } from 'horde/backend/ProjectStore';
-import { SessionStatusBar, SessionValues, StatusBar, StatusBarStack, MultiOptionChooser, getTestSessionStatusColor, getStreamOptions, getStatusColors } from "./testAutomationCommon";
+import { SessionStatusBar, SessionValues, StatusBar, StatusBarStack, MultiOptionChooser, getTestSessionStatusColor, getStreamOptions, getStatusColors, statusTexts } from "./testAutomationCommon";
 import { TestHistoryGraph } from "./testHistoryGraph";
 
 type SessionCallback = (session: TestSessionResult) => void;
@@ -198,8 +198,17 @@ const TestHistoryModal: React.FC<{ test: TestNameRef, handler: TestDataHandler,
         </Modal>
 }
 
+// Manifest to persist the exapanded summary cards when phase view widget is opened.
+// The phase view being lower in the view stack, the expanded properties value of the card widget get lost at that time.
+const _persistExpandedCards: Map<TestNameRef, boolean> = new Map();
+
 const TestSummaryCard: React.FC<{ test: TestNameRef, sessions: [MetadataRef, TestSessionResult][], metaKeyMask: string[] | undefined, handler: TestDataHandler, onDetailsClick: SessionCallback, onHistoryClick: HistorySessionCallback }> = memo(({ test, sessions, metaKeyMask, handler, onDetailsClick, onHistoryClick }) => {
-    const [expanded, setExpanded] = useState(false);
+    const [expanded, setExpanded] = useState(_persistExpandedCards.get(test) ?? false);
+
+    const expandCard = useCallback((bExpand) => {
+        _persistExpandedCards.set(test, bExpand);
+        setExpanded(bExpand);
+    }, [test]);
  
     const { hordeClasses, modeColors } = getHordeStyling();
     const statusColors = getStatusColors();
@@ -302,19 +311,19 @@ const TestSummaryCard: React.FC<{ test: TestNameRef, sessions: [MetadataRef, Tes
     const testStack: StatusBarStack[] = [
         {
             value: testUnspecifiedFactor * 100,
-            title: "Unspecified",
+            title: statusTexts.get(TestOutcome.Unspecified),
             color: statusColors.get(StatusColor.Unspecified)!,
             stripes: true
         },
         {
             value: testFailedFactor * 100,
-            title: "Failure",
+            title: statusTexts.get(TestOutcome.Failure),
             color: statusColors.get(StatusColor.Failure)!,
             stripes: true
         }
     ]
  
-    return <Stack className={hordeClasses.raised} style={{ cursor: "pointer", height: "fit-content", backgroundColor: dashboard.darktheme ? "#242729" : modeColors.background, padding: 12, width: 332 }} onClick={() => setExpanded(!expanded)}>
+    return <Stack className={hordeClasses.raised} style={{ cursor: "pointer", height: "fit-content", backgroundColor: dashboard.darktheme ? "#242729" : modeColors.background, padding: 12, width: 332 }} onClick={() => expandCard(!expanded)}>
             <Stack horizontal verticalAlign="center" >
                 <Stack className="horde-no-darktheme" style={{ paddingTop: 1, paddingRight: 4 }}>
                     <FontIcon style={{ fontSize: 13 }} iconName={expanded ? "ChevronDown" : "ChevronRight"} />
@@ -332,11 +341,13 @@ const TestSummaryCard: React.FC<{ test: TestNameRef, sessions: [MetadataRef, Tes
                     }}>History</DefaultButton>
                 </Stack>
             </Stack>
-            <Stack>
-                <Stack style={{ paddingTop: 4 }}>
-                    {StatusBar(testStack, 308, 10, statusColors.get(StatusColor.Success)!, { margin: '3px !important' })}
-                </Stack>        
-            </Stack>
+            <Stack style={{ paddingTop: 4 }}>
+                {StatusBar(testStack, 308, 10, statusColors.get(StatusColor.Success)!, { margin: '3px !important' })}
+                <Stack horizontal style={{marginLeft: '5px'}} tokens={{childrenGap: 3}}>                    
+                    {!!testUnspecified && <Text variant="xSmall" style={{color: statusColors.get(StatusColor.Unspecified), minWidth: `${testUnspecifiedFactor*100}%`}} title={`${testUnspecified} ${statusTexts.get(TestOutcome.Unspecified)}`}> {testUnspecified}</Text>}
+                    {!!testFailed && <Text variant="xSmall" style={{color: statusColors.get(StatusColor.Failure), minWidth: `${testFailedFactor*100}%`}} title={`${testFailed} ${statusTexts.get(TestOutcome.Failure)}`}> {testFailed}</Text>}
+                </Stack>
+            </Stack>        
             {!!expanded &&
                 <Stack style={{ paddingLeft: 8 }} tokens={{ childrenGap: 4 }}>
                     {metaElements}
@@ -345,18 +356,19 @@ const TestSummaryCard: React.FC<{ test: TestNameRef, sessions: [MetadataRef, Tes
             </Stack>
 });
 
-const filterSessionsWithMetaMask = (testStatus: TestStatus, filteredTests: TestNameRef[], handler: TestDataHandler): [[TestNameRef, [MetadataRef, TestSessionStatus][]][], Map<string, string[]>] => {
+const filterSessionsWithMetaMask = (testStatus: TestStatus, filteredTests: TestNameRef[], handler: TestDataHandler): [[TestNameRef, [MetadataRef, TestSessionStatus][]][], [TestNameRef, [MetadataRef, TestSessionStatus][]][], Map<string, string[]>] => {
     const filteredMetadata = handler.filteredMetadata;
 
     const testMeta: Map<string, [TestNameRef, MetadataRef[]][]> = new Map();
 
     const failing: [TestNameRef, [MetadataRef, TestSessionStatus][]][] = [];
     const passing: [TestNameRef, [MetadataRef, TestSessionStatus][]][] = [];
+    const testFailureCount: Map<TestNameRef, number> = new Map();
 
     filteredTests.forEach(test => {
         // sort by failed
         const status = testStatus!.tests.get(test)!;
-        let anyFailed = false;
+        let totalFailureCount = 0;
 
         const metaRefs: MetadataRef[] = [];
         const testRefMeta: [TestNameRef, MetadataRef[]][] = testMeta.get(test.name) ?? testMeta.set(test.name, []).get(test.name)!;
@@ -373,15 +385,16 @@ const filterSessionsWithMetaMask = (testStatus: TestStatus, filteredTests: TestN
 
             const outcome = session.outcome;
             if ( outcome === TestOutcome.Failure || outcome === TestOutcome.Unspecified ) {
-                anyFailed = true;
+                totalFailureCount += session.phasesFailedCount + session.phasesUnspecifiedCount;
             }
 
             metaResults.push([meta, metaStatus]);
             metaRefs.push(meta);
         });
 
-        if (anyFailed) {
+        if (totalFailureCount > 0) {
             failing.push([test, metaResults]);
+            testFailureCount.set(test, totalFailureCount);
         } else {
             passing.push([test, metaResults]);
         }
@@ -398,7 +411,9 @@ const filterSessionsWithMetaMask = (testStatus: TestStatus, filteredTests: TestN
         return [testName, uniqueKeys];
     }));
 
-    return [[...failing, ...passing], testMetaKeys];
+    failing.sort(([testA, ], [testB, ]) => testFailureCount.get(testB)! - testFailureCount.get(testA)!);
+
+    return [failing, passing, testMetaKeys];
 }
 
 export const TestSummary: React.FC<{ handler: TestDataHandler }> = ({ handler }) => {
@@ -427,9 +442,9 @@ export const TestSummary: React.FC<{ handler: TestDataHandler }> = ({ handler })
         return null;
     }
 
-    const [testSessions, testMetaKeys] = filterSessionsWithMetaMask(testStatus, filteredTests, handler);
+    const [failingTestSessions, passingTestSessions, testMetaKeys] = filterSessionsWithMetaMask(testStatus, filteredTests, handler);
 
-    const tests = testSessions.map(([test, sessions]) => {
+    const testSummaryTemplate =([test, sessions]) => {
 
         if (searchTestId && !historyState.test && test.id === searchTestId) {
             historyState.test = test;
@@ -447,11 +462,14 @@ export const TestSummary: React.FC<{ handler: TestDataHandler }> = ({ handler })
                     onDetailsClick={onDetailsClick}
                 />
 
-    });
+    }
+
+    const failingTests = failingTestSessions.map(testSummaryTemplate);
+    const passingTests = passingTestSessions.map(testSummaryTemplate);
 
     const { hordeClasses } = getHordeStyling();
 
-    return <Stack className={hordeClasses.raised}>
+    return <Stack tokens={{childrenGap: 12}}>
             {!!historyState.test &&
                 <TestHistoryModal
                     test={historyState.test}
@@ -463,15 +481,33 @@ export const TestSummary: React.FC<{ handler: TestDataHandler }> = ({ handler })
                     onDetailsClick={historyState.onClick!}
                 />
             }
-            <Stack style={{ paddingBottom: 12 }}>
-                <Stack>
-                    <Text style={{ fontFamily: "Horde Open Sans Semibold" }} variant='medium'>Latest Test Results</Text>
+            {failingTests.length > 0 &&
+                <Stack className={hordeClasses.raised}>
+                    <Stack style={{ paddingBottom: 12 }}>
+                        <Stack>
+                            <Text style={{ fontFamily: "Horde Open Sans Semibold" }} variant='medium'>Latest Failing Test Results</Text>
+                        </Stack>
+                    </Stack>
+                    <Stack style={{ paddingLeft: 12 }} tokens={{ childrenGap: 4 }}>
+                        <Stack wrap horizontal style={{ width: "100%" }} tokens={{ childrenGap: 12 }}>
+                            {failingTests}
+                        </Stack>
+                    </Stack>
                 </Stack>
-            </Stack>
-            <Stack style={{ paddingLeft: 12 }} tokens={{ childrenGap: 4 }}>
-                <Stack wrap horizontal style={{ width: "100%" }} tokens={{ childrenGap: 12 }}>
-                    {tests}
+            }
+            {passingTests.length > 0 &&
+                <Stack className={hordeClasses.raised}>
+                    <Stack style={{ paddingBottom: 12 }}>
+                        <Stack>
+                            <Text style={{ fontFamily: "Horde Open Sans Semibold" }} variant='medium'>Latest Passing Test Results</Text>
+                        </Stack>
+                    </Stack>
+                    <Stack style={{ paddingLeft: 12 }} tokens={{ childrenGap: 4 }}>
+                        <Stack wrap horizontal style={{ width: "100%" }} tokens={{ childrenGap: 12 }}>
+                            {passingTests}
+                        </Stack>
+                    </Stack>
                 </Stack>
-            </Stack>
+            }
         </Stack>
 }
\ No newline at end of file
diff --git a/Engine/Source/Programs/Horde/HordeDashboard/plugins/test-automation/testHistoryGraph.ts b/Engine/Source/Programs/Horde/HordeDashboard/plugins/test-automation/testHistoryGraph.ts
--- a/Engine/Source/Programs/Horde/HordeDashboard/plugins/test-automation/testHistoryGraph.ts
+++ b/Engine/Source/Programs/Horde/HordeDashboard/plugins/test-automation/testHistoryGraph.ts
@@ -5,7 +5,7 @@ import { TestDataHandler, TestNameRef, MetadataRef, TestMetaStatus, TestSessionR
 import dashboard, { StatusColor } from 'horde/backend/Dashboard';
 import { getHumanTime, getShortNiceTime, msecToElapsed } from 'horde/base/utilities/timeUtils';
 import { getHordeStyling } from 'horde/styles/Styles';
-import { getStatusColors, SessionValues } from "./testAutomationCommon";
+import { getStatusColors, statusTexts, SessionValues } from "./testAutomationCommon";
 import { renderToString } from "react-dom/server"
 import { TestOutcome } from "./api";
 
@@ -393,17 +393,17 @@ const SessionPieChart = (container: SVGGElement, session: TestSessionResult, rad
     const stack: ItemStack[] = [
         {
             value: metaUnspecifiedFactor * 100,
-            name: "Unspecified",
+            name: statusTexts.get(TestOutcome.Unspecified)!,
             color: statusColors.get(StatusColor.Unspecified)!,
         },
         {
             value: metaFailedFactor * 100,
-            name: "Failure",
+            name: statusTexts.get(TestOutcome.Failure)!,
             color: statusColors.get(StatusColor.Failure)!,
         },
         {
             value: (1 - (metaFailedFactor + metaUnspecifiedFactor)) * 100,
-            name: "Passed",
+            name: statusTexts.get(TestOutcome.Success)!,
             color: statusColors.get(StatusColor.Success)!
         }
     ];
@@ -442,17 +442,17 @@ const SessionConcentricCircle = (container: SVGGElement, session: TestSessionRes
     const stack: ItemStack[] = [
         {
             value: metaSuccessFactor && 1,
-            name: "Passed",
+            name: statusTexts.get(TestOutcome.Success)!,
             color: statusColors.get(StatusColor.Success)!
         },
         {
             value: metaUnspecifiedFactor && ((metaSuccessFactor && (metaFailedFactor + metaUnspecifiedFactor)) || 1),
-            name: "Unspecified",
+            name: statusTexts.get(TestOutcome.Unspecified)!,
             color: statusColors.get(StatusColor.Unspecified)!,
         },
         {
             value: metaFailedFactor,
-            name: "Failure",
+            name: statusTexts.get(TestOutcome.Failure)!,
             color: statusColors.get(StatusColor.Failure)!,
         },
     ];
