Horde: Fix server build

diff --git a/Engine/Source/Programs/Horde/Plugins/Build/HordeServer.Build.Tests/Issues/IssueServiceTests.cs b/Engine/Source/Programs/Horde/Plugins/Build/HordeServer.Build.Tests/Issues/IssueServiceTests.cs
--- a/Engine/Source/Programs/Horde/Plugins/Build/HordeServer.Build.Tests/Issues/IssueServiceTests.cs
+++ b/Engine/Source/Programs/Horde/Plugins/Build/HordeServer.Build.Tests/Issues/IssueServiceTests.cs
@@ -2233,7 +2233,7 @@ namespace HordeServer.Tests.Issues
 				IReadOnlyList<IIssue> issues = await IssueCollection.FindIssuesAsync();
 				Assert.AreEqual(1, issues.Count);
 				Assert.AreEqual("Gauntlet:test", issues[0].Fingerprints[0].Type);
-				Assert.AreEqual(new IssueKey($"{job.StreamId}:Update Version Files", IssueKeyType.None), issues[0].Fingerprints[0].Keys.First());
+				Assert.AreEqual(new IssueKey($"Update Version Files", IssueKeyType.None), issues[0].Fingerprints[0].Keys.First());
 				Assert.AreEqual("Automation test errors in Update Version Files", issues[0].Summary);
 			}
 			// #2
@@ -2337,7 +2337,7 @@ namespace HordeServer.Tests.Issues
 				Assert.AreEqual("Gauntlet:test", issues[0].Fingerprints[0].Type);
 				Assert.AreEqual(10, issues[0].Fingerprints[0].Keys.Count);
 				Assert.AreEqual("hash:", issues[0].Fingerprints[0].Keys.First().Name.Substring(0, 5));
-				Assert.AreEqual($":{job.StreamId}:Update Version Files", issues[0].Fingerprints[0].Keys.First().Name.Substring(37));
+				Assert.AreEqual($":Update Version Files", issues[0].Fingerprints[0].Keys.First().Name.Substring(37));
 				Assert.AreEqual("Automation test errors in Update Version Files", issues[0].Summary);
 			}
 			// #4
@@ -2359,7 +2359,7 @@ namespace HordeServer.Tests.Issues
 				Assert.AreEqual(1, issues.Count);
 				Assert.AreEqual("Gauntlet:fatal", issues[0].Fingerprints[0].Type);
 				Assert.AreEqual(1, issues[0].Fingerprints[0].Keys.Count);
-				Assert.AreEqual($"{job.StreamId}:Update Version Files", issues[0].Fingerprints[0].Keys.First().Name);
+				Assert.AreEqual($"Update Version Files", issues[0].Fingerprints[0].Keys.First().Name);
 				Assert.AreEqual("Automation fatal errors in Update Version Files", issues[0].Summary);
 			}
 			// #5
