Horde: Optimize CreateAgentObjectAsync to only fetch agent document if update fails

diff --git a/Engine/Source/Programs/Horde/Plugins/Compute/HordeServer.Compute/Agents/AgentCollection.cs b/Engine/Source/Programs/Horde/Plugins/Compute/HordeServer.Compute/Agents/AgentCollection.cs
--- a/Engine/Source/Programs/Horde/Plugins/Compute/HordeServer.Compute/Agents/AgentCollection.cs
+++ b/Engine/Source/Programs/Horde/Plugins/Compute/HordeServer.Compute/Agents/AgentCollection.cs
@@ -507,11 +507,15 @@ namespace HordeServer.Agents
 				// Reconcile the differences
 				for (; ; )
 				{
-					// Fetch the document again
-					document = await GetDocumentAsync(agentId, cancellationToken);
+					// Fetch the document again if update failed below.
+					// On first iteration, document is already set.
 					if (document == null)
 					{
-						return null;
+						document = await GetDocumentAsync(agentId, cancellationToken);
+						if (document == null)
+						{
+							return null;
+						}
 					}
 
 					SessionId? sessionId = document.SessionId;
@@ -528,7 +532,7 @@ namespace HordeServer.Agents
 						break;
 					}
 
-					// If there's no matching session, try to to remove it from the document. We can expect the document to be authoritative here, since we always create the session before
+					// If there's no matching session, try to remove it from the document. We can expect the document to be authoritative here, since we always create the session before
 					// updating the document; if it no longer exists, the document is invalid.
 					document = await TryUpdateAsync(document, Builders<AgentDocument>.Update.Unset(x => x.SessionId).Set(x => x.LastOnlineTime, _clock.UtcNow), cancellationToken);
 					if (document != null)
