Horde: Improve logging of AWS EC2 errors in fleet manager

diff --git a/Engine/Source/Programs/Horde/Plugins/Build/HordeServer.Build.Tests/Fleet/FakeAmazonEc2.cs b/Engine/Source/Programs/Horde/Plugins/Build/HordeServer.Build.Tests/Fleet/FakeAmazonEc2.cs
--- a/Engine/Source/Programs/Horde/Plugins/Build/HordeServer.Build.Tests/Fleet/FakeAmazonEc2.cs
+++ b/Engine/Source/Programs/Horde/Plugins/Build/HordeServer.Build.Tests/Fleet/FakeAmazonEc2.cs
@@ -179,7 +179,7 @@ public class FakeAmazonEc2
 
 			if (capacity <= 0)
 			{
-				throw new AmazonEC2Exception("Insufficient capacity.");
+				throw new AmazonEC2Exception("Insufficient capacity.") { ErrorCode = "InsufficientInstanceCapacity" };
 			}
 
 			OnStartInstances(request, cancellationToken);
diff --git a/Engine/Source/Programs/Horde/Plugins/Compute/HordeServer.Compute/Agents/Fleet/FleetService.cs b/Engine/Source/Programs/Horde/Plugins/Compute/HordeServer.Compute/Agents/Fleet/FleetService.cs
--- a/Engine/Source/Programs/Horde/Plugins/Compute/HordeServer.Compute/Agents/Fleet/FleetService.cs
+++ b/Engine/Source/Programs/Horde/Plugins/Compute/HordeServer.Compute/Agents/Fleet/FleetService.cs
@@ -271,7 +271,7 @@ namespace HordeServer.Agents.Fleet
 			}
 			catch (Exception ex)
 			{
-				_logger.LogInformation(ex, "Failed to scale {PoolName}", pool.Name);
+				_logger.LogInformation(ex, "Failed to scale {PoolName}. Reason: {Reason}", pool.Name, ex.Message);
 				return new ScaleResult(FleetManagerOutcome.Failure, 0, 0, "Exception during scaling. See log.");
 			}
 
diff --git a/Engine/Source/Programs/Horde/Plugins/Compute/HordeServer.Compute/Agents/Fleet/Providers/AwsRecyclingFleetManager.cs b/Engine/Source/Programs/Horde/Plugins/Compute/HordeServer.Compute/Agents/Fleet/Providers/AwsRecyclingFleetManager.cs
--- a/Engine/Source/Programs/Horde/Plugins/Compute/HordeServer.Compute/Agents/Fleet/Providers/AwsRecyclingFleetManager.cs
+++ b/Engine/Source/Programs/Horde/Plugins/Compute/HordeServer.Compute/Agents/Fleet/Providers/AwsRecyclingFleetManager.cs
@@ -45,17 +45,6 @@ public class AwsFleetManagerException : Exception
 	public AwsFleetManagerException(string message) : base(message) { }
 }
 
-/// <summary>
-/// Exception when AWS API returns insufficient capacity
-/// </summary>
-public class AwsInsufficientCapacityException : AwsFleetManagerException
-{
-	/// <inheritdoc />
-	public AwsInsufficientCapacityException(string message) : base(message)
-	{
-	}
-}
-
 /// <summary>
 /// Generic HTTP response error from AWS API
 /// </summary>
@@ -274,7 +263,7 @@ public sealed class AwsRecyclingFleetManager : IFleetManager
 				_ec2StartInstanceCounter.Add(1, new[] { instanceTypeTag, azTag, statusTag });
 				break;
 			}
-			catch (AwsInsufficientCapacityException)
+			catch (AmazonEC2Exception e) when (e.ErrorCode == "InsufficientInstanceCapacity")
 			{
 				KeyValuePair<string, object?> statusTag = KeyValuePair.Create("status", (object?)"errorCapacity");
 				_ec2StartInstanceCounter.Add(1, new[] { instanceTypeTag, azTag, statusTag });
@@ -315,7 +304,6 @@ public sealed class AwsRecyclingFleetManager : IFleetManager
 		startRequest.InstanceIds.AddRange(instanceIds);
 
 		StartInstancesResponse startResponse;
-		try
 		{
 			using TelemetrySpan ec2Span = _tracer.StartActiveSpan("StartInstancesEc2");
 			span.SetAttribute("req.instanceIds", String.Join(',', startRequest.InstanceIds));
@@ -323,16 +311,7 @@ public sealed class AwsRecyclingFleetManager : IFleetManager
 			span.SetAttribute("res.statusCode", (int)startResponse.HttpStatusCode);
 			span.SetAttribute("res.numInstances", startResponse.StartingInstances.Count);
 		}
-		catch (AmazonEC2Exception e)
-		{
-			if (e.Message.Contains("insufficient capacity", StringComparison.OrdinalIgnoreCase))
-			{
-				throw new AwsInsufficientCapacityException($"Not enough capacity in AZ when using instance type {instanceType}");
-			}
-
-			throw;
-		}
-
+		
 		if (!IsResponseSuccessful(startResponse))
 		{
 			throw new AwsBadHttpResponseException(startResponse.HttpStatusCode);
