Horde: Adding symbolic link support to vite configuration

diff --git a/Engine/Source/Programs/Horde/HordeDashboard/tsconfig.json b/Engine/Source/Programs/Horde/HordeDashboard/tsconfig.json
--- a/Engine/Source/Programs/Horde/HordeDashboard/tsconfig.json
+++ b/Engine/Source/Programs/Horde/HordeDashboard/tsconfig.json
@@ -25,6 +25,8 @@
       "hordePlugins/*": ["./plugins/*"]
     },
 
+    "preserveSymlinks": true,
+
     /* Linting */
     "strict": false,
     "noFallthroughCasesInSwitch": true,
diff --git a/Engine/Source/Programs/Horde/HordeDashboard/vite.config.ts b/Engine/Source/Programs/Horde/HordeDashboard/vite.config.ts
--- a/Engine/Source/Programs/Horde/HordeDashboard/vite.config.ts
+++ b/Engine/Source/Programs/Horde/HordeDashboard/vite.config.ts
@@ -35,6 +35,9 @@ export default defineConfig({
       },
     }
   },
+  resolve: {
+    preserveSymlinks: true
+  },
   preview: {
     proxy: {
       '/api': {
diff --git a/Engine/Source/Programs/Horde/HordeServer/HordeServer.csproj b/Engine/Source/Programs/Horde/HordeServer/HordeServer.csproj
--- a/Engine/Source/Programs/Horde/HordeServer/HordeServer.csproj
+++ b/Engine/Source/Programs/Horde/HordeServer/HordeServer.csproj
@@ -122,6 +122,13 @@
     <ProjectReference Include="..\Plugins\Tools\HordeServer.Tools\HordeServer.Tools.csproj" />
   </ItemGroup>
 
+   <Choose>
+      <When Condition="Exists('..\..\..\..\Restricted\NotForLicensees\Source\Programs\Horde\EpicPlugins\BuildHealth\HordeServer.Epic.BuildHealth\HordeServer.Epic.BuildHealth.csproj')">
+         <ItemGroup>
+            <ProjectReference Include="..\..\..\..\Restricted\NotForLicensees\Source\Programs\Horde\EpicPlugins\BuildHealth\HordeServer.Epic.BuildHealth\HordeServer.Epic.BuildHealth.csproj" />
+         </ItemGroup>
+      </When>
+   </Choose>		
   <ItemGroup>
     <Content Include="..\ThirdParty\**" Condition="$([MSBuild]::IsOSPlatform('Windows'))">
       <CopyToOutputDirectory>Always</CopyToOutputDirectory>
