[Horde][BuildHealth][Dashboard] Linked selection mode, clear all filters functionality, and hierarchical clearing bug (edge case of last item).

diff --git a/Engine/Source/Programs/Horde/HordeDashboard/src/components/buildhealth/BuildHealthDataHandler.ts b/Engine/Source/Programs/Horde/HordeDashboard/src/components/buildhealth/BuildHealthDataHandler.ts
--- a/Engine/Source/Programs/Horde/HordeDashboard/src/components/buildhealth/BuildHealthDataHandler.ts
+++ b/Engine/Source/Programs/Horde/HordeDashboard/src/components/buildhealth/BuildHealthDataHandler.ts
@@ -3,9 +3,10 @@
 import backend from "horde/backend";
 import { JobData, JobStreamQuery, ProjectData, StreamData } from "horde/backend/Api";
 import { makeObservable, observable, runInAction } from "mobx";
-import { BuildHealthOptionsController, BuildHealthOptionsState } from "./BuildHealthOptions";
+import { BuildHealthOptionsController, BuildHealthOptionsState, IHierarchyValidator } from "./BuildHealthOptions";
 import { StepRefData, TemplateRefData } from "./BuildHealthDataTypes";
-import { decodeTemplateKey } from "./BuildHealthUtilities";
+import { decodeTemplateKey, encodeTemplateKey, encodeTemplateKeyFromStrings } from "./BuildHealthUtilities";
+import { find } from "@fluentui/react";
 
 // #region -- API Data Types & Utilities --
 
@@ -32,11 +33,11 @@ export function expandRefreshRequest(request: DataHandlerRefreshRequest): DataHa
         expanded.streams = true;
     }
 
-    if (request.streams) {
+    if (expanded.streams) {
         expanded.jobs = true;
     }
 
-    if (request.jobs) {
+    if (expanded.jobs) {
         expanded.steps = true;
     }
 
@@ -49,7 +50,7 @@ export function expandRefreshRequest(request: DataHandlerRefreshRequest): DataHa
  * Data handler for all build health data used for dropdowns.
  * @todo Data sharing & caching would ideally occur between the BuildHealth View & Step Outcome.
  */
-export class BuildHealthDataHandler {
+export class BuildHealthDataHandler implements IHierarchyValidator {
     private readonly STREAMS_FILTER = "id,state,streamId,name,batches";
     private readonly STEP_DATA_SAMPLE_SIZE = 5;
 
@@ -59,6 +60,7 @@ export class BuildHealthDataHandler {
     private buildHealthOptions: BuildHealthOptionsState;
     private buildHealthOptionsController: BuildHealthOptionsController;
     private cachedStreamToTemplatesData: Map<string, TemplateRefData[]> = new Map<string, TemplateRefData[]>();
+    private cachedStreamTempalteIdToSteps: Map<string, StepRefData[]> = new Map<string, StepRefData[]>();
 
     // #endregion -- Private Members --
 
@@ -99,6 +101,15 @@ export class BuildHealthDataHandler {
 
     // #region -- Public API --
 
+    /**
+     * Resets the handler, retunring it to a project-only data state.
+     */
+    reset() {
+        this.streamsData.length = 0;
+        this.templatesData.length = 0;
+        this.stepData.length = 0;
+    }
+
     /**
      * Requests a hierarchical data refresh. Refreshw ill occur if the options have changed since last refresh.
      * @param request The data refresh request message to use in order to control refresh depth.
@@ -177,6 +188,10 @@ export class BuildHealthDataHandler {
             runInAction(() => {
                 this.templatesData = templatesData;
             });
+        } else if (request.jobs) { // We have no enabled streams, so we should have no template data showing.
+            runInAction(() => {
+                this.templatesData = [];
+            });
         }
 
         await this.getStepData(request);
@@ -188,6 +203,7 @@ export class BuildHealthDataHandler {
      */
     async getStepData(request: DataHandlerRefreshRequest) {
         const entries = Object.entries(this.buildHealthOptions.enabledTemplates);
+        this.cachedStreamTempalteIdToSteps.clear();
         if (request.steps && entries.length > 0) {
             const results = await Promise.all(entries.map(async ([key]) => {
                 const templateData = decodeTemplateKey(key);
@@ -214,6 +230,8 @@ export class BuildHealthDataHandler {
                 }
 
                 const distinctSteps = [...stepMap.values()];
+                this.cachedStreamTempalteIdToSteps.set(key, distinctSteps);
+
                 return distinctSteps;
             }));
 
@@ -222,8 +240,46 @@ export class BuildHealthDataHandler {
             runInAction(() => {
                 this.stepData = [...allSteps.values()];
             });
+        } else if (request.steps) {  // We have no enabled templates, so we should have no step data showing.
+            runInAction(() => {
+                this.stepData = [];
+            });
         }
     }
 
+    /**
+     * @inheritdoc
+     */
+    isValidTemplate(streamId: string, templateId: string): boolean {
+        if (!this.cachedStreamToTemplatesData.has(streamId)) {
+            return false;
+        }
+        if (this.cachedStreamToTemplatesData.get(streamId)?.find(x => x.id === templateId) === undefined) {
+            return false;
+        }
+
+        return true;
+    }
+
+    /**
+     * @inheritdoc
+     */
+    isValidStep(streamId: string, templateId: string, stepName: string): boolean {
+        if (!this.cachedStreamToTemplatesData.has(streamId)) {
+            return false;
+        }
+        let template = this.cachedStreamToTemplatesData.get(streamId)?.find(x => x.id === templateId)
+
+        if (template === undefined) {
+            return false;
+        }
+        let encodedTemplateKey: string = encodeTemplateKeyFromStrings(streamId, templateId);
+        if (!this.cachedStreamTempalteIdToSteps.get(encodedTemplateKey)?.find(x => x.name == stepName)) {
+            return false;
+        }
+
+        return true;
+    }
+
     // #endregion -- Public API --
 }
\ No newline at end of file
diff --git a/Engine/Source/Programs/Horde/HordeDashboard/src/components/buildhealth/BuildHealthOptions.ts b/Engine/Source/Programs/Horde/HordeDashboard/src/components/buildhealth/BuildHealthOptions.ts
--- a/Engine/Source/Programs/Horde/HordeDashboard/src/components/buildhealth/BuildHealthOptions.ts
+++ b/Engine/Source/Programs/Horde/HordeDashboard/src/components/buildhealth/BuildHealthOptions.ts
@@ -3,7 +3,7 @@
 import { action, makeObservable, observable, runInAction } from "mobx";
 import { INCLUDE_PREFLIGHT_PARAM, PROJECTS_URL_PARAM_NAME, STEPS_URL_PARAM_NAME, STREAMS_URL_PARAM_NAME, TEMPLATES_URL_PARAM_NAME, TIME_SPAN_URL_PARAM } from "../buildhealth/BuildHealthDataTypes";
 import { BuildHealthDataHandler } from "./BuildHealthDataHandler";
-import { decodeStepKey, decodeTemplateKey, encodeStepKey, encodeTemplateKey } from "./BuildHealthUtilities";
+import { decodeStepKey, decodeTemplateKey, encodeStepKey, encodeStepNameFromStringsPreserveCase, encodeTemplateKey, encodeTemplateKeyFromStrings } from "./BuildHealthUtilities";
 import { Location } from 'react-router-dom';
 import LZString from "lz-string"
 
@@ -224,6 +224,7 @@ export class BuildHealthOptionsState {
     @observable enabledSteps: Record<string, string> = {};
     @observable jobHistoryTimeSpan: TimeSpan = DEFAULT_JOB_HISTORY_TIME_SPAN;
     @observable includePreflight: boolean = false;
+    @observable linkedSelectionModeEnabled: boolean = false;
 
     // #endregion -- User Selected Options --
 
@@ -241,6 +242,28 @@ export class BuildHealthOptionsState {
     // #endregion -- Constructor --
 }
 
+/**
+ * An interface that can provide hierarchy checks on stream, template, and step relationships.
+ */
+export interface IHierarchyValidator {
+    /**
+     * Determines whether a stream & template combination is valid (in that it exists).
+     * @param streamId The stream.
+     * @param templateId The template.
+     * @returns True if this stream has a template of the given id, false otherwise.
+     */
+    isValidTemplate(streamId: string, templateId: string): boolean;
+
+/**
+     * Determines whether a stream, template, and step combination is valid (in that it exists).
+     * @param streamId The stream.
+     * @param templateId The template.
+     * @param stepName The step name.
+     * @returns True if this stream has a template of the given id, which has a step of the given name, false otherwise.
+     */
+    isValidStep(streamId: string, templateId: string, stepName: string): boolean;
+}
+
 /**
  * Controller responsible for modifying the underlying BuildHealthOptionsState.
  */
@@ -249,7 +272,8 @@ export class BuildHealthOptionsController {
 
     readonly state: BuildHealthOptionsState;
     readonly uiState: BuildHealthUIOptionsState;
-
+    readonly handler: BuildHealthDataHandler;
+    private hierarchyValidator?: IHierarchyValidator;
     private lastSynchronize = -1;
 
     // #endregion -- Private Members --
@@ -274,8 +298,64 @@ export class BuildHealthOptionsController {
 
     // #endregion -- Constructor --
 
+    // #region -- Private API --
+
+    private toggleTemplateInternal(id: string, name: string, enabled: boolean) {
+        if (enabled) {
+            this.state.enabledTemplates[id] = name;
+        } else {
+            delete this.state.enabledTemplates[id];
+            let { streamId, templateId } = decodeTemplateKey(id);
+            this.clearSteps({ streamId, templateId });
+        }
+    }
+
+    private toggleLinkedTemplates(sourceId: string, name: string, enabled: boolean) {
+        let { streamId, templateId } = decodeTemplateKey(sourceId);
+        Object.keys(this.state.enabledStreams).forEach(streamKey => {
+            if (streamKey !== streamId) {
+                let linkedTempalteId = encodeTemplateKeyFromStrings(streamKey, templateId);
+                if (this.hierarchyValidator?.isValidTemplate(streamKey, templateId)) { // we need to see if this is a valid pair?
+                    this.toggleTemplateInternal(linkedTempalteId, name, enabled);
+                }
+            }
+        });
+    }
+
+    private toggleStepInternal(id: string, name: string, enabled: boolean) {
+        if (enabled) {
+            this.state.enabledSteps[id] = name;
+        } else {
+            delete this.state.enabledSteps[id];
+        }
+    }
+
+    private toggleLinkedSteps(sourceId: string, name: string, enabled: boolean) {
+        let sourceDecodedStepKey = decodeStepKey(sourceId);
+        Object.keys(this.state.enabledTemplates).forEach(key => {
+            let targetDecodedTempalteKey = decodeTemplateKey(key);
+            // Set this step on every other enabled template, except the narrow case of the stream-template-step - we've just come from there.
+            if (targetDecodedTempalteKey.templateId !== sourceDecodedStepKey.templateId || targetDecodedTempalteKey.streamId !== sourceDecodedStepKey.streamId) {
+                let linkedEncodedStepId = encodeStepNameFromStringsPreserveCase(targetDecodedTempalteKey.streamId, targetDecodedTempalteKey.templateId, sourceDecodedStepKey.stepName);
+                if (this.hierarchyValidator?.isValidStep(targetDecodedTempalteKey.streamId, targetDecodedTempalteKey.templateId, sourceDecodedStepKey.stepName)) { // we need to see if this is a valid pair?
+                    this.toggleStepInternal(linkedEncodedStepId, name, enabled);
+                }
+            }
+        });
+    }
+
+    // #endregion -- Private API --
+
     // #region -- Public API --
 
+    /**
+     * Sets the hierarchy validator for the controller.
+     * @param validator The hierarchy valildator to us.
+     */
+    setHierarchyValidator(validator: IHierarchyValidator) {
+        this.hierarchyValidator = validator;
+    }
+
     /**
      * Toggles the state of the provided project. Will clear all other set projects.
      * @param id The id of the project to toggle.
@@ -289,6 +369,18 @@ export class BuildHealthOptionsController {
         this.setOptionsChanged();
     }
 
+    /**
+     * Clears all currently selected options.
+     */
+    @action
+    clearAll() {
+        Object.keys(this.state.enabledProjects).forEach(key => {
+            delete this.state.enabledProjects[key];
+        });
+        this.clearStreams(true);
+        this.setOptionsChanged();
+    }
+
     /**
      * Toggles the state of the provided project.
      * @param id The id of the project to toggle.
@@ -344,16 +436,43 @@ export class BuildHealthOptionsController {
      */
     @action
     toggleTemplate(id: string, name: string, enabled: boolean) {
-        if (enabled) {
-            this.state.enabledTemplates[id] = name;
-        } else {
-            delete this.state.enabledTemplates[id];
-            let { streamId, templateId } = decodeTemplateKey(id);
-            this.clearSteps({ streamId, templateId });
+        this.toggleTemplateInternal(id, name, enabled);
+        if (this.state.linkedSelectionModeEnabled) {
+            this.toggleLinkedTemplates(id, name, enabled);
         }
         this.setOptionsChanged();
     }
 
+    /**
+     * Toggles all linkedable templates, for all active streams.
+     * @param enabled True if we want to enable, false otherwise.
+     */
+    @action
+    toggleAllLinkableTemplates(enabled: boolean) {
+        if (this.state.linkedSelectionModeEnabled) {
+            Object.keys(this.state.enabledTemplates).forEach(templateKey => {
+                let templateName = this.state.enabledTemplates[templateKey];
+                this.toggleLinkedTemplates(templateKey, templateName, enabled);
+            });
+            this.setOptionsChanged();
+        }
+    }
+
+    /**
+     * Toggles all linkedable steps, for all active templates.
+     * @param enabled True if we want to enable, false otherwise.
+     */
+    @action
+    toggleAllLinkableSteps(enabled: boolean) {
+        if (this.state.linkedSelectionModeEnabled) {
+            Object.keys(this.state.enabledSteps).forEach(stepKey => {
+                let stepName = this.state.enabledSteps[stepKey];
+                this.toggleLinkedSteps(stepKey, stepName, enabled);
+            });
+            this.setOptionsChanged();
+        }
+    }
+
     /**
      * Toggles the state of the provided step.
      * @param id The id of the step to toggle.
@@ -362,10 +481,9 @@ export class BuildHealthOptionsController {
      */
     @action
     toggleStep(id: string, name: string, enabled: boolean) {
-        if (enabled) {
-            this.state.enabledSteps[id] = name;
-        } else {
-            delete this.state.enabledSteps[id];
+        this.toggleStepInternal(id, name, enabled);
+        if (this.state.linkedSelectionModeEnabled) {
+            this.toggleLinkedSteps(id, name, enabled);
         }
         this.setOptionsChanged();
     }
@@ -390,6 +508,16 @@ export class BuildHealthOptionsController {
         this.setOptionsChanged();
     }
 
+    /**
+     * Sets whether to enter into linked selection mode, or not.
+     * @param linkedSelectionModeEnabled True to enter linked selection mode, false otherwise.
+     */
+    @action
+    setLinkedSelectionMode(linkedSelectionModeEnabled: boolean) {
+        this.state.linkedSelectionModeEnabled = linkedSelectionModeEnabled;
+        this.setOptionsChanged();
+    }
+
     /**
      * Sets whether the date anchors should be used when visualizing the table.
      * @param includeDateAnchors 
@@ -434,7 +562,6 @@ export class BuildHealthOptionsController {
                 this.clearTemplates(cascade, key);
             }
         });
-
     }
 
     /**
@@ -453,7 +580,6 @@ export class BuildHealthOptionsController {
                 }
             }
         });
-
     }
 
     /**
@@ -500,8 +626,6 @@ export class BuildHealthOptionsController {
         const lastTimeRange = this.state.jobHistoryTimeSpan.key.toString();
         const includePreflight = this.state.includePreflight;
 
-        if (!projects.length) return params;
-
         params.set(PROJECTS_URL_PARAM_NAME, LZString.compressToEncodedURIComponent(projects));
         params.set(STREAMS_URL_PARAM_NAME, LZString.compressToEncodedURIComponent(streams));
         params.set(TEMPLATES_URL_PARAM_NAME, LZString.compressToEncodedURIComponent(jobs));
diff --git a/Engine/Source/Programs/Horde/HordeDashboard/src/components/buildhealth/BuildHealthUtilities.ts b/Engine/Source/Programs/Horde/HordeDashboard/src/components/buildhealth/BuildHealthUtilities.ts
--- a/Engine/Source/Programs/Horde/HordeDashboard/src/components/buildhealth/BuildHealthUtilities.ts
+++ b/Engine/Source/Programs/Horde/HordeDashboard/src/components/buildhealth/BuildHealthUtilities.ts
@@ -11,6 +11,16 @@ export function encodeTemplateKey(template: TemplateRefData): string {
     return `${template.streamId}${KEY_SEPARATOR}${template.id}`;
 }
 
+/**
+ * Produces an encoded template key for unique identification.
+ * @param streamId The stream id to use in the key.
+ * @param templateId The template id to use in the key.
+ * @returns 
+ */
+export function encodeTemplateKeyFromStrings(streamId: string, templateId: string): string {
+    return `${streamId}${KEY_SEPARATOR}${templateId}`;
+}
+
 /**
  * Decodes an encoded template key, providing it's constituent parts.
  * @param templateKey The template key to decode.
@@ -33,6 +43,17 @@ export function encodeStepNameFromStrings(streamId: string, templateId: string,
     return `${streamId.toLocaleLowerCase()}${KEY_SEPARATOR}${templateId}${KEY_SEPARATOR}${stepName.toLocaleLowerCase()}`;
 }
 
+/**
+ * Produces an encoded step key for unique identification. Preserves case of step name.
+ * @param streamId The streamId to encode.
+ * @param templateId The templateId to encode.
+ * @param stepName The stepName to encode.
+ * @returns The key representing the step.
+ */
+export function encodeStepNameFromStringsPreserveCase(streamId: string, templateId: string, stepName: string) {
+    return `${streamId.toLocaleLowerCase()}${KEY_SEPARATOR}${templateId}${KEY_SEPARATOR}${stepName}`;
+}
+
 /**
  * Produces an encoded step key for unique identification.
  * @param step The step to produce the key for.
diff --git a/Engine/Source/Programs/Horde/HordeDashboard/src/components/buildhealth/BuildhealthView.tsx b/Engine/Source/Programs/Horde/HordeDashboard/src/components/buildhealth/BuildhealthView.tsx
--- a/Engine/Source/Programs/Horde/HordeDashboard/src/components/buildhealth/BuildhealthView.tsx
+++ b/Engine/Source/Programs/Horde/HordeDashboard/src/components/buildhealth/BuildhealthView.tsx
@@ -1,9 +1,9 @@
 // Copyright Epic Games, Inc. All Rights Reserved.
 
-import { Checkbox, ComboBox, DirectionalHint, Dropdown, DropdownMenuItemType, IDropdownOption, IDropdownProps, ISelectableOption, Label, SearchBox, Stack, TooltipHost } from "@fluentui/react";
+import { Checkbox, ComboBox, DirectionalHint, Dropdown, DropdownMenuItemType, IconButton, IDropdownOption, IDropdownProps, ISelectableOption, Label, PrimaryButton, SearchBox, Stack, TextField, TooltipHost } from "@fluentui/react";
 import { GetProjectResponse, GetStreamResponse } from "horde/backend/Api";
 import { observer } from "mobx-react-lite";
-import { useCallback, useEffect, useMemo, useRef, useState } from "react";
+import { useCallback, useEffect, useMemo, useState } from "react";
 import { useLocation, useNavigate } from "react-router";
 import { useWindowSize } from "../../base/utilities/hooks";
 import { getHordeStyling } from "../../styles/Styles";
@@ -15,6 +15,8 @@ import { KEY_SEPARATOR, PARAMETER_KEY_PREFIX, StepRefData, TemplateRefData } fro
 import { BuildHealthOptionsController, BuildHealthOptionsState, BuildHealthUIOptionsState, JobHistoryTimeSpans, loadBuildHealthOptionsFromParams, parseBuildHealthQueryParams, TimeSpan } from "./BuildHealthOptions";
 import { encodeStepKey, encodeTemplateKey } from "./BuildHealthUtilities";
 import { stepOutcomeTableClasses } from "./stepoutcome/StepOutcomeSharedUIComponents";
+import { runInAction } from "mobx";
+import { StepOutcomeLegend } from "./stepoutcome/StepOutcomeLegend";
 
 // #region -- Dropdown Components --
 
@@ -34,7 +36,7 @@ type BuildHealthOptionList = {
 
 // #endregion -- Dropdown Data Types --
 
-// #region -- Base Combobox Components --
+// #region -- Base Dropdown Components --
 
 interface ISearchableDropdownProps extends IDropdownProps {
     onSearchValueChanged?: (searchValue: string) => void;
@@ -88,11 +90,6 @@ const SearchableDropdown: React.FC<ISearchableDropdownProps> = (props) => {
     );
 };
 
-
-// #endregion -- Base Combobox Components --
-
-// #region -- Base Dropdown Components --
-
 /**
  * React Component that is a single select drop down.
  * @returns React Component.
@@ -283,7 +280,7 @@ const ProjectDropdownSingle: React.FC<{ handler: BuildHealthDataHandler, buildHe
         return listParam;
     }).sort((a, b) => a.text.localeCompare(b.text));
 
-    let selectedProjectKey = Object.keys(buildHealthOptions.state.enabledProjects)[0] ?? undefined;
+    let selectedProjectKey = Object.keys(buildHealthOptions.state.enabledProjects)[0] ?? null;
 
     // Note: This may look odd, but there is really no reason to force a full hierarchy of project refresh; we just want to pulse streams
     let request: DataHandlerRefreshRequest = expandRefreshRequest({ streams: true });
@@ -355,12 +352,31 @@ const TemplateDropdown: React.FC<{ handler: BuildHealthDataHandler, buildHealthO
     return templateDropdownParams.items.length > 0 && (
         <Stack id={`job_dropdown`}>
             <Label>Template(s)</Label>
-            <MultiListParameter handler={handler} buildHealthOptions={buildHealthOptions} param={templateDropdownParams} selectedKeys={selectedTemplateKeys} request={request} onChange={(option: IDropdownOption<BuildHealthOptionData>) => {
-                const key = option.key as string;
-                const isSelected = !!option.selected;
-
-                buildHealthOptions.toggleTemplate(key, option.text, isSelected);
-            }} />
+            <Stack horizontal tokens={{ childrenGap: 8 }} verticalAlign="end">
+                <MultiListParameter handler={handler} buildHealthOptions={buildHealthOptions} param={templateDropdownParams} selectedKeys={selectedTemplateKeys} request={request} onChange={(option: IDropdownOption<BuildHealthOptionData>) => {
+                    const key = option.key as string;
+                    const isSelected = !!option.selected;
+
+                    buildHealthOptions.toggleTemplate(key, option.text, isSelected);
+                }} />
+                {
+                    buildHealthOptions.state.linkedSelectionModeEnabled &&
+                    (
+                        <TooltipHost
+                            content={"Link templates: Apply selected templates across selected streams."}
+                            calloutProps={{ directionalHint: DirectionalHint.rightCenter }}>
+                            <IconButton
+                                iconProps={{ iconName: 'Link' }}
+                                styles={{ root: { width: 15, minWidth: 15, padding: 0 } }} onClick={() => {
+                                    buildHealthOptions.toggleAllLinkableTemplates(true);
+                                    handler.requestHierarchicalRefresh(request);
+                                    buildHealthOptions.synchronizeDerivedKeys();
+                                }
+                                } />
+                        </TooltipHost>
+                    )
+                }
+            </Stack>
         </Stack>
     );
 });
@@ -388,12 +404,31 @@ const StepDropdown: React.FC<{ handler: BuildHealthDataHandler, buildHealthOptio
     return stepDropdownParams.items.length > 0 && (
         <Stack id={`step_dropdown`}>
             <Label>Step(s)</Label>
-            <MultiListParameter handler={handler} buildHealthOptions={buildHealthOptions} param={stepDropdownParams} selectedKeys={selectedStepKeys} request={request} onChange={(option: IDropdownOption<BuildHealthOptionData>) => {
-                const key = option.key as string;
-                const isSelected = !!option.selected;
-
-                buildHealthOptions.toggleStep(key, option.text, isSelected);
-            }} />
+            <Stack horizontal tokens={{ childrenGap: 8 }} verticalAlign="end">
+                <MultiListParameter handler={handler} buildHealthOptions={buildHealthOptions} param={stepDropdownParams} selectedKeys={selectedStepKeys} request={request} onChange={(option: IDropdownOption<BuildHealthOptionData>) => {
+                    const key = option.key as string;
+                    const isSelected = !!option.selected;
+
+                    buildHealthOptions.toggleStep(key, option.text, isSelected);
+                }} />
+                {
+                    buildHealthOptions.state.linkedSelectionModeEnabled &&
+                    (
+                        <TooltipHost
+                            content={"Link steps: Apply selected steps across selected templates."}
+                            calloutProps={{ directionalHint: DirectionalHint.rightCenter }}>
+                            <IconButton
+                                iconProps={{ iconName: 'Link' }}
+                                styles={{ root: { width: 15, minWidth: 15, padding: 0 } }} onClick={() => {
+                                    buildHealthOptions.toggleAllLinkableSteps(true);
+                                    handler.requestHierarchicalRefresh(request);
+                                    buildHealthOptions.synchronizeDerivedKeys();
+                                }
+                                } />
+                        </TooltipHost>
+                    )
+                }
+            </Stack>
         </Stack>
     );
 });
@@ -439,6 +474,7 @@ const BuildHealthSidebar: React.FC<{ handler: BuildHealthDataHandler, buildHealt
     const { hordeClasses, modeColors } = getHordeStyling();
     const [timeSpanKey, setTimeSpanKey] = useState(() => buildHealthOptions.state.jobHistoryTimeSpan.key);
     const [includePreflight, setIncludePreflight] = useState(() => buildHealthOptions.state.includePreflight)
+    const [linkedSelectionMode, setLinkedSelectionMode] = useState(() => buildHealthOptions.state.linkedSelectionModeEnabled)
 
     // #region -- Visual Options -- 
 
@@ -460,6 +496,11 @@ const BuildHealthSidebar: React.FC<{ handler: BuildHealthDataHandler, buildHealt
         buildHealthOptions.setIncludePreflight(isChecked);
     }
 
+    const updateLinkedSelectionMode = (isChecked) => {
+        setLinkedSelectionMode(isChecked);
+        buildHealthOptions.setLinkedSelectionMode(isChecked);
+    }
+
     const updateIncludeDateAnchors = (isChecked) => {
         setHideDateAnchors(isChecked);
         buildHealthOptions.setHideDateAnchors(isChecked);
@@ -495,7 +536,7 @@ const BuildHealthSidebar: React.FC<{ handler: BuildHealthDataHandler, buildHealt
     let timeComboWidth = 150;
 
     return (
-        <Stack style={{ minWidth: 300, maxWidth: 450, paddingRight: 18 }}>
+        <Stack style={{ minWidth: 375, maxWidth: 500, paddingRight: 18 }}>
             <Stack className={hordeClasses.modal} tokens={{ childrenGap: 12 }}>
                 <Stack
                     styles={{
@@ -509,11 +550,37 @@ const BuildHealthSidebar: React.FC<{ handler: BuildHealthDataHandler, buildHealt
                     <Stack style={{ paddingTop: 0, paddingBottom: 4 }}>
                         <Label>Step Outcome Filters</Label>
                     </Stack>
-                    <ProjectDropdownSingle handler={handler} buildHealthOptions={buildHealthOptions} />
-                    <StreamDropdown handler={handler} buildHealthOptions={buildHealthOptions} />
-                    <TemplateDropdown handler={handler} buildHealthOptions={buildHealthOptions} />
-                    <StepDropdown handler={handler} buildHealthOptions={buildHealthOptions} />
-                    <Stack>
+                    <Stack styles={{
+                        root: {
+                            border: `1px solid ${modeColors.content}`,
+                            borderRadius: 4,
+                            padding: 8,
+                        },
+                    }}>
+                        <ProjectDropdownSingle handler={handler} buildHealthOptions={buildHealthOptions} />
+                        <StreamDropdown handler={handler} buildHealthOptions={buildHealthOptions} />
+                        <TemplateDropdown handler={handler} buildHealthOptions={buildHealthOptions} />
+                        <StepDropdown handler={handler} buildHealthOptions={buildHealthOptions} />
+                        <TooltipHost content={
+                            <Stack style={{ maxWidth: 300 }}>
+                                <Stack>
+                                    Clears the current Project/Stream/Template/Step selection.
+                                </Stack>
+                            </Stack>
+                        }
+                            calloutProps={{
+                                directionalHint: DirectionalHint.rightCenter
+                            }}
+                        >
+                            <PrimaryButton text="Clear" onClick={() => {
+
+                                runInAction(() => {
+                                    buildHealthController.clearAll();
+                                    handler.reset();
+                                    buildHealthOptions.synchronizeDerivedKeys();
+                                });
+                            }} />
+                        </TooltipHost>
                         <TooltipHost content="Adjust the historical job history time span."
                             calloutProps={{
                                 directionalHint: DirectionalHint.rightCenter
@@ -530,6 +597,8 @@ const BuildHealthSidebar: React.FC<{ handler: BuildHealthDataHandler, buildHealt
                                 buildHealthOptions.synchronizeDerivedKeys();
                             }}
                         />
+                    </Stack>
+                    <Stack tokens={{ childrenGap: 5 }}>
                         <TooltipHost content={
                             <div style={{ maxWidth: 300 }}>
                                 <div>
@@ -542,11 +611,28 @@ const BuildHealthSidebar: React.FC<{ handler: BuildHealthDataHandler, buildHealt
                             }}
                         >
                             <Checkbox
-                                label="Include presubmits"
+                                label="Include Presubmits"
                                 checked={includePreflight}
                                 onChange={(_, isChecked) => updateIncludePreflight(!!isChecked)}
                             />
                         </TooltipHost>
+                        <TooltipHost content={
+                            <div style={{ maxWidth: 300 }}>
+                                <div>
+                                    Linked editing mode: Links filter selects across streams & templates where applicable.
+                                </div>
+                            </div>
+                        }
+                            calloutProps={{
+                                directionalHint: DirectionalHint.rightCenter
+                            }}
+                        >
+                            <Checkbox
+                                label="Linked Selection Mode"
+                                checked={linkedSelectionMode}
+                                onChange={(_, isChecked) => updateLinkedSelectionMode(!!isChecked)}
+                            />
+                        </TooltipHost>
                     </Stack>
                 </Stack>
                 <Stack
@@ -595,6 +681,16 @@ const BuildHealthSidebar: React.FC<{ handler: BuildHealthDataHandler, buildHealt
                         />
                     </TooltipHost>
                 </Stack>
+                <Stack
+                    styles={{
+                        root: {
+                            border: `1px solid ${modeColors.content}`,
+                            borderRadius: 4,
+                            padding: 8,
+                        },
+                    }}>
+                    <StepOutcomeLegend />
+                </Stack>
             </Stack>
         </Stack>
     );
@@ -671,5 +767,6 @@ const buildHealthState = new BuildHealthOptionsState();
 const buildHealthUIState = new BuildHealthUIOptionsState();
 const buildHealthController = new BuildHealthOptionsController(buildHealthState, buildHealthUIState);
 const handler = new BuildHealthDataHandler(buildHealthState, buildHealthController);
+buildHealthController.setHierarchyValidator(handler);
 
 // #endregion -- Script --
\ No newline at end of file
diff --git a/Engine/Source/Programs/Horde/HordeDashboard/src/components/buildhealth/stepoutcome/StepOutcome.tsx b/Engine/Source/Programs/Horde/HordeDashboard/src/components/buildhealth/stepoutcome/StepOutcome.tsx
--- a/Engine/Source/Programs/Horde/HordeDashboard/src/components/buildhealth/stepoutcome/StepOutcome.tsx
+++ b/Engine/Source/Programs/Horde/HordeDashboard/src/components/buildhealth/stepoutcome/StepOutcome.tsx
@@ -1,13 +1,13 @@
 // Copyright Epic Games, Inc. All Rights Reserved.
 
-import { Dialog, DialogType, Label, mergeStyleSets, PrimaryButton, ScrollablePane, Spinner, SpinnerSize, Stack } from "@fluentui/react";
+import { Dialog, DialogType, Label, mergeStyleSets, ScrollablePane, Spinner, SpinnerSize, Stack } from "@fluentui/react";
 import { Text } from '@fluentui/react/lib/Text';
 import { observer } from "mobx-react-lite";
 import React, { useEffect, useState } from 'react';
 import { StepOutcomeD3Table } from "./StepOutcomeD3Component";
 import { prettyPrintStepOutcomeFilters, prettyPrintStepOutcomeFiltersSummary, StepOutcomeDataHandler, StepOutcomeFilters, StepOutcomeViewOptions } from "./StepOutcomeDataHandler";
 import { AbortedStr, StepOutcomeTableEntry } from "./StepOutcomeDataTypes";
-import { sharedBorderStyle, stepOutcomeTableClasses } from "./StepOutcomeSharedUIComponents";
+import { sharedBorderStyle } from "./StepOutcomeSharedUIComponents";
 import { isTerminalState, toDate } from "./StepOutcomeUtilities";
 
 // #region -- View Components --
@@ -63,7 +63,7 @@ const StepOutcomeFocusModal: React.FC<StepOutcomeFocusModalProps> = ({ item, onD
         ['State', item.stepResponse.state],
         // For simplicity, we will compeltely override the outcome of this step with the job's outcome, as the step was aborted.
         item.stepResponse.state === AbortedStr ? ['Job Outcome', handler.getJobResponseData(item.jobId)?.summarizedOutcome ?? "-"] : isTerminalState(item) ? ['Outcome', item.stepResponse.outcome] : null,
-        ['Job Link', <a href={`/job/${item.jobId}`} target="_blank" rel="noopener noreferrer">View Job</a>],
+        ['Job Link', <a href={`/job/${item.jobId}`} target="_blank" rel="noopen er noreferrer">View Job</a>],
         ['Step Link', <a href={`/job/${item.jobId}?step=${item.stepResponse.id}`} target="_blank" rel="noopener noreferrer">View Step</a>],
     ].filter((d): d is (string | JSX.Element)[] => Boolean(d));
 
diff --git a/Engine/Source/Programs/Horde/HordeDashboard/src/components/buildhealth/stepoutcome/StepOutcomeD3Component.tsx b/Engine/Source/Programs/Horde/HordeDashboard/src/components/buildhealth/stepoutcome/StepOutcomeD3Component.tsx
--- a/Engine/Source/Programs/Horde/HordeDashboard/src/components/buildhealth/stepoutcome/StepOutcomeD3Component.tsx
+++ b/Engine/Source/Programs/Horde/HordeDashboard/src/components/buildhealth/stepoutcome/StepOutcomeD3Component.tsx
@@ -1,6 +1,6 @@
 // Copyright Epic Games, Inc. All Rights Reserved.
 
-import { DefaultButton, DirectionalHint, mergeStyleSets, Stack, TooltipHost } from "@fluentui/react";
+import { DefaultButton, DirectionalHint, IconButton, mergeStyleSets, Stack, TooltipHost } from "@fluentui/react";
 import * as d3 from "d3";
 import { getHordeStyling, preloadFonts } from "horde/styles/Styles";
 import { useEffect, useRef } from "react";
@@ -43,7 +43,7 @@ function saveFile(content: string, filename: string) {
     URL.revokeObjectURL(url);
 }
 
-function exportStepOutcomeDataToCSV(stepOutcomeTable: StepOutcomeTable, handler : StepOutcomeDataHandler) {
+function exportStepOutcomeDataToCSV(stepOutcomeTable: StepOutcomeTable, handler: StepOutcomeDataHandler) {
     const defaultName = new Date()
         .toLocaleDateString("en-US", {
             year: "numeric",
@@ -509,7 +509,7 @@ function GenerateD3Table(stepOutcomeTable: StepOutcomeTable, handler: StepOutcom
                         const corrected = new Date(original.getTime() - 4 * 60 * 60 * 1000);
                         div.text(""),
                             div.style("cursor", "pointer")
-                                .style("background-color", getCellStyle(entry, stepOutcomeTable.viewModelOptions).bgColor)
+                                .style("background-color", getCellStyle(entry).bgColor)
                                 .style("box-shadow", getCellAccents(entry, handler.getJobResponseData(entry.jobId)).boxShadow)
                                 .on("click", () => onCellClick?.(entry))
                                 .on("mouseover", function () {
@@ -757,19 +757,14 @@ function GenerateD3Table(stepOutcomeTable: StepOutcomeTable, handler: StepOutcom
                     content="Export the current step outcome history data to CSV."
                     directionalHint={DirectionalHint.topLeftEdge}
                 >
-                    <DefaultButton
-                        text=">"
+                    <IconButton
+                        iconProps={{ iconName: 'Download' }}
                         styles={{
                             root: {
                                 width: 30,
                                 height: 30,
                                 minWidth: 0,
                                 padding: 0,
-                                lineHeight: "15px",
-                            },
-                            label: {
-                                fontSize: 10,
-                                lineHeight: "15px",
                             },
                         }}
                         onClick={() => exportStepOutcomeDataToCSV(stepOutcomeTable, handler)}
diff --git a/Engine/Source/Programs/Horde/HordeDashboard/src/components/buildhealth/stepoutcome/StepOutcomeLegend.tsx b/Engine/Source/Programs/Horde/HordeDashboard/src/components/buildhealth/stepoutcome/StepOutcomeLegend.tsx
new file mode 100644
--- /dev/null
+++ b/Engine/Source/Programs/Horde/HordeDashboard/src/components/buildhealth/stepoutcome/StepOutcomeLegend.tsx
@@ -0,0 +1,176 @@
+// #region -- Legend Component --
+
+import { Stack, Label, TooltipHost, DirectionalHint } from "@fluentui/react";
+import { JobStepState, JobStepOutcome } from "horde/backend/Api";
+import { JobResponseData, StepOutcomeTableEntry } from "./StepOutcomeDataTypes";
+import { getCellStyle, getCellAccents, getCellClasses } from "./StepOutcomeSharedUIComponents";
+
+export const StepOutcomeLegend: React.FC = () => {
+    const fakeSuccessEntry = {
+        stepResponse: {
+            state: JobStepState.Completed,
+            outcome: JobStepOutcome.Success
+        }
+    } as Pick<StepOutcomeTableEntry, "stepResponse"> as StepOutcomeTableEntry;
+
+    const fakeFailureEntry = {
+        stepResponse: {
+            state: JobStepState.Completed,
+            outcome: JobStepOutcome.Failure
+        }
+    } as Pick<StepOutcomeTableEntry, "stepResponse"> as StepOutcomeTableEntry;
+
+    const fakeWarningEntry = {
+        stepResponse: {
+            state: JobStepState.Completed,
+            outcome: JobStepOutcome.Warnings
+        }
+    } as Pick<StepOutcomeTableEntry, "stepResponse"> as StepOutcomeTableEntry;
+
+    const fakeRunningEntry = {
+        stepResponse: {
+            state: JobStepState.Running,
+            outcome: JobStepOutcome.Unspecified
+        }
+    } as Pick<StepOutcomeTableEntry, "stepResponse"> as StepOutcomeTableEntry;
+
+    const fakeRunningSummaryData = {
+        summarizedOutcome: JobStepOutcome.Unspecified
+    } as Partial<JobResponseData> as JobResponseData | undefined;
+
+    const fakeSkippedEntry = {
+        stepResponse: {
+            state: JobStepState.Skipped,
+            outcome: JobStepOutcome.Failure
+        }
+    } as Pick<StepOutcomeTableEntry, "stepResponse"> as StepOutcomeTableEntry;
+
+    const fakeSkippedSummaryData = {
+        summarizedOutcome: JobStepOutcome.Failure
+    } as Partial<JobResponseData> as JobResponseData | undefined;
+
+    const fakeWaitingEntry = {
+        stepResponse: {
+            state: JobStepState.Waiting,
+            outcome: JobStepOutcome.Success
+        }
+    } as Pick<StepOutcomeTableEntry, "stepResponse"> as StepOutcomeTableEntry;
+
+    const fakeReadyEntry = {
+        stepResponse: {
+            state: JobStepState.Ready,
+            outcome: JobStepOutcome.Success
+        }
+    } as Pick<StepOutcomeTableEntry, "stepResponse"> as StepOutcomeTableEntry;
+
+    const fakeReadyWatingummaryData = {
+        summarizedOutcome: JobStepOutcome.Unspecified
+    } as Partial<JobResponseData> as JobResponseData | undefined;
+
+    const fakeAbortedEntry = {
+        stepResponse: {
+            state: JobStepState.Aborted,
+            outcome: JobStepOutcome.Failure
+        }
+    } as Pick<StepOutcomeTableEntry, "stepResponse"> as StepOutcomeTableEntry;
+
+    const fakeAbortedFailureSummaryData = {
+        summarizedOutcome: JobStepOutcome.Failure
+    } as Partial<JobResponseData> as JobResponseData | undefined;
+
+    const fakeAbortedWarningSummaryData = {
+        summarizedOutcome: JobStepOutcome.Warnings
+    } as Partial<JobResponseData> as JobResponseData | undefined;
+
+    const fakeAbortedSuccessSummaryData = {
+        summarizedOutcome: JobStepOutcome.Success
+    } as Partial<JobResponseData> as JobResponseData | undefined;
+
+    return (
+        <Stack>
+            <Label>Legend</Label>
+            <Stack>
+                <table>
+                    <tbody>
+                        <tr>
+                            <td>
+                                <TooltipHost content="When a step has run to completion." calloutProps={{ directionalHint: DirectionalHint.topCenter }}>
+                                    <div style={{ width: "100%", height: "100%" }}>Complete Exection (solid fill)</div>
+                                </TooltipHost>
+                            </td>
+                            <td style={{
+                                width: 20,
+                                height: 20,
+                                backgroundColor: getCellStyle(fakeFailureEntry).bgColor,
+                            }}><TooltipHost content="The step has completed with failures." ><div style={{ width: "100%", height: "100%" }} /></TooltipHost></td>
+                            <td style={{
+                                width: 20,
+                                height: 20,
+                                backgroundColor: getCellStyle(fakeWarningEntry).bgColor,
+                            }}><TooltipHost content="The step has completed with warnings." ><div style={{ width: "100%", height: "100%" }} /></TooltipHost></td>
+                            <td style={{
+                                width: 20,
+                                height: 20,
+                                backgroundColor: getCellStyle(fakeSuccessEntry).bgColor,
+                            }}><TooltipHost content="The step has completed with success." ><div style={{ width: "100%", height: "100%" }} /></TooltipHost></td>
+                        </tr>
+                        <tr>
+                            <td>
+                                <TooltipHost content="When a step cannot execute (dependency failed), is currently executing, or hasn't yet executed." calloutProps={{ directionalHint: DirectionalHint.topCenter }}>
+                                    <div style={{ width: "100%", height: "100%" }}>Incomplete Execution (outline)</div>
+                                </TooltipHost>
+                            </td>
+                            <td style={{
+                                width: 20,
+                                height: 20,
+                                backgroundColor: getCellStyle(fakeSkippedEntry).bgColor,
+                                boxShadow: getCellAccents(fakeSkippedEntry, fakeSkippedSummaryData).boxShadow,
+                            }}><TooltipHost content="The step has not run, and has been skipped due to a dependency failing." ><div style={{ width: "100%", height: "100%" }} /></TooltipHost></td>
+                            <td style={{
+                                width: 20,
+                                height: 20,
+                                backgroundColor: getCellStyle(fakeRunningEntry).bgColor,
+                                boxShadow: getCellAccents(fakeRunningEntry, fakeRunningSummaryData).boxShadow,
+                            }}><TooltipHost content="The step is currently running." ><div style={{ width: "100%", height: "100%" }} /></TooltipHost></td>
+                            <td style={{
+                                width: 20,
+                                height: 20,
+                                backgroundColor: getCellStyle(fakeReadyEntry).bgColor,
+                                boxShadow: getCellAccents(fakeReadyEntry, fakeReadyWatingummaryData).boxShadow,
+                            }}><TooltipHost content="The step has yet to run, but is in a ready state or waiting on agent." ><div style={{ width: "100%", height: "100%" }} /></TooltipHost></td>
+                        </tr>
+                        <tr>
+                            <td>
+                                <TooltipHost content="When a step as been aborted." calloutProps={{ directionalHint: DirectionalHint.topCenter }}>
+                                    <div style={{ width: "100%", height: "100%" }}>Cancelled Execution (strike-through)</div>
+                                </TooltipHost>
+                            </td>
+                            <td className={getCellClasses(fakeAbortedEntry, fakeAbortedFailureSummaryData).join(" ")}
+                                style={{
+                                    width: 20,
+                                    height: 20,
+                                    backgroundColor: getCellStyle(fakeAbortedEntry).bgColor,
+                                    boxShadow: getCellAccents(fakeAbortedEntry, fakeAbortedFailureSummaryData).boxShadow,
+                                }}><TooltipHost content="The step has been aborted, and other steps in the job have failed." ><div style={{ width: "100%", height: "100%" }} /></TooltipHost></td>
+                            <td className={getCellClasses(fakeAbortedEntry, fakeAbortedWarningSummaryData).join(" ")}
+                                style={{
+                                    width: 20,
+                                    height: 20,
+                                    backgroundColor: getCellStyle(fakeAbortedEntry).bgColor,
+                                    boxShadow: getCellAccents(fakeAbortedEntry, fakeAbortedWarningSummaryData).boxShadow,
+                                }}><TooltipHost content="The step has been aborted, and other steps in the job encountered warnings." ><div style={{ width: "100%", height: "100%" }} /></TooltipHost></td>
+                            <td className={getCellClasses(fakeAbortedEntry, fakeAbortedSuccessSummaryData).join(" ")}
+                                style={{
+                                    width: 20,
+                                    height: 20,
+                                    backgroundColor: getCellStyle(fakeAbortedEntry).bgColor,
+                                    boxShadow: getCellAccents(fakeAbortedEntry, fakeAbortedSuccessSummaryData).boxShadow,
+                                }}><TooltipHost content="The step has been aborted, but all other steps in the job completed with success." ><div style={{ width: "100%", height: "100%" }} /></TooltipHost></td>
+                        </tr>
+                    </tbody>
+                </table>
+            </Stack>
+        </Stack >)
+}
+
+// #endregion -- Legend Component --
\ No newline at end of file
diff --git a/Engine/Source/Programs/Horde/HordeDashboard/src/components/buildhealth/stepoutcome/StepOutcomeSharedUIComponents.tsx b/Engine/Source/Programs/Horde/HordeDashboard/src/components/buildhealth/stepoutcome/StepOutcomeSharedUIComponents.tsx
--- a/Engine/Source/Programs/Horde/HordeDashboard/src/components/buildhealth/stepoutcome/StepOutcomeSharedUIComponents.tsx
+++ b/Engine/Source/Programs/Horde/HordeDashboard/src/components/buildhealth/stepoutcome/StepOutcomeSharedUIComponents.tsx
@@ -145,7 +145,7 @@ export function getCellAccents(entry: StepOutcomeTableEntry | null, jobState: Jo
  * @param viewModelOptions The view model options to consider while applying styles.
  * @returns The matching style given the state of the entry. 
  */
-export function getCellStyle(entry: StepOutcomeTableEntry | null, viewModelOptions: StepOutcomeViewOptions): { bgColor: string } {
+export function getCellStyle(entry: StepOutcomeTableEntry | null): { bgColor: string } {
     const colors: Record<string, string> = getColorRecords();
 
     let resultColor: string = 'transparent';
