Horde: Show changes from restricted in release tool

diff --git a/Engine/Source/Programs/Horde/Tools/ReleaseNotes/Program.cs b/Engine/Source/Programs/Horde/Tools/ReleaseNotes/Program.cs
--- a/Engine/Source/Programs/Horde/Tools/ReleaseNotes/Program.cs
+++ b/Engine/Source/Programs/Horde/Tools/ReleaseNotes/Program.cs
@@ -16,6 +16,7 @@ namespace CreateReleaseNotes
 		[
 			"Engine/Source/Programs/Horde/...",
 			"Engine/Source/Programs/Shared/...",
+			"Engine/Restricted/NotForLicensees/Source/Programs/Horde/...",
 		];
 		
 		private static readonly string[] s_ignoreFilter =
@@ -53,6 +54,12 @@ namespace CreateReleaseNotes
 			"-/...",
 			"/Engine/Source/Programs/Horde/HordeDashboard/..."
 		];
+		
+		static readonly string[] s_restrictedFilter =
+		[
+			"-/...",
+			"/Engine/Restricted/..."
+		];
 
 		record ChangeInfo(int Number, string Description, string? JiraTicket);
 
@@ -152,6 +159,7 @@ namespace CreateReleaseNotes
 			FileFilter serverFilter = new (s_serverFilter);
 			FileFilter agentFilter = new (s_agentFilter);
 			FileFilter dashboardFilter = new (s_dashboardFilter);
+			FileFilter restrictedFilter = new (s_restrictedFilter);
 
 			logger.LogInformation("");
 			foreach (DescribeRecord describeRecord in describeRecords.OrderBy(x => x.Number))
@@ -179,6 +187,7 @@ namespace CreateReleaseNotes
 				bool affectsServer = false;
 				bool affectsAgent = false;
 				bool affectsDashboard = false;
+				bool affectsRestricted = false;
 				bool ignoreChangelist = false;
 
 				foreach (DescribeFileRecord fileRecord in describeRecord.Files)
@@ -192,17 +201,26 @@ namespace CreateReleaseNotes
 						break;
 					}
 					
-					if (describeRecord.Number > serverChange && serverFilter.Matches(path))
-					{
-						affectsServer = true;
-					}
-					if (describeRecord.Number > agentChange && agentFilter.Matches(path))
+					if (restrictedFilter.Matches(path))
 					{
-						affectsAgent = true;
+						affectsServer = Path.GetExtension(path) is ".cs" or ".csproj";
+						affectsDashboard = Path.GetExtension(path) is ".ts" or ".tsx";
+						affectsRestricted = true;
 					}
-					if (describeRecord.Number > dashboardChange && dashboardFilter.Matches(path))
+					else
 					{
-						affectsDashboard = true;
+						if (describeRecord.Number > serverChange && serverFilter.Matches(path))
+						{
+							affectsServer = true;
+						}
+						if (describeRecord.Number > agentChange && agentFilter.Matches(path))
+						{
+							affectsAgent = true;
+						}
+						if (describeRecord.Number > dashboardChange && dashboardFilter.Matches(path))
+						{
+							affectsDashboard = true;
+						}
 					}
 				}
 				
@@ -211,14 +229,14 @@ namespace CreateReleaseNotes
 					continue;
 				}
 				
-				if (!Regex.IsMatch(describeRecord.Description, @"^\s*#rnx\s*$", RegexOptions.Multiline))
+				if (!Regex.IsMatch(describeRecord.Description, @"^\s*#rnx\s*$", RegexOptions.Multiline) && !affectsRestricted)
 				{
 					parsedChanges.Add(new ChangeInfo(describeRecord.Number, description, jiraTicket));
 				}
 
-				if (affectsServer || affectsAgent || affectsDashboard)
+				if (affectsServer || affectsAgent || affectsDashboard || affectsRestricted)
 				{
-					string types = (affectsServer? "S" : " ") + (affectsAgent ? "A" : " ") + (affectsDashboard? "D" : " ");
+					string types = (affectsServer? "S" : " ") + (affectsAgent ? "A" : " ") + (affectsDashboard? "D" : " ") + (affectsRestricted? "R" : " ");
 					logger.LogInformation("[{Types}] {Change} {Author,-20} {Description}", types, describeRecord.Number, describeRecord.User.ToLower(), logDescription);
 				}
 			}
