Horde: Fix bug in release tool restricted check

diff --git a/Engine/Source/Programs/Horde/Tools/ReleaseNotes/Program.cs b/Engine/Source/Programs/Horde/Tools/ReleaseNotes/Program.cs
--- a/Engine/Source/Programs/Horde/Tools/ReleaseNotes/Program.cs
+++ b/Engine/Source/Programs/Horde/Tools/ReleaseNotes/Program.cs
@@ -203,8 +203,8 @@ namespace CreateReleaseNotes
 					
 					if (restrictedFilter.Matches(path))
 					{
-						affectsServer = Path.GetExtension(path) is ".cs" or ".csproj";
-						affectsDashboard = Path.GetExtension(path) is ".ts" or ".tsx";
+						affectsServer |= Path.GetExtension(path) is ".cs" or ".csproj";
+						affectsDashboard |= Path.GetExtension(path) is ".ts" or ".tsx";
 						affectsRestricted = true;
 					}
 					else
