[Horde][BuildHealth][Server] Added root commit to issues

diff --git a/Engine/Source/Programs/Horde/Plugins/Build/HordeServer.Build/Issues/IIssue.cs b/Engine/Source/Programs/Horde/Plugins/Build/HordeServer.Build/Issues/IIssue.cs
--- a/Engine/Source/Programs/Horde/Plugins/Build/HordeServer.Build/Issues/IIssue.cs
+++ b/Engine/Source/Programs/Horde/Plugins/Build/HordeServer.Build/Issues/IIssue.cs
@@ -108,6 +108,11 @@ namespace HordeServer.Issues
 		/// </summary>
 		public CommitId? FixCommitId { get; }
 
+		/// <summary>
+		/// Root commit for this issue
+		/// </summary>
+		public CommitId? RootCommitId { get; }
+
 		/// <summary>
 		/// Whether the issue is marked fixed as systemic
 		/// </summary>
diff --git a/Engine/Source/Programs/Horde/Plugins/Build/HordeServer.Build/Issues/IIssueCollection.cs b/Engine/Source/Programs/Horde/Plugins/Build/HordeServer.Build/Issues/IIssueCollection.cs
--- a/Engine/Source/Programs/Horde/Plugins/Build/HordeServer.Build/Issues/IIssueCollection.cs
+++ b/Engine/Source/Programs/Horde/Plugins/Build/HordeServer.Build/Issues/IIssueCollection.cs
@@ -236,6 +236,7 @@ namespace HordeServer.Issues
 	/// <param name="DeclinedById">Name of a user that has declined the issue</param>
 	/// <param name="FixCommitId">Fix changelist for the issue. Pass a default value to clear the fix changelist.</param>
 	/// <param name="FixSystemic">Whether the issue should be marked fixed as a systemic issue</param>
+	/// <param name="RootCommitId">Root changelist for the issue. Pass an empty value to clear the root changelist.</param>
 	/// <param name="ResolvedById">User that resolved the issue (may be ObjectId.Empty to clear)</param>
 	/// <param name="ExcludeSpanIds">List of span ids to exclude from this issue</param>
 	/// <param name="LastSeenAt"></param>
@@ -256,6 +257,7 @@ namespace HordeServer.Issues
 		UserId? DeclinedById = null,
 		CommitId? FixCommitId = null,
 		bool? FixSystemic = null,
+		CommitId? RootCommitId = null,
 		UserId? ResolvedById = null,
 		List<ObjectId>? ExcludeSpanIds = null,
 		DateTime? LastSeenAt = null,
diff --git a/Engine/Source/Programs/Horde/Plugins/Build/HordeServer.Build/Issues/IssueCollection.cs b/Engine/Source/Programs/Horde/Plugins/Build/HordeServer.Build/Issues/IssueCollection.cs
--- a/Engine/Source/Programs/Horde/Plugins/Build/HordeServer.Build/Issues/IssueCollection.cs
+++ b/Engine/Source/Programs/Horde/Plugins/Build/HordeServer.Build/Issues/IssueCollection.cs
@@ -115,6 +115,16 @@ namespace HordeServer.Issues
 			[BsonIgnoreIfDefault, BsonDefaultValue(false)]
 			public bool FixSystemic { get; set; }
 
+			[BsonIgnore]
+			public CommitId? RootCommitId 
+			{
+				get => (RootCommitName != null) ? new CommitId(RootCommitName) : null;
+				set => RootCommitName = value?.Name;
+			}
+
+			[BsonElement("RootCommit"), BsonIgnoreIfNull]
+			public string? RootCommitName { get; set; }
+
 			public List<IssueStreamDocument> Streams { get; set; } = new List<IssueStreamDocument>();
 
 			[BsonElement("Suspects"), BsonIgnoreIfNull]
@@ -909,6 +919,17 @@ namespace HordeServer.Issues
 					issueLogger.LogInformation("Issue was marked as not fixed as a systemic issue");
 				}
 			}
+			if (newIssue.RootCommitId != oldIssue.RootCommitId)
+			{
+				if (newIssue.RootCommitId == null)
+				{
+					issueLogger.LogInformation("Root commit for issue was cleared");
+				}
+				else
+				{
+					issueLogger.LogInformation("{Commit} marked as root commit for Issue", newIssue.RootCommitId);
+				}
+			}
 			if (newIssue.ResolvedById != oldIssue.ResolvedById)
 			{
 				if (newIssue.ResolvedById == null)
@@ -1328,6 +1349,17 @@ namespace HordeServer.Issues
 					updates.Add(Builders<IssueDocument>.Update.Unset(x => x.FixCommitName).Unset(x => x.FixChange).Unset(x => x.FixSystemic));
 				}
 			}
+			if (options.RootCommitId != null)
+			{
+				if (String.IsNullOrEmpty(options.RootCommitId.Name))
+				{
+					updates.Add(Builders<IssueDocument>.Update.Unset(x => x.RootCommitName));
+				}
+				else
+				{
+					updates.Add(Builders<IssueDocument>.Update.Set(x => x.RootCommitName, options.RootCommitId.Name));
+				}
+			}
 			if (options.ResolvedById != null)
 			{
 				if (options.ResolvedById.Value != UserId.Empty)
diff --git a/Engine/Source/Programs/Horde/Plugins/Build/HordeServer.Build/Issues/IssueService.cs b/Engine/Source/Programs/Horde/Plugins/Build/HordeServer.Build/Issues/IssueService.cs
--- a/Engine/Source/Programs/Horde/Plugins/Build/HordeServer.Build/Issues/IssueService.cs
+++ b/Engine/Source/Programs/Horde/Plugins/Build/HordeServer.Build/Issues/IssueService.cs
@@ -483,6 +483,7 @@ namespace HordeServer.Issues
 		/// <param name="declinedById">Name of the user that has declined this issue</param>
 		/// <param name="fixCommitId">Fix changelist for the issue</param>
 		/// <param name="fixSystemic">There was a systemic fix for this issue</param>
+		/// <param name="rootCommitId">Root changelist for the issue</param>
 		/// <param name="resolvedById">Whether the issue has been resolved</param>
 		/// <param name="addSpanIds">Add spans to this issue</param>
 		/// <param name="removeSpanIds">Remove spans from this issue</param>
@@ -493,7 +494,7 @@ namespace HordeServer.Issues
 		/// <param name="workflowThreadUrl">The workfloe thread created for the issue</param>
 		/// <param name="cancellationToken">Cancellation token for the operation</param>
 		/// <returns>True if the issue was updated</returns>
-		public async Task<bool> UpdateIssueAsync(int id, string? summary = null, string? description = null, bool? promoted = null, UserId? ownerId = null, UserId? nominatedById = null, bool? acknowledged = null, UserId? declinedById = null, CommitId? fixCommitId = null, bool? fixSystemic = null, UserId? resolvedById = null, List<ObjectId>? addSpanIds = null, List<ObjectId>? removeSpanIds = null, string? externalIssueKey = null, UserId? quarantinedById = null, UserId? forceClosedById = null, UserId? initiatedById = null, Uri? workflowThreadUrl = null, CancellationToken cancellationToken = default)
+		public async Task<bool> UpdateIssueAsync(int id, string? summary = null, string? description = null, bool? promoted = null, UserId? ownerId = null, UserId? nominatedById = null, bool? acknowledged = null, UserId? declinedById = null, CommitId? fixCommitId = null, bool? fixSystemic = null, CommitId? rootCommitId = null, UserId? resolvedById = null, List<ObjectId>? addSpanIds = null, List<ObjectId>? removeSpanIds = null, string? externalIssueKey = null, UserId? quarantinedById = null, UserId? forceClosedById = null, UserId? initiatedById = null, Uri? workflowThreadUrl = null, CancellationToken cancellationToken = default)
 		{
 			using TelemetrySpan tracerSpan = _tracer.StartActiveSpan($"{nameof(IssueService)}.{nameof(UpdateIssueAsync)}");
 			
@@ -507,7 +508,7 @@ namespace HordeServer.Issues
 					return false;
 				}
 
-				issue = await _issueCollection.TryUpdateIssueAsync(issue, initiatedById, new UpdateIssueOptions { UserSummary = summary, Description = description, Promoted = promoted, OwnerId = ownerId ?? resolvedById, NominatedById = nominatedById, DeclinedById = declinedById, Acknowledged = acknowledged, FixCommitId = fixCommitId, FixSystemic = fixSystemic, ResolvedById = resolvedById, ExcludeSpanIds = removeSpanIds, ExternalIssueKey = externalIssueKey, QuarantinedById = quarantinedById, ForceClosedById = forceClosedById, WorkflowThreadUrl = workflowThreadUrl }, cancellationToken: cancellationToken);
+				issue = await _issueCollection.TryUpdateIssueAsync(issue, initiatedById, new UpdateIssueOptions { UserSummary = summary, Description = description, Promoted = promoted, OwnerId = ownerId ?? resolvedById, NominatedById = nominatedById, DeclinedById = declinedById, Acknowledged = acknowledged, FixCommitId = fixCommitId, FixSystemic = fixSystemic, RootCommitId = rootCommitId, ResolvedById = resolvedById, ExcludeSpanIds = removeSpanIds, ExternalIssueKey = externalIssueKey, QuarantinedById = quarantinedById, ForceClosedById = forceClosedById, WorkflowThreadUrl = workflowThreadUrl }, cancellationToken: cancellationToken);
 				if (issue != null)
 				{
 					break;
diff --git a/Engine/Source/Programs/Horde/Plugins/Build/HordeServer.Build/Issues/IssueTelemetry.cs b/Engine/Source/Programs/Horde/Plugins/Build/HordeServer.Build/Issues/IssueTelemetry.cs
--- a/Engine/Source/Programs/Horde/Plugins/Build/HordeServer.Build/Issues/IssueTelemetry.cs
+++ b/Engine/Source/Programs/Horde/Plugins/Build/HordeServer.Build/Issues/IssueTelemetry.cs
@@ -14,7 +14,7 @@ namespace HordeServer.Issues
 	/// Record used for Issue telemetry. Models a <see cref="IIssue"/>.
 	/// </summary>
 	/// <remarks>Modifications to this record's structure should be versioned appropriately through <see cref="SchemaVersion"/> via <see cref="CurrentSchemaVersion"/>, and call sites should continue to publish in a backwards compatible manner.</remarks>
-	public record IssueTelemetry(string EventName, int SchemaVersion, int Id, DateTime? AcknowledgedAt, DateTime CreatedAt, UserId? OwnerId, string? ExternalIssueKey, CommitId? FixChange, bool FixedSystemic, DateTime LastSeenAt, DateTime? NominatedAt, UserId? NominatedById, DateTime? ResolvedAt, UserId? ResolvedById, IssueSeverity Severity, string Summary, DateTime? VerifiedAt)
+	public record IssueTelemetry(string EventName, int SchemaVersion, int Id, DateTime? AcknowledgedAt, DateTime CreatedAt, UserId? OwnerId, string? ExternalIssueKey, CommitId? FixChange, bool FixedSystemic, CommitId? RootChange, DateTime LastSeenAt, DateTime? NominatedAt, UserId? NominatedById, DateTime? ResolvedAt, UserId? ResolvedById, IssueSeverity Severity, string Summary, DateTime? VerifiedAt)
 	{
 		/// <summary>
 		/// Default event name for the IssueTelemetry.
@@ -40,6 +40,7 @@ namespace HordeServer.Issues
 					issue.ExternalIssueKey,
 					issue.FixCommitId,
 					issue.FixSystemic,
+					issue.RootCommitId,
 					issue.LastSeenAt,
 					issue.NominatedAt,
 					issue.NominatedById,
diff --git a/Engine/Source/Programs/Horde/Plugins/Build/HordeServer.Build/Issues/IssuesController.cs b/Engine/Source/Programs/Horde/Plugins/Build/HordeServer.Build/Issues/IssuesController.cs
--- a/Engine/Source/Programs/Horde/Plugins/Build/HordeServer.Build/Issues/IssuesController.cs
+++ b/Engine/Source/Programs/Horde/Plugins/Build/HordeServer.Build/Issues/IssuesController.cs
@@ -254,6 +254,7 @@ namespace HordeServer.Issues
 			response.AcknowledgedAt = issue.AcknowledgedAt;
 			response.FixCommitId = issue.FixCommitId;
 			response.FixSystemic = issue.FixSystemic;
+			response.RootCommitId = issue.RootCommitId;
 			response.ResolvedAt = issue.ResolvedAt;
 			if (resolvedBy != null)
 			{
@@ -502,6 +503,7 @@ namespace HordeServer.Issues
 			response.AcknowledgedAt = issue.AcknowledgedAt;
 			response.FixCommitId = issue.FixCommitId;
 			response.FixSystemic = issue.FixSystemic;
+			response.RootCommitId = issue.RootCommitId;
 			response.ResolvedAt = issue.ResolvedAt;
 			response.ResolvedBy = details.ResolvedBy?.Login;
 			response.ResolvedById = details.ResolvedBy?.Id.ToString();
@@ -837,7 +839,7 @@ namespace HordeServer.Issues
 				removeSpans = request.RemoveSpans.ConvertAll(x => ObjectId.Parse(x));
 			}
 
-			if (!await _issueService.UpdateIssueAsync(issueId, request.Summary, request.Description, request.Promoted, newOwnerId, newNominatedById, request.Acknowledged, newDeclinedById, request.FixCommitId, request.FixSystemic, newResolvedById, addSpans, removeSpans, request.ExternalIssueKey, newQuarantinedById, newForceClosedById, initiatedById: User.GetUserId()))
+			if (!await _issueService.UpdateIssueAsync(issueId, request.Summary, request.Description, request.Promoted, newOwnerId, newNominatedById, request.Acknowledged, newDeclinedById, request.FixCommitId, request.FixSystemic, request.RootCommitId, newResolvedById, addSpans, removeSpans, request.ExternalIssueKey, newQuarantinedById, newForceClosedById, initiatedById: User.GetUserId()))
 			{
 				return NotFound();
 			}
diff --git a/Engine/Source/Programs/Shared/EpicGames.Horde/Issues/IssueMessages.cs b/Engine/Source/Programs/Shared/EpicGames.Horde/Issues/IssueMessages.cs
--- a/Engine/Source/Programs/Shared/EpicGames.Horde/Issues/IssueMessages.cs
+++ b/Engine/Source/Programs/Shared/EpicGames.Horde/Issues/IssueMessages.cs
@@ -398,6 +398,16 @@ namespace EpicGames.Horde.Issues
 		}
 		bool? _fixSystemic;
 
+		/// <summary>
+		/// Commit that introduced this issue
+		/// </summary>
+		public CommitId? RootCommitId
+		{
+			get => _rootCommitId;
+			set => _rootCommitId = value;
+		}
+		CommitId? _rootCommitId;
+
 		/// <summary>
 		/// Time at which the issue was resolved
 		/// </summary>
@@ -625,6 +635,16 @@ namespace EpicGames.Horde.Issues
 		}
 		bool? _fixSystemic;
 
+		/// <summary>
+		/// Changelist that introduced this issue
+		/// </summary>
+		public CommitId? RootCommitId
+		{
+			get => _rootCommitId;
+			set => _rootCommitId = value;
+		}
+		CommitId? _rootCommitId;
+
 		/// <summary>
 		/// Time at which the issue was resolved
 		/// </summary>
@@ -732,7 +752,7 @@ namespace EpicGames.Horde.Issues
 		/// </summary>
 		public CommitId? FixCommitId
 		{
-			get => _fixCommitId ?? ((_fixChange < 0)? null : (_fixChange == 0)? CommitId.Empty : CommitId.FromPerforceChange(_fixChange));
+			get => _fixCommitId ?? ((_fixChange < 0) ? null : (_fixChange == 0) ? CommitId.Empty : CommitId.FromPerforceChange(_fixChange));
 			set => _fixCommitId = value;
 		}
 		CommitId? _fixCommitId;
@@ -747,6 +767,16 @@ namespace EpicGames.Horde.Issues
 		}
 		bool? _fixSystemic;
 
+		/// <summary>
+		/// The change at which the issue is claimed to have been introduced.
+		/// </summary>
+		public CommitId? RootCommitId
+		{
+			get => _rootCommitId == null ? null : String.IsNullOrWhiteSpace(_rootCommitId.Name) ? CommitId.Empty : _rootCommitId;
+			set => _rootCommitId = value;
+		}
+		CommitId? _rootCommitId;
+
 		/// <summary>
 		/// Whether the issue should be marked as resolved
 		/// </summary>
