Horde: Root Cause Analysis server side

diff --git a/Engine/Source/Programs/Horde/Plugins/Build/HordeServer.Build/Issues/IIssue.cs b/Engine/Source/Programs/Horde/Plugins/Build/HordeServer.Build/Issues/IIssue.cs
--- a/Engine/Source/Programs/Horde/Plugins/Build/HordeServer.Build/Issues/IIssue.cs
+++ b/Engine/Source/Programs/Horde/Plugins/Build/HordeServer.Build/Issues/IIssue.cs
@@ -113,6 +113,16 @@ namespace HordeServer.Issues
 		/// </summary>
 		public CommitId? RootCommitId { get; }
 
+		/// <summary>
+		/// Root cause category for this issue
+		/// </summary>
+		public string? RootCauseCategory { get; }
+
+		/// <summary>
+		/// Root cause summary for this issue
+		/// </summary>
+		public string? RootCauseSummary { get; }
+
 		/// <summary>
 		/// Whether the issue is marked fixed as systemic
 		/// </summary>
diff --git a/Engine/Source/Programs/Horde/Plugins/Build/HordeServer.Build/Issues/IIssueCollection.cs b/Engine/Source/Programs/Horde/Plugins/Build/HordeServer.Build/Issues/IIssueCollection.cs
--- a/Engine/Source/Programs/Horde/Plugins/Build/HordeServer.Build/Issues/IIssueCollection.cs
+++ b/Engine/Source/Programs/Horde/Plugins/Build/HordeServer.Build/Issues/IIssueCollection.cs
@@ -237,6 +237,8 @@ namespace HordeServer.Issues
 	/// <param name="FixCommitId">Fix changelist for the issue. Pass a default value to clear the fix changelist.</param>
 	/// <param name="FixSystemic">Whether the issue should be marked fixed as a systemic issue</param>
 	/// <param name="RootCommitId">Root changelist for the issue. Pass an empty value to clear the root changelist.</param>
+	/// <param name="RootCauseCategory">Root cause category</param>
+	/// <param name="RootCauseSummary">Root cause summary.</param>
 	/// <param name="ResolvedById">User that resolved the issue (may be ObjectId.Empty to clear)</param>
 	/// <param name="ExcludeSpanIds">List of span ids to exclude from this issue</param>
 	/// <param name="LastSeenAt"></param>
@@ -258,6 +260,8 @@ namespace HordeServer.Issues
 		CommitId? FixCommitId = null,
 		bool? FixSystemic = null,
 		CommitId? RootCommitId = null,
+		string? RootCauseCategory = null,
+		string? RootCauseSummary = null,
 		UserId? ResolvedById = null,
 		List<ObjectId>? ExcludeSpanIds = null,
 		DateTime? LastSeenAt = null,
diff --git a/Engine/Source/Programs/Horde/Plugins/Build/HordeServer.Build/Issues/IssueCollection.cs b/Engine/Source/Programs/Horde/Plugins/Build/HordeServer.Build/Issues/IssueCollection.cs
--- a/Engine/Source/Programs/Horde/Plugins/Build/HordeServer.Build/Issues/IssueCollection.cs
+++ b/Engine/Source/Programs/Horde/Plugins/Build/HordeServer.Build/Issues/IssueCollection.cs
@@ -125,6 +125,12 @@ namespace HordeServer.Issues
 			[BsonElement("RootCommit"), BsonIgnoreIfNull]
 			public string? RootCommitName { get; set; }
 
+			[BsonIgnoreIfNull]
+			public string? RootCauseCategory { get; set; }
+
+			[BsonIgnoreIfNull]
+			public string? RootCauseSummary { get; set; }
+
 			public List<IssueStreamDocument> Streams { get; set; } = new List<IssueStreamDocument>();
 
 			[BsonElement("Suspects"), BsonIgnoreIfNull]
@@ -1360,6 +1366,31 @@ namespace HordeServer.Issues
 					updates.Add(Builders<IssueDocument>.Update.Set(x => x.RootCommitName, options.RootCommitId.Name));
 				}
 			}
+
+			if (options.RootCauseCategory != null)
+			{
+				if (options.RootCauseCategory.Length == 0)
+				{
+					updates.Add(Builders<IssueDocument>.Update.Unset(x => x.RootCauseCategory));
+				}
+				else
+				{
+					updates.Add(Builders<IssueDocument>.Update.Set(x => x.RootCauseCategory, options.RootCauseCategory));
+				}
+			}
+
+			if (options.RootCauseSummary != null)
+			{
+				if (options.RootCauseSummary.Length == 0)
+				{
+					updates.Add(Builders<IssueDocument>.Update.Unset(x => x.RootCauseSummary));
+				}
+				else
+				{
+					updates.Add(Builders<IssueDocument>.Update.Set(x => x.RootCauseSummary, options.RootCauseSummary));
+				}
+			}
+
 			if (options.ResolvedById != null)
 			{
 				if (options.ResolvedById.Value != UserId.Empty)
diff --git a/Engine/Source/Programs/Horde/Plugins/Build/HordeServer.Build/Issues/IssueService.cs b/Engine/Source/Programs/Horde/Plugins/Build/HordeServer.Build/Issues/IssueService.cs
--- a/Engine/Source/Programs/Horde/Plugins/Build/HordeServer.Build/Issues/IssueService.cs
+++ b/Engine/Source/Programs/Horde/Plugins/Build/HordeServer.Build/Issues/IssueService.cs
@@ -484,6 +484,8 @@ namespace HordeServer.Issues
 		/// <param name="fixCommitId">Fix changelist for the issue</param>
 		/// <param name="fixSystemic">There was a systemic fix for this issue</param>
 		/// <param name="rootCommitId">Root changelist for the issue</param>
+		/// <param name="rootCauseCategory">Root cause category for the issue</param>
+		/// <param name="rootCauseSummary">Root cause summary for the issue</param>
 		/// <param name="resolvedById">Whether the issue has been resolved</param>
 		/// <param name="addSpanIds">Add spans to this issue</param>
 		/// <param name="removeSpanIds">Remove spans from this issue</param>
@@ -494,7 +496,7 @@ namespace HordeServer.Issues
 		/// <param name="workflowThreadUrl">The workfloe thread created for the issue</param>
 		/// <param name="cancellationToken">Cancellation token for the operation</param>
 		/// <returns>True if the issue was updated</returns>
-		public async Task<bool> UpdateIssueAsync(int id, string? summary = null, string? description = null, bool? promoted = null, UserId? ownerId = null, UserId? nominatedById = null, bool? acknowledged = null, UserId? declinedById = null, CommitId? fixCommitId = null, bool? fixSystemic = null, CommitId? rootCommitId = null, UserId? resolvedById = null, List<ObjectId>? addSpanIds = null, List<ObjectId>? removeSpanIds = null, string? externalIssueKey = null, UserId? quarantinedById = null, UserId? forceClosedById = null, UserId? initiatedById = null, Uri? workflowThreadUrl = null, CancellationToken cancellationToken = default)
+		public async Task<bool> UpdateIssueAsync(int id, string? summary = null, string? description = null, bool? promoted = null, UserId? ownerId = null, UserId? nominatedById = null, bool? acknowledged = null, UserId? declinedById = null, CommitId? fixCommitId = null, bool? fixSystemic = null, CommitId? rootCommitId = null, string? rootCauseCategory = null, string? rootCauseSummary = null, UserId? resolvedById = null, List<ObjectId>? addSpanIds = null, List<ObjectId>? removeSpanIds = null, string? externalIssueKey = null, UserId? quarantinedById = null, UserId? forceClosedById = null, UserId? initiatedById = null, Uri? workflowThreadUrl = null, CancellationToken cancellationToken = default)
 		{
 			using TelemetrySpan tracerSpan = _tracer.StartActiveSpan($"{nameof(IssueService)}.{nameof(UpdateIssueAsync)}");
 			
@@ -508,7 +510,7 @@ namespace HordeServer.Issues
 					return false;
 				}
 
-				issue = await _issueCollection.TryUpdateIssueAsync(issue, initiatedById, new UpdateIssueOptions { UserSummary = summary, Description = description, Promoted = promoted, OwnerId = ownerId ?? resolvedById, NominatedById = nominatedById, DeclinedById = declinedById, Acknowledged = acknowledged, FixCommitId = fixCommitId, FixSystemic = fixSystemic, RootCommitId = rootCommitId, ResolvedById = resolvedById, ExcludeSpanIds = removeSpanIds, ExternalIssueKey = externalIssueKey, QuarantinedById = quarantinedById, ForceClosedById = forceClosedById, WorkflowThreadUrl = workflowThreadUrl }, cancellationToken: cancellationToken);
+				issue = await _issueCollection.TryUpdateIssueAsync(issue, initiatedById, new UpdateIssueOptions { UserSummary = summary, Description = description, Promoted = promoted, OwnerId = ownerId ?? resolvedById, NominatedById = nominatedById, DeclinedById = declinedById, Acknowledged = acknowledged, FixCommitId = fixCommitId, FixSystemic = fixSystemic, RootCommitId = rootCommitId, RootCauseCategory = rootCauseCategory, RootCauseSummary = rootCauseSummary, ResolvedById = resolvedById, ExcludeSpanIds = removeSpanIds, ExternalIssueKey = externalIssueKey, QuarantinedById = quarantinedById, ForceClosedById = forceClosedById, WorkflowThreadUrl = workflowThreadUrl }, cancellationToken: cancellationToken);
 				if (issue != null)
 				{
 					break;
diff --git a/Engine/Source/Programs/Horde/Plugins/Build/HordeServer.Build/Issues/IssuesController.cs b/Engine/Source/Programs/Horde/Plugins/Build/HordeServer.Build/Issues/IssuesController.cs
--- a/Engine/Source/Programs/Horde/Plugins/Build/HordeServer.Build/Issues/IssuesController.cs
+++ b/Engine/Source/Programs/Horde/Plugins/Build/HordeServer.Build/Issues/IssuesController.cs
@@ -255,6 +255,8 @@ namespace HordeServer.Issues
 			response.FixCommitId = issue.FixCommitId;
 			response.FixSystemic = issue.FixSystemic;
 			response.RootCommitId = issue.RootCommitId;
+			response.RootCauseCategory = issue.RootCauseCategory;
+			response.RootCauseSummary = issue.RootCauseSummary;
 			response.ResolvedAt = issue.ResolvedAt;
 			if (resolvedBy != null)
 			{
@@ -504,6 +506,8 @@ namespace HordeServer.Issues
 			response.FixCommitId = issue.FixCommitId;
 			response.FixSystemic = issue.FixSystemic;
 			response.RootCommitId = issue.RootCommitId;
+			response.RootCauseCategory = issue.RootCauseCategory;
+			response.RootCauseSummary = issue.RootCauseSummary;
 			response.ResolvedAt = issue.ResolvedAt;
 			response.ResolvedBy = details.ResolvedBy?.Login;
 			response.ResolvedById = details.ResolvedBy?.Id.ToString();
@@ -839,7 +843,7 @@ namespace HordeServer.Issues
 				removeSpans = request.RemoveSpans.ConvertAll(x => ObjectId.Parse(x));
 			}
 
-			if (!await _issueService.UpdateIssueAsync(issueId, request.Summary, request.Description, request.Promoted, newOwnerId, newNominatedById, request.Acknowledged, newDeclinedById, request.FixCommitId, request.FixSystemic, request.RootCommitId, newResolvedById, addSpans, removeSpans, request.ExternalIssueKey, newQuarantinedById, newForceClosedById, initiatedById: User.GetUserId()))
+			if (!await _issueService.UpdateIssueAsync(issueId, request.Summary, request.Description, request.Promoted, newOwnerId, newNominatedById, request.Acknowledged, newDeclinedById, request.FixCommitId, request.FixSystemic, request.RootCommitId, request.RootCauseCategory, request.RootCauseSummary, newResolvedById, addSpans, removeSpans, request.ExternalIssueKey, newQuarantinedById, newForceClosedById, initiatedById: User.GetUserId()))
 			{
 				return NotFound();
 			}
diff --git a/Engine/Source/Programs/Shared/EpicGames.Horde/Issues/IssueMessages.cs b/Engine/Source/Programs/Shared/EpicGames.Horde/Issues/IssueMessages.cs
--- a/Engine/Source/Programs/Shared/EpicGames.Horde/Issues/IssueMessages.cs
+++ b/Engine/Source/Programs/Shared/EpicGames.Horde/Issues/IssueMessages.cs
@@ -408,6 +408,16 @@ namespace EpicGames.Horde.Issues
 		}
 		CommitId? _rootCommitId;
 
+		/// <summary>
+		/// Root Cause Category for this issue
+		/// </summary>
+		public string? RootCauseCategory { get; set; }
+
+		/// <summary>
+		/// Root Cause Summary for this issue
+		/// </summary>
+		public string? RootCauseSummary { get; set; }
+
 		/// <summary>
 		/// Time at which the issue was resolved
 		/// </summary>
@@ -625,6 +635,16 @@ namespace EpicGames.Horde.Issues
 		}
 		CommitId? _fixCommitId;
 
+		/// <summary>
+		/// Root Cause Category for this issue
+		/// </summary>
+		public string? RootCauseCategory { get; set; }
+
+		/// <summary>
+		/// Root Cause Summary for this issue
+		/// </summary>
+		public string? RootCauseSummary { get; set; }
+
 		/// <summary>
 		/// Whether the issue is marked fixed as a systemic issue
 		/// </summary>
@@ -777,6 +797,16 @@ namespace EpicGames.Horde.Issues
 		}
 		CommitId? _rootCommitId;
 
+		/// <summary>
+		/// Root Cause Category for this issue
+		/// </summary>
+		public string? RootCauseCategory { get; set; }
+
+		/// <summary>
+		/// Root Cause Summary for this issue
+		/// </summary>
+		public string? RootCauseSummary { get; set; }
+
 		/// <summary>
 		/// Whether the issue should be marked as resolved
 		/// </summary>
