Horde: Fix compiler warnings in TelemetryIntegrationTests

diff --git a/Engine/Source/Programs/Horde/Plugins/Analytics/HordeServer.Analytics.Integration.Tests/Telemetry/TelemetryIntegrationTests.cs b/Engine/Source/Programs/Horde/Plugins/Analytics/HordeServer.Analytics.Integration.Tests/Telemetry/TelemetryIntegrationTests.cs
--- a/Engine/Source/Programs/Horde/Plugins/Analytics/HordeServer.Analytics.Integration.Tests/Telemetry/TelemetryIntegrationTests.cs
+++ b/Engine/Source/Programs/Horde/Plugins/Analytics/HordeServer.Analytics.Integration.Tests/Telemetry/TelemetryIntegrationTests.cs
@@ -25,7 +25,6 @@ using HordeServer.Telemetry.Metrics;
 using HordeServer.Tests;
 using HordeServer.Tests.Issues;
 using Microsoft.Extensions.DependencyInjection;
-using Moq;
 
 namespace HordeServer.Analytics.Integration.Tests.Telemetry
 {
@@ -273,7 +272,6 @@ namespace HordeServer.Analytics.Integration.Tests.Telemetry
 			BuildConfig buildConfig = new BuildConfig();
 			buildConfig.Projects.Add(projectConfig);
 
-
 			GlobalConfig globalConfig = new GlobalConfig();
 			globalConfig.Plugins.AddBuildConfig(buildConfig);
 			globalConfig.Plugins.AddAnalyticsTestConfig(analyticsConfig);
@@ -296,7 +294,7 @@ namespace HordeServer.Analytics.Integration.Tests.Telemetry
 			return label;
 		}
 
-		private static Task<IGraph> UpdateGraph(IGraphCollection graphCollection, IGraph rootGraph, ITemplate template)
+		private static Task<IGraph> UpdateGraph(IGraphCollection graphCollection, IGraph rootGraph, ITemplate _)
 		{
 			// Nodes
 
@@ -387,7 +385,6 @@ namespace HordeServer.Analytics.Integration.Tests.Telemetry
 			Assert.IsNotNull(actualJobSummaryTelemetry);
 		}
 
-
 		[TestMethod]
 		public async Task TestJobLabelSummaryEmissionAsync()
 		{
