Horde: Improve tracing and logging for issue locks

- Environment.StackTrace is likely too big to pass as attributes to traces. Removed from log statements as well and replaced with "purpose" parameter.
- Added extra try/finally to ensure Redis lock is always disposed

diff --git a/Engine/Source/Programs/Horde/Plugins/Build/HordeServer.Build/Issues/IIssueCollection.cs b/Engine/Source/Programs/Horde/Plugins/Build/HordeServer.Build/Issues/IIssueCollection.cs
--- a/Engine/Source/Programs/Horde/Plugins/Build/HordeServer.Build/Issues/IIssueCollection.cs
+++ b/Engine/Source/Programs/Horde/Plugins/Build/HordeServer.Build/Issues/IIssueCollection.cs
@@ -26,8 +26,9 @@ namespace HordeServer.Issues
 		/// <summary>
 		/// Attempts to enters a critical section over the issue collection. Used for performing operations that require coordination of several documents (eg. attaching spans to issues).
 		/// </summary>
+		/// <param name="purpose">Descriptive name for what purpose lock is held. Solely for logging and debugging</param>
 		/// <returns>The disposable critical section</returns>
-		Task<IAsyncDisposable> EnterCriticalSectionAsync();
+		Task<IAsyncDisposable> EnterCriticalSectionAsync(string purpose);
 
 		#region Issues
 
diff --git a/Engine/Source/Programs/Horde/Plugins/Build/HordeServer.Build/Issues/IssueCollection.cs b/Engine/Source/Programs/Horde/Plugins/Build/HordeServer.Build/Issues/IssueCollection.cs
--- a/Engine/Source/Programs/Horde/Plugins/Build/HordeServer.Build/Issues/IssueCollection.cs
+++ b/Engine/Source/Programs/Horde/Plugins/Build/HordeServer.Build/Issues/IssueCollection.cs
@@ -649,26 +649,34 @@ namespace HordeServer.Issues
 		sealed class IssueLock : IAsyncDisposable
 		{
 			readonly RedisLock _redisLock;
+			readonly string _purpose;
 			readonly Stopwatch _timer = Stopwatch.StartNew();
 			readonly TelemetrySpan _telemetrySpan;
 			readonly ILogger _logger;
 
-			public IssueLock(RedisLock redisLock, Tracer tracer, ILogger logger)
+			public IssueLock(RedisLock redisLock, string purpose, Tracer tracer, ILogger logger)
 			{
 				_redisLock = redisLock;
-				_telemetrySpan = tracer.StartActiveSpan("Holding Issue Lock").SetAttribute("Stack", Environment.StackTrace);
+				_purpose = purpose;
+				_telemetrySpan = tracer.StartActiveSpan("Holding Issue Lock").SetAttribute("horde.purpose", _purpose);
 				_logger = logger;
 			}
 
 			public async ValueTask DisposeAsync()
 			{
-				_telemetrySpan.Dispose();
-				if (_timer.Elapsed.TotalSeconds >= 2.5)
+				try
 				{
-					_logger.LogWarning("Issue lock held for {TimeSpan}s. Released from {Stack}.", (int)_timer.Elapsed.TotalSeconds, Environment.StackTrace);
-					_timer.Reset();
+					_telemetrySpan.Dispose();
+					if (_timer.Elapsed.TotalSeconds >= 2.5)
+					{
+						_logger.LogWarning("Issue lock '{Purpose}' held for {TimeSpan}s", _purpose, (int)_timer.Elapsed.TotalSeconds);
+						_timer.Reset();
+					}
+				}
+				finally
+				{
+					await _redisLock.DisposeAsync();
 				}
-				await _redisLock.DisposeAsync();
 			}
 		}
 
@@ -743,25 +751,25 @@ namespace HordeServer.Issues
 		}
 
 		/// <inheritdoc/>
-		public async Task<IAsyncDisposable> EnterCriticalSectionAsync()
+		public async Task<IAsyncDisposable> EnterCriticalSectionAsync(string purpose)
 		{
 			Stopwatch timer = Stopwatch.StartNew();
 			TimeSpan nextNotifyTime = TimeSpan.FromSeconds(10.0);
 
 			RedisLock issueLock = new(_redisService.GetDatabase(), "issues/lock");
-			using (TelemetrySpan telemetrySpan = _tracer.StartActiveSpan("Wait for Issue Lock").SetAttribute("Trace", Environment.StackTrace))
+			using (TelemetrySpan telemetrySpan = _tracer.StartActiveSpan("Wait for Issue Lock").SetAttribute("horde.purpose", purpose))
 			{
 				while (!await issueLock.AcquireAsync(TimeSpan.FromMinutes(1)))
 				{
 					if (timer.Elapsed > nextNotifyTime)
 					{
-						_logger.LogWarning("Waiting on lock over issue collection for {TimeSpan} ms", timer.ElapsedMilliseconds);
+						_logger.LogWarning("Waiting on lock '{Purpose}' over issue collection for {TimeSpan} ms", purpose, timer.ElapsedMilliseconds);
 						nextNotifyTime *= 2;
 					}
 					await Task.Delay(TimeSpan.FromMilliseconds(100));
 				}
 			}
-			return new IssueLock(issueLock, _tracer, _logger);
+			return new IssueLock(issueLock, purpose, _tracer, _logger);
 		}
 
 		async Task<IssueDocument?> TryUpdateIssueAsync(IIssue issue, UpdateDefinition<IssueDocument> update, CancellationToken cancellationToken)
diff --git a/Engine/Source/Programs/Horde/Plugins/Build/HordeServer.Build/Issues/IssueService.cs b/Engine/Source/Programs/Horde/Plugins/Build/HordeServer.Build/Issues/IssueService.cs
--- a/Engine/Source/Programs/Horde/Plugins/Build/HordeServer.Build/Issues/IssueService.cs
+++ b/Engine/Source/Programs/Horde/Plugins/Build/HordeServer.Build/Issues/IssueService.cs
@@ -517,7 +517,7 @@ namespace HordeServer.Issues
 				}
 			}
 
-			await using IAsyncDisposable csLock = await _issueCollection.EnterCriticalSectionAsync();
+			await using IAsyncDisposable csLock = await _issueCollection.EnterCriticalSectionAsync(nameof(UpdateIssueAsync));
 
 			List<IIssue> updateIssues = new List<IIssue>();
 			updateIssues.Add(issue);
@@ -582,26 +582,37 @@ namespace HordeServer.Issues
 
 		async Task UpdateCompleteStepsAsync(CancellationToken cancellationToken)
 		{
-			using TelemetrySpan tracerSpan = _tracer.StartActiveSpan($"{nameof(IssueService)}.{nameof(UpdateCompleteStepsAsync)}");
-			while (await _completeSteps.Reader.WaitToReadAsync(cancellationToken))
+			try
 			{
-				CompleteStep? completeStep;
-				if (_completeSteps.Reader.TryRead(out completeStep))
+				while (await _completeSteps.Reader.WaitToReadAsync(cancellationToken))
 				{
-					IJob? job = await _jobService.GetJobAsync(completeStep.JobId, cancellationToken);
-					if (job != null)
+					if (!_completeSteps.Reader.TryRead(out CompleteStep? completeStep))
 					{
-						try
+						continue;
+					}
+					
+					try
+					{
+						IJob? job = await _jobService.GetJobAsync(completeStep.JobId, cancellationToken);
+						if (job != null)
 						{
 							await UpdateCompleteStepAsync(job, job.Graph, completeStep.BatchId, completeStep.StepId, CancellationToken.None);
 						}
-						catch (Exception ex)
-						{
-							_logger.LogError(ex, "Exception while updating issues for {JobId}:{BatchId}:{StepId} - {Message}", completeStep.JobId, completeStep.BatchId, completeStep.StepId, ex.Message);
-						}
+					}
+					catch (Exception ex)
+					{
+						_logger.LogError(ex, "Exception while updating issues for {JobId}:{BatchId}:{StepId} - {Message}", completeStep.JobId, completeStep.BatchId, completeStep.StepId, ex.Message);
 					}
 				}
 			}
+			catch (OperationCanceledException) when (cancellationToken.IsCancellationRequested)
+			{
+			}
+			catch (Exception ex)
+			{
+				_logger.LogError(ex, "Fatal error in {MethodName} - background thread stopping", nameof(UpdateCompleteStepsAsync));
+				throw;
+			}
 		}
 
 		internal async Task UpdateCompleteStepAsync(IJob job, IGraph graph, JobStepBatchId batchId, JobStepId stepId, CancellationToken cancellationToken = default)
@@ -673,7 +684,7 @@ namespace HordeServer.Issues
 			HashSet<IssueEventGroupInternal> eventGroups = await GetEventGroupsForStepAsync(job, batch, step, node, annotations, workflow, cancellationToken);
 
 			// Try to update all the events. We may need to restart this due to optimistic transactions, so keep track of any existing spans we do not need to check against.
-			await using (IAsyncDisposable issueLock = await _issueCollection.EnterCriticalSectionAsync())
+			await using (IAsyncDisposable issueLock = await _issueCollection.EnterCriticalSectionAsync(nameof(UpdateCompleteStepAsync)))
 			{
 				HashSet<ObjectId> checkedSpanIds = new HashSet<ObjectId>();
 				for (; ; )
