Added a fix where we would not delete the cleanup batch file if there is a timeout.

diff --git a/Engine/Source/Programs/Horde/Drivers/JobDriver/Execution/JobExecutor.cs b/Engine/Source/Programs/Horde/Drivers/JobDriver/Execution/JobExecutor.cs
--- a/Engine/Source/Programs/Horde/Drivers/JobDriver/Execution/JobExecutor.cs
+++ b/Engine/Source/Programs/Horde/Drivers/JobDriver/Execution/JobExecutor.cs
@@ -1668,12 +1668,22 @@ namespace JobDriver.Execution
 						cancellationSource.CancelAfter(TimeSpan.FromSeconds(30.0));
 						await ExecuteProcessAsync(fileName, arguments, null, filter, logger, cancellationSource.Token);
 					}
-					FileUtils.ForceDeleteFile(cleanupScript);
 				}
 				catch (OperationCanceledException)
 				{
 					filter.WriteLine("Cleanup script did not complete within allotted time. Aborting.");
 				}
+				finally
+				{
+					try
+					{
+						FileUtils.ForceDeleteFile(cleanupScript);
+					}
+					catch (Exception ex)
+					{
+						logger.LogWarning(KnownLogEvents.Horde_BuildHealth_Ignore, ex, "Failed to delete cleanup script {CleanupScript}", cleanupScript.FullName);
+					}
+				}
 			}
 		}
 
