Horde: Adding support for Build Health duplicate issue flagging, server side

diff --git a/Engine/Source/Programs/Horde/Plugins/Build/HordeServer.Build/Issues/IIssue.cs b/Engine/Source/Programs/Horde/Plugins/Build/HordeServer.Build/Issues/IIssue.cs
--- a/Engine/Source/Programs/Horde/Plugins/Build/HordeServer.Build/Issues/IIssue.cs
+++ b/Engine/Source/Programs/Horde/Plugins/Build/HordeServer.Build/Issues/IIssue.cs
@@ -123,6 +123,11 @@ namespace HordeServer.Issues
 		/// </summary>
 		public string? RootCauseSummary { get; }
 
+		/// <summary>
+		/// An issue id that is marked as a duplicate
+		/// </summary>
+		public int? DuplicateIssueId { get; }
+
 		/// <summary>
 		/// Whether the issue is marked fixed as systemic
 		/// </summary>
diff --git a/Engine/Source/Programs/Horde/Plugins/Build/HordeServer.Build/Issues/IIssueCollection.cs b/Engine/Source/Programs/Horde/Plugins/Build/HordeServer.Build/Issues/IIssueCollection.cs
--- a/Engine/Source/Programs/Horde/Plugins/Build/HordeServer.Build/Issues/IIssueCollection.cs
+++ b/Engine/Source/Programs/Horde/Plugins/Build/HordeServer.Build/Issues/IIssueCollection.cs
@@ -240,6 +240,7 @@ namespace HordeServer.Issues
 	/// <param name="RootCommitId">Root changelist for the issue. Pass an empty value to clear the root changelist.</param>
 	/// <param name="RootCauseCategory">Root cause category</param>
 	/// <param name="RootCauseSummary">Root cause summary.</param>
+	/// <param name="DuplicateIssueId">Duplicate issue id</param>
 	/// <param name="ResolvedById">User that resolved the issue (may be ObjectId.Empty to clear)</param>
 	/// <param name="ExcludeSpanIds">List of span ids to exclude from this issue</param>
 	/// <param name="LastSeenAt"></param>
@@ -263,6 +264,7 @@ namespace HordeServer.Issues
 		CommitId? RootCommitId = null,
 		string? RootCauseCategory = null,
 		string? RootCauseSummary = null,
+		int? DuplicateIssueId = null,
 		UserId? ResolvedById = null,
 		List<ObjectId>? ExcludeSpanIds = null,
 		DateTime? LastSeenAt = null,
diff --git a/Engine/Source/Programs/Horde/Plugins/Build/HordeServer.Build/Issues/IssueCollection.cs b/Engine/Source/Programs/Horde/Plugins/Build/HordeServer.Build/Issues/IssueCollection.cs
--- a/Engine/Source/Programs/Horde/Plugins/Build/HordeServer.Build/Issues/IssueCollection.cs
+++ b/Engine/Source/Programs/Horde/Plugins/Build/HordeServer.Build/Issues/IssueCollection.cs
@@ -131,6 +131,9 @@ namespace HordeServer.Issues
 			[BsonIgnoreIfNull]
 			public string? RootCauseSummary { get; set; }
 
+			[BsonIgnoreIfNull]
+			public int? DuplicateIssueId { get; set; }
+
 			public List<IssueStreamDocument> Streams { get; set; } = new List<IssueStreamDocument>();
 
 			[BsonElement("Suspects"), BsonIgnoreIfNull]
@@ -665,17 +668,17 @@ namespace HordeServer.Issues
 			public async ValueTask DisposeAsync()
 			{
 				try
+			{
+				_telemetrySpan.Dispose();
+				if (_timer.Elapsed.TotalSeconds >= 2.5)
 				{
-					_telemetrySpan.Dispose();
-					if (_timer.Elapsed.TotalSeconds >= 2.5)
-					{
 						_logger.LogWarning("Issue lock '{Purpose}' held for {TimeSpan}s", _purpose, (int)_timer.Elapsed.TotalSeconds);
-						_timer.Reset();
-					}
+					_timer.Reset();
+				}
 				}
 				finally
 				{
-					await _redisLock.DisposeAsync();
+				await _redisLock.DisposeAsync();
 				}
 			}
 		}
@@ -944,6 +947,19 @@ namespace HordeServer.Issues
 					issueLogger.LogInformation("{Commit} marked as root commit for Issue", newIssue.RootCommitId);
 				}
 			}
+
+			if (newIssue.DuplicateIssueId != oldIssue.DuplicateIssueId)
+			{
+				if (newIssue.DuplicateIssueId == null)
+				{
+					issueLogger.LogInformation("Duplicate Issue Id cleared");
+				}
+				else
+				{
+					issueLogger.LogInformation("Duplicate Issue Id set to {IssueId}", newIssue.DuplicateIssueId);
+				}
+			}
+
 			if (newIssue.ResolvedById != oldIssue.ResolvedById)
 			{
 				if (newIssue.ResolvedById == null)
@@ -1399,6 +1415,18 @@ namespace HordeServer.Issues
 				}
 			}
 
+			if (options.DuplicateIssueId != null)
+			{
+				if (options.DuplicateIssueId < 0)
+				{
+					updates.Add(Builders<IssueDocument>.Update.Unset(x => x.DuplicateIssueId));
+				}
+				else
+				{
+					updates.Add(Builders<IssueDocument>.Update.Set(x => x.DuplicateIssueId, options.DuplicateIssueId));
+				}
+			}
+
 			if (options.ResolvedById != null)
 			{
 				if (options.ResolvedById.Value != UserId.Empty)
diff --git a/Engine/Source/Programs/Horde/Plugins/Build/HordeServer.Build/Issues/IssueService.cs b/Engine/Source/Programs/Horde/Plugins/Build/HordeServer.Build/Issues/IssueService.cs
--- a/Engine/Source/Programs/Horde/Plugins/Build/HordeServer.Build/Issues/IssueService.cs
+++ b/Engine/Source/Programs/Horde/Plugins/Build/HordeServer.Build/Issues/IssueService.cs
@@ -486,6 +486,7 @@ namespace HordeServer.Issues
 		/// <param name="rootCommitId">Root changelist for the issue</param>
 		/// <param name="rootCauseCategory">Root cause category for the issue</param>
 		/// <param name="rootCauseSummary">Root cause summary for the issue</param>
+		/// <param name="duplicateIssueId">Duplicate issue id</param>
 		/// <param name="resolvedById">Whether the issue has been resolved</param>
 		/// <param name="addSpanIds">Add spans to this issue</param>
 		/// <param name="removeSpanIds">Remove spans from this issue</param>
@@ -496,7 +497,7 @@ namespace HordeServer.Issues
 		/// <param name="workflowThreadUrl">The workfloe thread created for the issue</param>
 		/// <param name="cancellationToken">Cancellation token for the operation</param>
 		/// <returns>True if the issue was updated</returns>
-		public async Task<bool> UpdateIssueAsync(int id, string? summary = null, string? description = null, bool? promoted = null, UserId? ownerId = null, UserId? nominatedById = null, bool? acknowledged = null, UserId? declinedById = null, CommitId? fixCommitId = null, bool? fixSystemic = null, CommitId? rootCommitId = null, string? rootCauseCategory = null, string? rootCauseSummary = null, UserId? resolvedById = null, List<ObjectId>? addSpanIds = null, List<ObjectId>? removeSpanIds = null, string? externalIssueKey = null, UserId? quarantinedById = null, UserId? forceClosedById = null, UserId? initiatedById = null, Uri? workflowThreadUrl = null, CancellationToken cancellationToken = default)
+		public async Task<bool> UpdateIssueAsync(int id, string? summary = null, string? description = null, bool? promoted = null, UserId? ownerId = null, UserId? nominatedById = null, bool? acknowledged = null, UserId? declinedById = null, CommitId? fixCommitId = null, bool? fixSystemic = null, CommitId? rootCommitId = null, string? rootCauseCategory = null, string? rootCauseSummary = null, int? duplicateIssueId = null, UserId? resolvedById = null, List<ObjectId>? addSpanIds = null, List<ObjectId>? removeSpanIds = null, string? externalIssueKey = null, UserId? quarantinedById = null, UserId? forceClosedById = null, UserId? initiatedById = null, Uri? workflowThreadUrl = null, CancellationToken cancellationToken = default)
 		{
 			using TelemetrySpan tracerSpan = _tracer.StartActiveSpan($"{nameof(IssueService)}.{nameof(UpdateIssueAsync)}");
 			
@@ -510,7 +511,7 @@ namespace HordeServer.Issues
 					return false;
 				}
 
-				issue = await _issueCollection.TryUpdateIssueAsync(issue, initiatedById, new UpdateIssueOptions { UserSummary = summary, Description = description, Promoted = promoted, OwnerId = ownerId ?? resolvedById, NominatedById = nominatedById, DeclinedById = declinedById, Acknowledged = acknowledged, FixCommitId = fixCommitId, FixSystemic = fixSystemic, RootCommitId = rootCommitId, RootCauseCategory = rootCauseCategory, RootCauseSummary = rootCauseSummary, ResolvedById = resolvedById, ExcludeSpanIds = removeSpanIds, ExternalIssueKey = externalIssueKey, QuarantinedById = quarantinedById, ForceClosedById = forceClosedById, WorkflowThreadUrl = workflowThreadUrl }, cancellationToken: cancellationToken);
+				issue = await _issueCollection.TryUpdateIssueAsync(issue, initiatedById, new UpdateIssueOptions { UserSummary = summary, Description = description, Promoted = promoted, OwnerId = ownerId ?? resolvedById, NominatedById = nominatedById, DeclinedById = declinedById, Acknowledged = acknowledged, FixCommitId = fixCommitId, FixSystemic = fixSystemic, RootCommitId = rootCommitId, RootCauseCategory = rootCauseCategory, RootCauseSummary = rootCauseSummary, DuplicateIssueId = duplicateIssueId, ResolvedById = resolvedById, ExcludeSpanIds = removeSpanIds, ExternalIssueKey = externalIssueKey, QuarantinedById = quarantinedById, ForceClosedById = forceClosedById, WorkflowThreadUrl = workflowThreadUrl }, cancellationToken: cancellationToken);
 				if (issue != null)
 				{
 					break;
diff --git a/Engine/Source/Programs/Horde/Plugins/Build/HordeServer.Build/Issues/IssuesController.cs b/Engine/Source/Programs/Horde/Plugins/Build/HordeServer.Build/Issues/IssuesController.cs
--- a/Engine/Source/Programs/Horde/Plugins/Build/HordeServer.Build/Issues/IssuesController.cs
+++ b/Engine/Source/Programs/Horde/Plugins/Build/HordeServer.Build/Issues/IssuesController.cs
@@ -257,6 +257,7 @@ namespace HordeServer.Issues
 			response.RootCommitId = issue.RootCommitId;
 			response.RootCauseCategory = issue.RootCauseCategory;
 			response.RootCauseSummary = issue.RootCauseSummary;
+			response.DuplicateIssueId = issue.DuplicateIssueId;
 			response.ResolvedAt = issue.ResolvedAt;
 			if (resolvedBy != null)
 			{
@@ -508,6 +509,7 @@ namespace HordeServer.Issues
 			response.RootCommitId = issue.RootCommitId;
 			response.RootCauseCategory = issue.RootCauseCategory;
 			response.RootCauseSummary = issue.RootCauseSummary;
+			response.DuplicateIssueId = issue.DuplicateIssueId;
 			response.ResolvedAt = issue.ResolvedAt;
 			response.ResolvedBy = details.ResolvedBy?.Login;
 			response.ResolvedById = details.ResolvedBy?.Id.ToString();
@@ -843,7 +845,7 @@ namespace HordeServer.Issues
 				removeSpans = request.RemoveSpans.ConvertAll(x => ObjectId.Parse(x));
 			}
 
-			if (!await _issueService.UpdateIssueAsync(issueId, request.Summary, request.Description, request.Promoted, newOwnerId, newNominatedById, request.Acknowledged, newDeclinedById, request.FixCommitId, request.FixSystemic, request.RootCommitId, request.RootCauseCategory, request.RootCauseSummary, newResolvedById, addSpans, removeSpans, request.ExternalIssueKey, newQuarantinedById, newForceClosedById, initiatedById: User.GetUserId()))
+			if (!await _issueService.UpdateIssueAsync(issueId, request.Summary, request.Description, request.Promoted, newOwnerId, newNominatedById, request.Acknowledged, newDeclinedById, request.FixCommitId, request.FixSystemic, request.RootCommitId, request.RootCauseCategory, request.RootCauseSummary, request.DuplicateIssueId, newResolvedById, addSpans, removeSpans, request.ExternalIssueKey, newQuarantinedById, newForceClosedById, initiatedById: User.GetUserId()))
 			{
 				return NotFound();
 			}
diff --git a/Engine/Source/Programs/Shared/EpicGames.Horde/Issues/IssueMessages.cs b/Engine/Source/Programs/Shared/EpicGames.Horde/Issues/IssueMessages.cs
--- a/Engine/Source/Programs/Shared/EpicGames.Horde/Issues/IssueMessages.cs
+++ b/Engine/Source/Programs/Shared/EpicGames.Horde/Issues/IssueMessages.cs
@@ -418,6 +418,11 @@ namespace EpicGames.Horde.Issues
 		/// </summary>
 		public string? RootCauseSummary { get; set; }
 
+		/// <summary>
+		/// Duplicate issue id
+		/// </summary>
+		public int? DuplicateIssueId { get; set; }
+
 		/// <summary>
 		/// Time at which the issue was resolved
 		/// </summary>
@@ -645,6 +650,11 @@ namespace EpicGames.Horde.Issues
 		/// </summary>
 		public string? RootCauseSummary { get; set; }
 
+		/// <summary>
+		/// Duplcate issue id
+		/// </summary>
+		public int? DuplicateIssueId { get; set; }
+
 		/// <summary>
 		/// Whether the issue is marked fixed as a systemic issue
 		/// </summary>
@@ -807,6 +817,11 @@ namespace EpicGames.Horde.Issues
 		/// </summary>
 		public string? RootCauseSummary { get; set; }
 
+		/// <summary>
+		/// Duplicate Issue Id
+		/// </summary>
+		public int? DuplicateIssueId { get; set; }
+
 		/// <summary>
 		/// Whether the issue should be marked as resolved
 		/// </summary>
