Horde: Avoid duplicate id setting recursion

diff --git a/Engine/Source/Programs/Horde/Plugins/Build/HordeServer.Build/Issues/IssueCollection.cs b/Engine/Source/Programs/Horde/Plugins/Build/HordeServer.Build/Issues/IssueCollection.cs
--- a/Engine/Source/Programs/Horde/Plugins/Build/HordeServer.Build/Issues/IssueCollection.cs
+++ b/Engine/Source/Programs/Horde/Plugins/Build/HordeServer.Build/Issues/IssueCollection.cs
@@ -1421,7 +1421,7 @@ namespace HordeServer.Issues
 				{
 					updates.Add(Builders<IssueDocument>.Update.Unset(x => x.DuplicateIssueId));
 				}
-				else
+				else if (options.DuplicateIssueId != issue.Id)
 				{
 					updates.Add(Builders<IssueDocument>.Update.Set(x => x.DuplicateIssueId, options.DuplicateIssueId));
 				}
