Horde: Root Cause Analysis, dashboard side

diff --git a/Engine/Source/Programs/Horde/HordeDashboard/src/backend/Api.ts b/Engine/Source/Programs/Horde/HordeDashboard/src/backend/Api.ts
--- a/Engine/Source/Programs/Horde/HordeDashboard/src/backend/Api.ts
+++ b/Engine/Source/Programs/Horde/HordeDashboard/src/backend/Api.ts
@@ -2,6 +2,9 @@
 
 //  Horde backend API types
 
+
+export type CommitId = string;
+
 // The session type
 export enum SessionType {
 
@@ -3213,6 +3216,18 @@ export type GetIssueResponse = {
 	/**Changelist that fixed this issue */
 	fixChange?: number;
 
+	/**Root cause commit id */
+	rootCommitId?: CommitId;
+
+	/**Root cause category*/
+	rootCauseCategory?: string;
+
+	/**Root cause summary*/
+	rootCauseSummary?: string;
+
+	/**Duplicate issue id*/
+	duplicateIssueId?: number;
+
 	/**Time at which the issue was resolved */
 	resolvedAt?: Date | string;
 
@@ -3273,6 +3288,18 @@ export type UpdateIssueRequest = {
 	/**The change at which the issue is claimed fixed */
 	fixChange?: number | null;
 
+	/**Root cause commit id */
+	rootCommitId?: CommitId | null;
+
+	/**Root cause category*/
+	rootCauseCategory?: string;
+
+	/**Root cause summary*/
+	rootCauseSummary?: string;
+
+	/**Duplicate issue id*/
+	duplicateIssueId?: number;
+
 	/**Whether the issue should be marked as resolved */
 	resolved?: boolean;
 
@@ -4529,6 +4556,18 @@ export type FindIssueResponse = {
 	/// Changelist that fixed this issue
 	fixChange?: number;
 
+	/**Root cause commit id */
+	rootCommitId?: CommitId;
+
+	/**Root cause category*/
+	rootCauseCategory?: string;
+
+	/**Root cause summary*/
+	rootCauseSummary?: string;
+
+	/**Duplicate issue id*/
+	duplicateIssueId?: number;
+
 	/// Time at which the issue was resolved
 	resolvedAt?: Date | string;
 
diff --git a/Engine/Source/Programs/Horde/HordeDashboard/src/base/components/ThemeTester/ThemeTester.tsx b/Engine/Source/Programs/Horde/HordeDashboard/src/base/components/ThemeTester/ThemeTester.tsx
--- a/Engine/Source/Programs/Horde/HordeDashboard/src/base/components/ThemeTester/ThemeTester.tsx
+++ b/Engine/Source/Programs/Horde/HordeDashboard/src/base/components/ThemeTester/ThemeTester.tsx
@@ -1,6 +1,6 @@
 // Copyright Epic Games, Inc. All Rights Reserved.
 
-import { Checkbox, DefaultButton, IContextualMenuProps, ITextStyles, Pivot, PivotItem, PrimaryButton, SearchBox, SpinButton, Stack, Text, TextField, Toggle } from "@fluentui/react";
+import { Checkbox, DefaultButton, IContextualMenuProps, ITextStyles, Pivot, PivotItem, PrimaryButton, SearchBox, Separator, SpinButton, Stack, Text, TextField, Toggle } from "@fluentui/react";
 import { useState } from "react";
 import { Breadcrumbs } from "../../../components/Breadcrumbs";
 import { TopNav } from "../../../components/TopNav";
@@ -195,9 +195,19 @@ const BasicInputsPanel: React.FC = () => {
                   <Text styles={mediumTextStyle} >Spinner</Text>
                </Stack>
                <Stack tokens={{ childrenGap: 10 }}>
-                  <SpinnerSticker/>
+                  <SpinnerSticker />
                </Stack>
             </Stack>
+
+            <Stack>
+               <Stack style={{ paddingTop: 24, paddingBottom: 12 }} >
+                  <Text styles={mediumTextStyle} >Separator</Text>
+               </Stack>
+               <Stack>
+                  <Separator />
+               </Stack>
+            </Stack>
+
          </Stack>
       </Stack>
    </Stack>
diff --git a/Engine/Source/Programs/Horde/HordeDashboard/src/components/IssueViewV2.tsx b/Engine/Source/Programs/Horde/HordeDashboard/src/components/IssueViewV2.tsx
--- a/Engine/Source/Programs/Horde/HordeDashboard/src/components/IssueViewV2.tsx
+++ b/Engine/Source/Programs/Horde/HordeDashboard/src/components/IssueViewV2.tsx
@@ -1,13 +1,13 @@
 // Copyright Epic Games, Inc. All Rights Reserved.
 
-import { Checkbox, CommandBar, CommandBarButton, DefaultButton, Dialog, DialogFooter, DialogType, getFocusStyle, ICommandBarItemProps, Icon, IconButton, IContextualMenuItem, IContextualMenuItemProps, IContextualMenuProps, IGroup, Label, List, mergeStyleSets, MessageBar, MessageBarType, Modal, Persona, PersonaSize, Pivot, PivotItem, PrimaryButton, Spinner, SpinnerSize, Stack, TagPicker, Text, TextField } from "@fluentui/react";
+import { Checkbox, CommandBar, CommandBarButton, DefaultButton, Dialog, DialogFooter, DialogType, Dropdown, DropdownMenuItemType, getFocusStyle, ICommandBarItemProps, Icon, IconButton, IContextualMenuItem, IContextualMenuItemProps, IContextualMenuProps, IDropdownOption, IGroup, Label, List, mergeStyleSets, MessageBar, MessageBarType, Modal, Persona, PersonaSize, Pivot, PivotItem, PrimaryButton, Separator, Spinner, SpinnerSize, Stack, TagPicker, Text, TextField, TooltipHost } from "@fluentui/react";
 import { action, makeObservable, observable } from 'mobx';
 import { observer } from "mobx-react-lite";
 import moment from "moment";
 import React, { useEffect, useRef, useState } from "react";
 import { Link, useLocation, useNavigate } from "react-router-dom";
 import backend from "../backend";
-import { CreateExternalIssueResponse, CreateJobRequest, EventSeverity, GetChangeSummaryResponse, GetExternalIssueProjectResponse, GetExternalIssueResponse, GetIssueResponse, GetIssueSpanResponse, GetIssueStepResponse, GetIssueStreamResponse, GetLogEventResponse, GetTemplateRefResponse, GetUserResponse, IssueData, IssueSeverity, UpdateIssueRequest } from "../backend/Api";
+import { CommitId, CreateExternalIssueResponse, CreateJobRequest, EventSeverity, GetChangeSummaryResponse, GetExternalIssueProjectResponse, GetExternalIssueResponse, GetIssueResponse, GetIssueSpanResponse, GetIssueStepResponse, GetIssueStreamResponse, GetLogEventResponse, GetTemplateRefResponse, GetUserResponse, IssueData, IssueSeverity, UpdateIssueRequest } from "../backend/Api";
 import dashboard, { StatusColor } from "../backend/Dashboard";
 import { projectStore } from "../backend/ProjectStore";
 import templateCache from '../backend/TemplateCache';
@@ -20,6 +20,7 @@ import { errorDialogStore } from "./error/ErrorStore";
 import { useQuery } from "horde/base/utilities/hooks";
 import { renderLine } from "./logs/LogRender";
 import { UserSelect } from "./UserSelect";
+import { CommitIdField } from "./commits/commitField";
 
 const smallThreshhold = 1100;
 
@@ -1252,6 +1253,14 @@ const IssueSummaryPanel: React.FC = () => {
 
    const items: SummaryItem[] = [];
 
+   if (issue.duplicateIssueId !== undefined) {
+      items.push({
+         title: `Duplicate of Issue`,
+         text: issue.duplicateIssueId === 0 ? "Unspecified" : issue.duplicateIssueId.toString()
+      })
+   }
+
+
    if (issue.workflowThreadUrl) {
       items.push({
          title: "Workflow",
@@ -1260,7 +1269,6 @@ const IssueSummaryPanel: React.FC = () => {
       })
    }
 
-
    if (!!issue.quarantineTimeUtc && !!issue.quarantinedByUserInfo) {
       items.push({
          title: "Quarantined By",
@@ -1342,6 +1350,27 @@ const IssueSummaryPanel: React.FC = () => {
       }
    }
 
+
+   if (issue.rootCommitId) {
+      items.push({
+         title: "Root Change",
+         text: `${issue.rootCommitId}`,
+         link: dashboard.swarmUrl ? `${dashboard.swarmUrl}/changes/${issue.rootCommitId}` : undefined
+      })
+   }
+   if (issue.rootCauseCategory) {
+      items.push({
+         title: "Root Category",
+         text: issue.rootCauseCategory
+      })
+   }
+   if (issue.rootCauseSummary) {
+      items.push({
+         title: "Root Summary",
+         text: issue.rootCauseSummary
+      })
+   }
+
    if (!items.length && !issue.description) {
       return null;
    }
@@ -1629,9 +1658,11 @@ const IssuePanelNew: React.FC = () => {
 
 const IssueCommandBar: React.FC = () => {
 
+   const query = useQuery();
+
    const [assignToOtherShown, setAssignToOtherShown] = useState<{ shown?: boolean, defaultUser?: GetUserResponse }>({});
 
-   const [markFixedShown, setMarkFixedShown] = useState(false);
+   const [markFixedShown, setMarkFixedShown] = useState(!!query.has("markfixed"));
    const [declineShown, setDeclineShown] = useState(false);
    const [ackShown, setAckShown] = useState(false);
    const [externalIssueLinkShown, setExternalIssueLinkShown] = useState(false);
@@ -2269,13 +2300,37 @@ const AckIssueModal: React.FC<{ onClose: () => void }> = ({ onClose }) => {
 
 };
 
+// @todo: Make these server configurable 
+const rcaCategoryOptions: IDropdownOption[] = [
+   { key: 'rca_asset_content', text: 'Asset/Content' },
+   { key: 'rca_code', text: 'Code' },
+   { key: 'rca_ui', text: 'UI/Localization' },
+   { key: 'rca_deployment', text: 'Build/Deployment/Version' },
+   { key: 'rca_data', text: 'Configuration/Data' },
+   { key: 'rca_gameplay', text: 'Gameplay/Engine ' },
+   { key: 'rca_backend', text: 'Backend/Server/Network' },
+   { key: 'rca_monitoring', text: 'Monitoring/No Code Change' },
+   { key: 'rca_use', text: 'Unintended Side Effect' },
+   { key: 'rca_buildagent', text: 'Build Agent' },
+   { key: 'rca_other', text: 'Other' }
+];
+
 const MarkFixedModal: React.FC<{ onClose: () => void }> = ({ onClose }) => {
 
-   const [state, setState] = useState<{ error?: string, fixCL?: string, submitting?: boolean }>({});
+   const issue = details.issue!;
+
+   const [state, setState] = useState<{ error?: string, fixCL?: number, rootCommitId?: CommitId, rootCauseSummary?: string, rootCauseCategory?: string, submitting?: boolean, dupeCheck?: boolean, dupeId?: string }>({
+      fixCL: issue.fixChange,
+      rootCommitId: issue.rootCommitId,
+      dupeId: issue.duplicateIssueId ? (issue.duplicateIssueId.toString() !== "0" ? issue.duplicateIssueId.toString() : undefined) : undefined,
+      dupeCheck: issue.duplicateIssueId === 0 || !!issue.duplicateIssueId,
+      rootCauseCategory: issue.rootCauseCategory,
+      rootCauseSummary: issue.rootCauseSummary
+   });
 
    const { hordeClasses } = getHordeStyling();
 
-   const issue = details.issue!;
+
    const issueId = issue.id;
 
    const headerText = `Mark Issue ${issueId} Fixed`;
@@ -2284,38 +2339,59 @@ const MarkFixedModal: React.FC<{ onClose: () => void }> = ({ onClose }) => {
       onClose();
    };
 
-   const onMarkFixed = async (clearFixed?: boolean) => {
+   const onUpdate = async (clearFixed?: boolean) => {
 
       state.submitting = true;
       setState({ ...state });
 
-      try {
+      let fixCL: number | null | undefined;
+      if (clearFixed) {
+         fixCL = -1;
+      } else if (state.fixCL) {
+         fixCL = state.fixCL;
+      }
+
+      // Root Cause Analysis      
+      let rootCommitId = state.rootCommitId?.length ? state.rootCommitId : null;
+      const rootCauseCategory = state.rootCauseCategory?.length ? state.rootCauseCategory : "";
+      const rootCauseSummary = state?.rootCauseSummary?.trim().length ? state.rootCauseSummary.trim() : "";
+
+      let dupeId = state.dupeCheck ? 0 : -1;
+      if (dupeId === 0 && state.dupeId) {
+         dupeId = parseInt(state.dupeId)
+      }
 
-         let fixCL = state.fixCL?.trim();
-         if (fixCL === "-") {
-            fixCL = "-1";
+      let summary: string | undefined;
+
+      if (dupeId >= 0) {
+         if (!issue.summary.startsWith("(Dupe) ")) {
+            summary = "(Dupe) " + issue.summary;
          }
-         if (!clearFixed && !fixCL) {
-            state.submitting = false;
-            state.error = "Please provide a Fix Change";
-            setState({ ...state });
+      } else if (issue.summary.startsWith("(Dupe) ")) {
+         summary = issue.summary.replace("(Dupe) ", "");
+      }
 
-         } else {
+      const request: UpdateIssueRequest = {
+         fixChange: fixCL,
+         summary: summary,
+         rootCommitId: rootCommitId,
+         rootCauseCategory: rootCauseCategory,
+         rootCauseSummary: rootCauseSummary,
+         ownerId: !issue.ownerInfo ? dashboard.userId : issue.ownerInfo.id,
+         acknowledged: !issue.ownerInfo || issue.ownerInfo?.id === dashboard.userId ? true : undefined,
+         resolved: clearFixed ? false : true,
+         duplicateIssueId: dupeId
+      }
 
-            await backend.updateIssue(issueId, {
-               fixChange: clearFixed ? null : parseInt(fixCL!),
-               ownerId: !issue.ownerInfo ? dashboard.userId : issue.ownerInfo.id,
-               acknowledged: !issue.ownerInfo || issue.ownerInfo?.id === dashboard.userId ? true : undefined,
-               resolved: clearFixed ? false : true
-            });
+      try {
+         await backend.updateIssue(issueId, request);
 
-            state.submitting = false;
-            state.error = "";
-            setState({ ...state });
-            details.refresh(issueId);
-            close();
+         state.submitting = false;
+         state.error = "";
+         setState({ ...state });
+         details.refresh(issueId);
+         close();
 
-         }
       } catch (reason) {
 
          state.submitting = false;
@@ -2325,7 +2401,11 @@ const MarkFixedModal: React.FC<{ onClose: () => void }> = ({ onClose }) => {
 
    };
 
-   const height = state.error ? 250 : 210;
+   let height = state.error ? 250 : 210;
+
+   height += 384;
+
+   const selectedCatKey = rcaCategoryOptions.find(c => c.text === state.rootCauseCategory)?.key;
 
    return <Modal className={hordeClasses.modal} isOpen={true} styles={{ main: { padding: 8, width: 540, height: height, minHeight: height } }} onDismiss={() => { close() }}>
       <Stack horizontal styles={{ root: { padding: 8 } }}>
@@ -2345,51 +2425,57 @@ const MarkFixedModal: React.FC<{ onClose: () => void }> = ({ onClose }) => {
          messageBarType={MessageBarType.error}
          isMultiline={false}> {state.error} </MessageBar>}
 
+      <Stack style={{ paddingLeft: 8, paddingRight: 8 }}>
+         <Stack styles={{ root: { paddingLeft: 8, width: 500, paddingTop: 12 } }}>
+            <CommitIdField label="Fix Change" placeholder="Change which fixed this issue" defaultCommitId={state.fixCL?.toString()} onChange={(commitId => { setState({ ...state, fixCL: commitId ? parseInt(commitId) : undefined }) })} />
+         </Stack>
 
-      <Stack styles={{ root: { paddingLeft: 8, width: 500, paddingTop: 12 } }}>
-         <TextField label="Fix Change"
-            value={state.fixCL ?? ""}
-
-            onChange={(ev, newValue) => {
-               ev.preventDefault();
-
-               let v = newValue;
-               if (v) {
-                  const numbers = /-?\d*/;
-                  const match = newValue?.match(numbers);
-
-                  if (match) {
-                     v = match.join("");
-                  } else {
-                     v = undefined;
-                  }
-               }
-
-               if (v && v !== '-') {
-                  const n = parseInt(v);
-                  if (n > 1000000000) {
-                     v = "1000000000";
-                  }
-
-                  if (n < -1) {
-                     v = "-1";
-                  }
-               }
+         <Stack styles={{ root: { paddingLeft: 8, width: 500, paddingTop: 12 } }}>
+            <Stack style={{ paddingLeft: 8, width: 500, paddingBottom: 12 }}>
+               <Separator />
+            </Stack>
+            <Stack tokens={{ childrenGap: 18 }}>
+               <Stack>
+                  <CommitIdField label="Root Cause Change" defaultCommitId={state.rootCommitId} placeholder="Change which caused this issue, please only enter if certain on root cause" onChange={(commitId => { setState({ ...state, rootCommitId: commitId }) })} />
+               </Stack>
+               <Stack>
+                  <Dropdown
+                     selectedKey={selectedCatKey}
+                     placeholder="Select an option"
+                     label="Root Cause Category"
+                     options={rcaCategoryOptions}
+                     onChange={(event, option) => setState({ ...state, rootCauseCategory: option?.text as string | undefined })}
+                  />
+               </Stack>
+               <Stack>
+                  <TextField spellCheck={false} autoComplete="off" label="Root Cause Summary" placeholder="Please enter a brief summary describing the root cause of this issue." value={state.rootCauseSummary} style={{ minHeight: 120, maxHeight: 120 }} multiline resizable={false} onChange={(ev, newValue) => {
+                     ev.preventDefault();
+                     setState({ ...state, rootCauseSummary: newValue ? newValue : "" })
+                  }} />
+               </Stack>
 
-               setState({ ...state, fixCL: v })
-            }} />
-      </Stack>
+               <Stack horizontal verticalFill verticalAlign="center" tokens={{ childrenGap: 24 }}>
+                  <Checkbox label="Duplicate Issue" checked={state.dupeCheck} onChange={(_, checked) => setState({ ...state, dupeCheck: checked })} />
+                  <TextField disabled={!state.dupeCheck} spellCheck={false} autoComplete="off" placeholder="Optional Issue ID" value={state.dupeId} onChange={(ev, newValue) => {
+                     ev.preventDefault();
+                     setState({ ...state, dupeId: newValue ? newValue?.trim() : "" })
+                  }} />
+               </Stack>
 
-      <Stack tokens={{ childrenGap: 16 }} styles={{ root: { paddingTop: 24, paddingLeft: 8, paddingBottom: 8 } }}>
-         <Stack horizontal>
-            <Stack>
-               {!!issue.resolvedAt && <DefaultButton text="Clear Fixed" disabled={state.submitting} onClick={() => { onMarkFixed(true); }} />}
             </Stack>
-            <Stack grow />
-            <Stack>
-               <Stack horizontal tokens={{ childrenGap: 12 }} style={{ paddingRight: 24 }}>
-                  <PrimaryButton text="Mark Fixed" disabled={state.submitting ?? false} onClick={() => { onMarkFixed(); }} />
-                  <DefaultButton text="Cancel" disabled={state.submitting} onClick={() => { close(); }} />
+         </Stack>
+
+         <Stack tokens={{ childrenGap: 16 }} styles={{ root: { paddingTop: 24, paddingLeft: 8, paddingBottom: 8 } }}>
+            <Stack horizontal>
+               <Stack>
+                  {!!issue.resolvedAt && <DefaultButton text="Clear Fixed" disabled={state.submitting} onClick={() => { onUpdate(true); }} />}
+               </Stack>
+               <Stack grow />
+               <Stack>
+                  <Stack horizontal tokens={{ childrenGap: 12 }} style={{ paddingRight: 24 }}>
+                     <PrimaryButton text="Update" disabled={state.submitting ?? false} onClick={() => { onUpdate(); }} />
+                     <DefaultButton text="Cancel" disabled={state.submitting} onClick={() => { close(); }} />
+                  </Stack>
                </Stack>
             </Stack>
          </Stack>
diff --git a/Engine/Source/Programs/Horde/HordeDashboard/src/components/commits/commitField.tsx b/Engine/Source/Programs/Horde/HordeDashboard/src/components/commits/commitField.tsx
new file mode 100644
--- /dev/null
+++ b/Engine/Source/Programs/Horde/HordeDashboard/src/components/commits/commitField.tsx
@@ -0,0 +1,65 @@
+import { Stack, TextField } from "@fluentui/react";
+import React from "react";
+import { CommitId } from "horde/backend/Api";
+import { useImperativeHandle, useState } from "react";
+
+export interface ICommitIdField {
+    commitId?: CommitId;
+}
+
+export interface ICommitIdFieldProps {
+    defaultCommitId?: CommitId;
+    label?: string;
+    placeholder?: string;
+    onChange?: (commitId: CommitId | undefined) => void;
+}
+
+export const CommitIdField = React.forwardRef<ICommitIdField, ICommitIdFieldProps>((props, ref) => {
+
+    const [commitId, setCommitId] = useState(props?.defaultCommitId);
+
+    useImperativeHandle(ref, () => ({
+        commitId: commitId
+    }));
+
+    return <Stack>
+        <TextField label={props?.label}
+            placeholder={props?.placeholder}
+            value={commitId ?? ""}
+
+            onChange={(ev, newValue) => {
+                ev.preventDefault();
+
+                let v = newValue;
+                if (v) {
+                    const numbers = /-?\d*/;
+                    const match = newValue?.match(numbers);
+
+                    if (match) {
+                        v = match.join("");
+                    } else {
+                        v = undefined;
+                    }
+                }
+
+                if (v && v !== '-') {
+                    const n = parseInt(v);
+                    if (n > 1000000000) {
+                        v = "1000000000";
+                    }
+
+                    if (n < -1) {
+                        v = "-1";
+                    }
+                }
+
+                setCommitId(v?.toString())
+
+                if (props.onChange) {
+                    props.onChange(v?.toString());
+                }
+
+            }} />
+    </Stack>
+
+})
\ No newline at end of file
diff --git a/Engine/Source/Programs/Horde/HordeDashboard/src/styles/darkTheme.ts b/Engine/Source/Programs/Horde/HordeDashboard/src/styles/darkTheme.ts
--- a/Engine/Source/Programs/Horde/HordeDashboard/src/styles/darkTheme.ts
+++ b/Engine/Source/Programs/Horde/HordeDashboard/src/styles/darkTheme.ts
@@ -386,7 +386,7 @@ export const darkTheme = createTheme({
             styles: {
                 root: {
                     "::before": {
-                        backgroundColor: "#121212"
+                        backgroundColor: "#2F2F2F"
                     }
                 }
             }
