Horde: Fix issue with p4 change sanity check, now that we have exceeded sanity

diff --git a/Engine/Source/Programs/Horde/HordeDashboard/src/components/commits/commitField.tsx b/Engine/Source/Programs/Horde/HordeDashboard/src/components/commits/commitField.tsx
--- a/Engine/Source/Programs/Horde/HordeDashboard/src/components/commits/commitField.tsx
+++ b/Engine/Source/Programs/Horde/HordeDashboard/src/components/commits/commitField.tsx
@@ -44,9 +44,6 @@ export const CommitIdField = React.forwardRef<ICommitIdField, ICommitIdFieldProp
 
                 if (v && v !== '-') {
                     const n = parseInt(v);
-                    if (n > 1000000000) {
-                        v = "1000000000";
-                    }
 
                     if (n < -1) {
                         v = "-1";
