Horde: Tweak deployment issue messaging

diff --git a/Engine/Source/Programs/Horde/Tools/ReleaseNotes/Program.cs b/Engine/Source/Programs/Horde/Tools/ReleaseNotes/Program.cs
--- a/Engine/Source/Programs/Horde/Tools/ReleaseNotes/Program.cs
+++ b/Engine/Source/Programs/Horde/Tools/ReleaseNotes/Program.cs
@@ -237,7 +237,7 @@ namespace CreateReleaseNotes
 				if (affectsServer || affectsAgent || affectsDashboard || affectsRestricted)
 				{
 					string types = (affectsServer? "S" : " ") + (affectsAgent ? "A" : " ") + (affectsDashboard? "D" : " ") + (affectsRestricted? "R" : " ");
-					logger.LogInformation("[{Types}] {Change} {Author,-20} {Description}", types, describeRecord.Number, describeRecord.User.ToLower(), logDescription);
+					logger.LogInformation("[{Types}] {Change} {Author,-20} {Restricted}{Description}", types, describeRecord.Number, describeRecord.User.ToLower(), affectsRestricted ? "Restricted - " : "", logDescription);
 				}
 			}
 
