Deleted HordeServer.DDC - This was never a supported way to use Horde as DDC cache. We recommend using Zen for most of these needs, or Cloud DDC if you need a cloud based geo replicated DDC.

diff --git a/Engine/Source/Programs/Horde/Docs/Config/Schema/AclActions.md b/Engine/Source/Programs/Horde/Docs/Config/Schema/AclActions.md
--- a/Engine/Source/Programs/Horde/Docs/Config/Schema/AclActions.md
+++ b/Engine/Source/Programs/Horde/Docs/Config/Schema/AclActions.md
@@ -354,19 +354,6 @@
 | `UbaCacheRead` |  |
 | `UbaCacheWrite` |  |
 
-## Ddc
-
-| Name | Description |
-| ---- | ----------- |
-| `DdcReadObject` | General read access to refs / blobs and so on |
-| `DdcWriteObject` | General write access to upload refs / blobs etc |
-| `DdcDeleteObject` | Access to delete blobs / refs etc |
-| `DdcDeleteBucket` | Access to delete a particular bucket |
-| `DdcDeleteNamespace` | Access to delete a whole namespace |
-| `DdcReadTransactionLog` | Access to read the transaction log |
-| `DdcWriteTransactionLog` | Access to write the transaction log |
-| `DdcAdminAction` | Access to perform administrative task |
-
 ## Devices
 
 | Name | Description |
@@ -525,19 +512,6 @@
 | `UbaCacheRead` |  |
 | `UbaCacheWrite` |  |
 
-## Ddc
-
-| Name | Description |
-| ---- | ----------- |
-| `DdcReadObject` |  |
-| `DdcWriteObject` |  |
-| `DdcDeleteObject` |  |
-| `DdcDeleteBucket` |  |
-| `DdcDeleteNamespace` |  |
-| `DdcReadTransactionLog` |  |
-| `DdcWriteTransactionLog` |  |
-| `DdcAdminAction` |  |
-
 ## Devices
 
 | Name | Description |
@@ -696,19 +670,6 @@
 | `UbaCacheRead` |  |
 | `UbaCacheWrite` |  |
 
-## Ddc
-
-| Name | Description |
-| ---- | ----------- |
-| `DdcReadObject` |  |
-| `DdcWriteObject` |  |
-| `DdcDeleteObject` |  |
-| `DdcDeleteBucket` |  |
-| `DdcDeleteNamespace` |  |
-| `DdcReadTransactionLog` |  |
-| `DdcWriteTransactionLog` |  |
-| `DdcAdminAction` |  |
-
 ## Devices
 
 | Name | Description |
@@ -867,19 +828,6 @@
 | `UbaCacheRead` |  |
 | `UbaCacheWrite` |  |
 
-## Ddc
-
-| Name | Description |
-| ---- | ----------- |
-| `DdcReadObject` |  |
-| `DdcWriteObject` |  |
-| `DdcDeleteObject` |  |
-| `DdcDeleteBucket` |  |
-| `DdcDeleteNamespace` |  |
-| `DdcReadTransactionLog` |  |
-| `DdcWriteTransactionLog` |  |
-| `DdcAdminAction` |  |
-
 ## Devices
 
 | Name | Description |
@@ -1044,19 +992,6 @@
 | `UbaCacheRead` |  |
 | `UbaCacheWrite` |  |
 
-## Ddc
-
-| Name | Description |
-| ---- | ----------- |
-| `DdcReadObject` |  |
-| `DdcWriteObject` |  |
-| `DdcDeleteObject` |  |
-| `DdcDeleteBucket` |  |
-| `DdcDeleteNamespace` |  |
-| `DdcReadTransactionLog` |  |
-| `DdcWriteTransactionLog` |  |
-| `DdcAdminAction` |  |
-
 ## Devices
 
 | Name | Description |
@@ -1234,19 +1169,6 @@
 | `UbaCacheRead` |  |
 | `UbaCacheWrite` |  |
 
-## Ddc
-
-| Name | Description |
-| ---- | ----------- |
-| `DdcReadObject` |  |
-| `DdcWriteObject` |  |
-| `DdcDeleteObject` |  |
-| `DdcDeleteBucket` |  |
-| `DdcDeleteNamespace` |  |
-| `DdcReadTransactionLog` |  |
-| `DdcWriteTransactionLog` |  |
-| `DdcAdminAction` |  |
-
 ## Devices
 
 | Name | Description |
@@ -1430,19 +1352,6 @@
 | `UbaCacheRead` |  |
 | `UbaCacheWrite` |  |
 
-## Ddc
-
-| Name | Description |
-| ---- | ----------- |
-| `DdcReadObject` |  |
-| `DdcWriteObject` |  |
-| `DdcDeleteObject` |  |
-| `DdcDeleteBucket` |  |
-| `DdcDeleteNamespace` |  |
-| `DdcReadTransactionLog` |  |
-| `DdcWriteTransactionLog` |  |
-| `DdcAdminAction` |  |
-
 ## Devices
 
 | Name | Description |
@@ -1633,19 +1542,6 @@
 | `UbaCacheRead` |  |
 | `UbaCacheWrite` |  |
 
-## Ddc
-
-| Name | Description |
-| ---- | ----------- |
-| `DdcReadObject` |  |
-| `DdcWriteObject` |  |
-| `DdcDeleteObject` |  |
-| `DdcDeleteBucket` |  |
-| `DdcDeleteNamespace` |  |
-| `DdcReadTransactionLog` |  |
-| `DdcWriteTransactionLog` |  |
-| `DdcAdminAction` |  |
-
 ## Devices
 
 | Name | Description |
diff --git a/Engine/Source/Programs/Horde/Docs/Config/Schema/Globals.md b/Engine/Source/Programs/Horde/Docs/Config/Schema/Globals.md
--- a/Engine/Source/Programs/Horde/Docs/Config/Schema/Globals.md
+++ b/Engine/Source/Programs/Horde/Docs/Config/Schema/Globals.md
@@ -62,7 +62,6 @@ Name | Description
 `analytics` | [AnalyticsConfig](#analyticsconfig)<br>Configuration for the analytics plugin
 `build` | [BuildConfig](#buildconfig)<br>Configuration for the build plugin
 `compute` | [ComputeConfig](#computeconfig)<br>Configuration for the compute plugin
-`ddc` | [EmptyPluginConfig](#emptypluginconfig)<br>Configuration for the ddc plugin
 `exampleplugin` | [ExamplePluginConfig](#examplepluginconfig)<br>Configuration for the exampleplugin plugin
 `experimental` | [ExperimentalConfig](#experimentalconfig)<br>Configuration for the experimental plugin
 `secrets` | [SecretsConfig](#secretsconfig)<br>Configuration for the secrets plugin
diff --git a/Engine/Source/Programs/Horde/Docs/Deployment/ServerSettings.md b/Engine/Source/Programs/Horde/Docs/Deployment/ServerSettings.md
--- a/Engine/Source/Programs/Horde/Docs/Deployment/ServerSettings.md
+++ b/Engine/Source/Programs/Horde/Docs/Deployment/ServerSettings.md
@@ -129,7 +129,6 @@ Name | Description
 `analytics` | [AnalyticsServerConfig](#analyticsserverconfig)<br>Configuration for the analytics plugin
 `build` | [BuildServerConfig](#buildserverconfig)<br>Configuration for the build plugin
 `compute` | [ComputeServerConfig](#computeserverconfig)<br>Configuration for the compute plugin
-`ddc` | [PluginServerConfig](#pluginserverconfig)<br>Configuration for the ddc plugin
 `exampleplugin` | [ExamplePluginServerConfig](#examplepluginserverconfig)<br>Configuration for the exampleplugin plugin
 `experimental` | [ExperimentalServerConfig](#experimentalserverconfig)<br>Configuration for the experimental plugin
 `secrets` | [SecretsServerConfig](#secretsserverconfig)<br>Configuration for the secrets plugin
diff --git a/Engine/Source/Programs/Horde/Docs/Features.md b/Engine/Source/Programs/Horde/Docs/Features.md
--- a/Engine/Source/Programs/Horde/Docs/Features.md
+++ b/Engine/Source/Programs/Horde/Docs/Features.md
@@ -9,7 +9,6 @@
 | Automation Hub | 🟩&nbsp;Production Ready | - |
 | Build Automation | 🟩&nbsp;Production Ready | Migrate jobs to a more lightweight data model.<br>Programmable state machine for jobs. |
 | Build Health | 🟩&nbsp;Production Ready | - |
-| DDC | 🟥&nbsp;Early | Recommend using UnrealCloudDDC rather than Horde. |
 | Experimental | 🟧&nbsp;Experimental | - |
 | Remote Execution | 🟩&nbsp;Production Ready | - |
 | UnrealGameSync&nbsp;Server | 🟩&nbsp;Production Ready | - |
diff --git a/Engine/Source/Programs/Horde/HordeServer/HordeServer.csproj b/Engine/Source/Programs/Horde/HordeServer/HordeServer.csproj
--- a/Engine/Source/Programs/Horde/HordeServer/HordeServer.csproj
+++ b/Engine/Source/Programs/Horde/HordeServer/HordeServer.csproj
@@ -113,7 +113,6 @@
     <ProjectReference Include="..\Plugins\Analytics\HordeServer.Analytics\HordeServer.Analytics.csproj" />
     <ProjectReference Include="..\Plugins\Build\HordeServer.Build\HordeServer.Build.csproj" />
     <ProjectReference Include="..\Plugins\Compute\HordeServer.Compute\HordeServer.Compute.csproj" />
-    <ProjectReference Include="..\Plugins\Ddc\HordeServer.Ddc\HordeServer.Ddc.csproj" />
     <ProjectReference Include="..\Plugins\ExamplePlugin\HordeServer.ExamplePlugin\HordeServer.ExamplePlugin.csproj" />
     <ProjectReference Include="..\Plugins\Experimental\HordeServer.Experimental\HordeServer.Experimental.csproj" />
     <ProjectReference Include="..\Plugins\Secrets\HordeServer.Secrets\HordeServer.Secrets.csproj" />
diff --git a/Engine/Source/Programs/Horde/Plugins/Ddc/HordeServer.Ddc.Tests/Ddc/Functional/ContentId/CompressedBlobTests.cs b/Engine/Source/Programs/Horde/Plugins/Ddc/HordeServer.Ddc.Tests/Ddc/Functional/ContentId/CompressedBlobTests.cs
deleted file mode 100644
diff --git a/Engine/Source/Programs/Horde/Plugins/Ddc/HordeServer.Ddc.Tests/Ddc/Functional/Objects/ReferencesTests.cs b/Engine/Source/Programs/Horde/Plugins/Ddc/HordeServer.Ddc.Tests/Ddc/Functional/Objects/ReferencesTests.cs
deleted file mode 100644
diff --git a/Engine/Source/Programs/Horde/Plugins/Ddc/HordeServer.Ddc.Tests/Ddc/Functional/Utils.cs b/Engine/Source/Programs/Horde/Plugins/Ddc/HordeServer.Ddc.Tests/Ddc/Functional/Utils.cs
deleted file mode 100644
diff --git a/Engine/Source/Programs/Horde/Plugins/Ddc/HordeServer.Ddc.Tests/Ddc/UnitTests/CompactBinaryTests.cs b/Engine/Source/Programs/Horde/Plugins/Ddc/HordeServer.Ddc.Tests/Ddc/UnitTests/CompactBinaryTests.cs
deleted file mode 100644
diff --git a/Engine/Source/Programs/Horde/Plugins/Ddc/HordeServer.Ddc.Tests/Ddc/UnitTests/CompressedBufferTests.cs b/Engine/Source/Programs/Horde/Plugins/Ddc/HordeServer.Ddc.Tests/Ddc/UnitTests/CompressedBufferTests.cs
deleted file mode 100644
diff --git a/Engine/Source/Programs/Horde/Plugins/Ddc/HordeServer.Ddc.Tests/HordeServer.Ddc.Tests.csproj b/Engine/Source/Programs/Horde/Plugins/Ddc/HordeServer.Ddc.Tests/HordeServer.Ddc.Tests.csproj
deleted file mode 100644
diff --git a/Engine/Source/Programs/Horde/Plugins/Ddc/HordeServer.Ddc/Ddc/.editorconfig b/Engine/Source/Programs/Horde/Plugins/Ddc/HordeServer.Ddc/Ddc/.editorconfig
deleted file mode 100644
diff --git a/Engine/Source/Programs/Horde/Plugins/Ddc/HordeServer.Ddc/Ddc/BlobContents.cs b/Engine/Source/Programs/Horde/Plugins/Ddc/HordeServer.Ddc/Ddc/BlobContents.cs
deleted file mode 100644
diff --git a/Engine/Source/Programs/Horde/Plugins/Ddc/HordeServer.Ddc/Ddc/BlobId.cs b/Engine/Source/Programs/Horde/Plugins/Ddc/HordeServer.Ddc/Ddc/BlobId.cs
deleted file mode 100644
diff --git a/Engine/Source/Programs/Horde/Plugins/Ddc/HordeServer.Ddc/Ddc/BlobService.cs b/Engine/Source/Programs/Horde/Plugins/Ddc/HordeServer.Ddc/Ddc/BlobService.cs
deleted file mode 100644
diff --git a/Engine/Source/Programs/Horde/Plugins/Ddc/HordeServer.Ddc/Ddc/BlobsController.cs b/Engine/Source/Programs/Horde/Plugins/Ddc/HordeServer.Ddc/Ddc/BlobsController.cs
deleted file mode 100644
diff --git a/Engine/Source/Programs/Horde/Plugins/Ddc/HordeServer.Ddc/Ddc/BucketId.cs b/Engine/Source/Programs/Horde/Plugins/Ddc/HordeServer.Ddc/Ddc/BucketId.cs
deleted file mode 100644
diff --git a/Engine/Source/Programs/Horde/Plugins/Ddc/HordeServer.Ddc/Ddc/BufferedPayload.cs b/Engine/Source/Programs/Horde/Plugins/Ddc/HordeServer.Ddc/Ddc/BufferedPayload.cs
deleted file mode 100644
diff --git a/Engine/Source/Programs/Horde/Plugins/Ddc/HordeServer.Ddc/Ddc/CommonHeaders.cs b/Engine/Source/Programs/Horde/Plugins/Ddc/HordeServer.Ddc/Ddc/CommonHeaders.cs
deleted file mode 100644
diff --git a/Engine/Source/Programs/Horde/Plugins/Ddc/HordeServer.Ddc/Ddc/CompressedBlobController.cs b/Engine/Source/Programs/Horde/Plugins/Ddc/HordeServer.Ddc/Ddc/CompressedBlobController.cs
deleted file mode 100644
diff --git a/Engine/Source/Programs/Horde/Plugins/Ddc/HordeServer.Ddc/Ddc/CompressedBufferUtils.cs b/Engine/Source/Programs/Horde/Plugins/Ddc/HordeServer.Ddc/Ddc/CompressedBufferUtils.cs
deleted file mode 100644
diff --git a/Engine/Source/Programs/Horde/Plugins/Ddc/HordeServer.Ddc/Ddc/ContentHash.cs b/Engine/Source/Programs/Horde/Plugins/Ddc/HordeServer.Ddc/Ddc/ContentHash.cs
deleted file mode 100644
diff --git a/Engine/Source/Programs/Horde/Plugins/Ddc/HordeServer.Ddc/Ddc/ContentId.cs b/Engine/Source/Programs/Horde/Plugins/Ddc/HordeServer.Ddc/Ddc/ContentId.cs
deleted file mode 100644
diff --git a/Engine/Source/Programs/Horde/Plugins/Ddc/HordeServer.Ddc/Ddc/ContentIdController.cs b/Engine/Source/Programs/Horde/Plugins/Ddc/HordeServer.Ddc/Ddc/ContentIdController.cs
deleted file mode 100644
diff --git a/Engine/Source/Programs/Horde/Plugins/Ddc/HordeServer.Ddc/Ddc/ContentIdStore.cs b/Engine/Source/Programs/Horde/Plugins/Ddc/HordeServer.Ddc/Ddc/ContentIdStore.cs
deleted file mode 100644
diff --git a/Engine/Source/Programs/Horde/Plugins/Ddc/HordeServer.Ddc/Ddc/FormatResolver.cs b/Engine/Source/Programs/Horde/Plugins/Ddc/HordeServer.Ddc/Ddc/FormatResolver.cs
deleted file mode 100644
diff --git a/Engine/Source/Programs/Horde/Plugins/Ddc/HordeServer.Ddc/Ddc/IBlobService.cs b/Engine/Source/Programs/Horde/Plugins/Ddc/HordeServer.Ddc/Ddc/IBlobService.cs
deleted file mode 100644
diff --git a/Engine/Source/Programs/Horde/Plugins/Ddc/HordeServer.Ddc/Ddc/IBlobStore.cs b/Engine/Source/Programs/Horde/Plugins/Ddc/HordeServer.Ddc/Ddc/IBlobStore.cs
deleted file mode 100644
diff --git a/Engine/Source/Programs/Horde/Plugins/Ddc/HordeServer.Ddc/Ddc/IContentIdService.cs b/Engine/Source/Programs/Horde/Plugins/Ddc/HordeServer.Ddc/Ddc/IContentIdService.cs
deleted file mode 100644
diff --git a/Engine/Source/Programs/Horde/Plugins/Ddc/HordeServer.Ddc/Ddc/IRefService.cs b/Engine/Source/Programs/Horde/Plugins/Ddc/HordeServer.Ddc/Ddc/IRefService.cs
deleted file mode 100644
diff --git a/Engine/Source/Programs/Horde/Plugins/Ddc/HordeServer.Ddc/Ddc/IReferenceResolver.cs b/Engine/Source/Programs/Horde/Plugins/Ddc/HordeServer.Ddc/Ddc/IReferenceResolver.cs
deleted file mode 100644
diff --git a/Engine/Source/Programs/Horde/Plugins/Ddc/HordeServer.Ddc/Ddc/IReferencesStore.cs b/Engine/Source/Programs/Horde/Plugins/Ddc/HordeServer.Ddc/Ddc/IReferencesStore.cs
deleted file mode 100644
diff --git a/Engine/Source/Programs/Horde/Plugins/Ddc/HordeServer.Ddc/Ddc/IRequestHelper.cs b/Engine/Source/Programs/Horde/Plugins/Ddc/HordeServer.Ddc/Ddc/IRequestHelper.cs
deleted file mode 100644
diff --git a/Engine/Source/Programs/Horde/Plugins/Ddc/HordeServer.Ddc/Ddc/JupiterAclAction.cs b/Engine/Source/Programs/Horde/Plugins/Ddc/HordeServer.Ddc/Ddc/JupiterAclAction.cs
deleted file mode 100644
diff --git a/Engine/Source/Programs/Horde/Plugins/Ddc/HordeServer.Ddc/Ddc/NamespaceNotFoundException.cs b/Engine/Source/Programs/Horde/Plugins/Ddc/HordeServer.Ddc/Ddc/NamespaceNotFoundException.cs
deleted file mode 100644
diff --git a/Engine/Source/Programs/Horde/Plugins/Ddc/HordeServer.Ddc/Ddc/NginxRedirectHelper.cs b/Engine/Source/Programs/Horde/Plugins/Ddc/HordeServer.Ddc/Ddc/NginxRedirectHelper.cs
deleted file mode 100644
diff --git a/Engine/Source/Programs/Horde/Plugins/Ddc/HordeServer.Ddc/Ddc/ObjectsController.cs b/Engine/Source/Programs/Horde/Plugins/Ddc/HordeServer.Ddc/Ddc/ObjectsController.cs
deleted file mode 100644
diff --git a/Engine/Source/Programs/Horde/Plugins/Ddc/HordeServer.Ddc/Ddc/RefId.cs b/Engine/Source/Programs/Horde/Plugins/Ddc/HordeServer.Ddc/Ddc/RefId.cs
deleted file mode 100644
diff --git a/Engine/Source/Programs/Horde/Plugins/Ddc/HordeServer.Ddc/Ddc/RefService.cs b/Engine/Source/Programs/Horde/Plugins/Ddc/HordeServer.Ddc/Ddc/RefService.cs
deleted file mode 100644
diff --git a/Engine/Source/Programs/Horde/Plugins/Ddc/HordeServer.Ddc/Ddc/RefsController.cs b/Engine/Source/Programs/Horde/Plugins/Ddc/HordeServer.Ddc/Ddc/RefsController.cs
deleted file mode 100644
diff --git a/Engine/Source/Programs/Horde/Plugins/Ddc/HordeServer.Ddc/Ddc/RequestHelper.cs b/Engine/Source/Programs/Horde/Plugins/Ddc/HordeServer.Ddc/Ddc/RequestHelper.cs
deleted file mode 100644
diff --git a/Engine/Source/Programs/Horde/Plugins/Ddc/HordeServer.Ddc/Ddc/RequestUtil.cs b/Engine/Source/Programs/Horde/Plugins/Ddc/HordeServer.Ddc/Ddc/RequestUtil.cs
deleted file mode 100644
diff --git a/Engine/Source/Programs/Horde/Plugins/Ddc/HordeServer.Ddc/Ddc/RequiredContentTypeAttribute.cs b/Engine/Source/Programs/Horde/Plugins/Ddc/HordeServer.Ddc/Ddc/RequiredContentTypeAttribute.cs
deleted file mode 100644
diff --git a/Engine/Source/Programs/Horde/Plugins/Ddc/HordeServer.Ddc/Ddc/Sha1Utils.cs b/Engine/Source/Programs/Horde/Plugins/Ddc/HordeServer.Ddc/Ddc/Sha1Utils.cs
deleted file mode 100644
diff --git a/Engine/Source/Programs/Horde/Plugins/Ddc/HordeServer.Ddc/Ddc/StreamUtils.cs b/Engine/Source/Programs/Horde/Plugins/Ddc/HordeServer.Ddc/Ddc/StreamUtils.cs
deleted file mode 100644
diff --git a/Engine/Source/Programs/Horde/Plugins/Ddc/HordeServer.Ddc/Ddc/StringUtils.cs b/Engine/Source/Programs/Horde/Plugins/Ddc/HordeServer.Ddc/Ddc/StringUtils.cs
deleted file mode 100644
diff --git a/Engine/Source/Programs/Horde/Plugins/Ddc/HordeServer.Ddc/DdcPlugin.cs b/Engine/Source/Programs/Horde/Plugins/Ddc/HordeServer.Ddc/DdcPlugin.cs
deleted file mode 100644
diff --git a/Engine/Source/Programs/Horde/Plugins/Ddc/HordeServer.Ddc/HordeServer.Ddc.csproj b/Engine/Source/Programs/Horde/Plugins/Ddc/HordeServer.Ddc/HordeServer.Ddc.csproj
deleted file mode 100644
diff --git a/Engine/Source/Programs/Horde/Plugins/Ddc/HordeServer.Ddc/Properties/AssemblyInfo.cs b/Engine/Source/Programs/Horde/Plugins/Ddc/HordeServer.Ddc/Properties/AssemblyInfo.cs
deleted file mode 100644
