Horde: Fixing some issues with the issue view

diff --git a/Engine/Source/Programs/Horde/HordeDashboard/src/components/IssueViewV2.tsx b/Engine/Source/Programs/Horde/HordeDashboard/src/components/IssueViewV2.tsx
--- a/Engine/Source/Programs/Horde/HordeDashboard/src/components/IssueViewV2.tsx
+++ b/Engine/Source/Programs/Horde/HordeDashboard/src/components/IssueViewV2.tsx
@@ -1256,7 +1256,8 @@ const IssueSummaryPanel: React.FC = () => {
    if (issue.duplicateIssueId !== undefined) {
       items.push({
          title: `Duplicate of Issue`,
-         text: issue.duplicateIssueId === 0 ? "Unspecified" : issue.duplicateIssueId.toString()
+         text: issue.duplicateIssueId === 0 ? "Unspecified" : issue.duplicateIssueId.toString(),
+         link: issue.duplicateIssueId === 0 ? undefined : `/index?issue=${issue.duplicateIssueId}`
       })
    }
 
@@ -1662,7 +1663,7 @@ const IssueCommandBar: React.FC = () => {
 
    const [assignToOtherShown, setAssignToOtherShown] = useState<{ shown?: boolean, defaultUser?: GetUserResponse }>({});
 
-   const [markFixedShown, setMarkFixedShown] = useState(!!query.has("markfixed"));
+   const [markFixedShown, setMarkFixedShown] = useState(!!query.has("modify"));
    const [declineShown, setDeclineShown] = useState(false);
    const [ackShown, setAckShown] = useState(false);
    const [externalIssueLinkShown, setExternalIssueLinkShown] = useState(false);
@@ -1819,6 +1820,13 @@ const IssueCommandBar: React.FC = () => {
    ];
 
    const moreList: IContextualMenuItem[] = [
+      {
+         key: 'more_modify_issue',
+         iconProps: { iconName: "Edit" },
+         text: `Modify Issue`,
+         style: { color: "unset" },
+         onClick: () => { setMarkFixedShown(true); }
+      },      
       {
          key: 'more_history_link',
          iconProps: { iconName: "SearchTemplate" },
@@ -2312,7 +2320,8 @@ const rcaCategoryOptions: IDropdownOption[] = [
    { key: 'rca_monitoring', text: 'Monitoring/No Code Change' },
    { key: 'rca_use', text: 'Unintended Side Effect' },
    { key: 'rca_buildagent', text: 'Build Agent' },
-   { key: 'rca_other', text: 'Other' }
+   { key: 'rca_other', text: 'Other' },
+   { key: 'rca_clear', text: 'Clear Selection' }
 ];
 
 const MarkFixedModal: React.FC<{ onClose: () => void }> = ({ onClose }) => {
@@ -2328,14 +2337,22 @@ const MarkFixedModal: React.FC<{ onClose: () => void }> = ({ onClose }) => {
       rootCauseSummary: issue.rootCauseSummary
    });
 
+   const navigate = useNavigate();
+
    const { hordeClasses } = getHordeStyling();
 
 
    const issueId = issue.id;
 
-   const headerText = `Mark Issue ${issueId} Fixed`;
+   const headerText = `Modify Issue ${issueId}`;
 
    const close = () => {
+
+      const url = new URL(window.location.href);
+      if (url.searchParams.has("modify")) {
+         url.searchParams.delete("modify");
+         navigate(url.pathname + url.search, { replace: true });
+      }
       onClose();
    };
 
@@ -2347,7 +2364,7 @@ const MarkFixedModal: React.FC<{ onClose: () => void }> = ({ onClose }) => {
       let fixCL: number | null | undefined;
       if (clearFixed) {
          fixCL = -1;
-      } else if (state.fixCL) {
+      } else {
          fixCL = state.fixCL;
       }
 
@@ -2359,6 +2376,9 @@ const MarkFixedModal: React.FC<{ onClose: () => void }> = ({ onClose }) => {
       let dupeId = state.dupeCheck ? 0 : -1;
       if (dupeId === 0 && state.dupeId) {
          dupeId = parseInt(state.dupeId)
+         if (dupeId === issueId) {
+            dupeId = -1;
+         }
       }
 
       let summary: string | undefined;
@@ -2379,7 +2399,7 @@ const MarkFixedModal: React.FC<{ onClose: () => void }> = ({ onClose }) => {
          rootCauseSummary: rootCauseSummary,
          ownerId: !issue.ownerInfo ? dashboard.userId : issue.ownerInfo.id,
          acknowledged: !issue.ownerInfo || issue.ownerInfo?.id === dashboard.userId ? true : undefined,
-         resolved: clearFixed ? false : true,
+         resolved: clearFixed ? false : (fixCL != undefined && fixCL >= 0),
          duplicateIssueId: dupeId
       }
 
@@ -2405,9 +2425,9 @@ const MarkFixedModal: React.FC<{ onClose: () => void }> = ({ onClose }) => {
 
    height += 384;
 
-   const selectedCatKey = rcaCategoryOptions.find(c => c.text === state.rootCauseCategory)?.key;
+   const selectedCatKey = rcaCategoryOptions.find(c => c.text === state.rootCauseCategory)?.key ?? null;
 
-   return <Modal className={hordeClasses.modal} isOpen={true} styles={{ main: { padding: 8, width: 540, height: height, minHeight: height } }} onDismiss={() => { close() }}>
+   return <Modal className={hordeClasses.modal} isOpen={true} styles={{ main: { padding: 8, width: 540, height: height, minHeight: height } }}>
       <Stack horizontal styles={{ root: { padding: 8 } }}>
          <Stack.Item grow={2}>
             <Text variant="mediumPlus">{headerText}</Text>
@@ -2427,7 +2447,8 @@ const MarkFixedModal: React.FC<{ onClose: () => void }> = ({ onClose }) => {
 
       <Stack style={{ paddingLeft: 8, paddingRight: 8 }}>
          <Stack styles={{ root: { paddingLeft: 8, width: 500, paddingTop: 12 } }}>
-            <CommitIdField label="Fix Change" placeholder="Change which fixed this issue" defaultCommitId={state.fixCL?.toString()} onChange={(commitId => { setState({ ...state, fixCL: commitId ? parseInt(commitId) : undefined }) })} />
+            <CommitIdField label="Fix Change" placeholder="Change which fixed this issue. Enter 0 to mark as fixed, without a change" defaultCommitId={state.fixCL?.toString()} onChange={(commitId => { 
+               setState({ ...state, fixCL: commitId ? parseInt(commitId) : undefined }) })} />
          </Stack>
 
          <Stack styles={{ root: { paddingLeft: 8, width: 500, paddingTop: 12 } }}>
@@ -2444,7 +2465,12 @@ const MarkFixedModal: React.FC<{ onClose: () => void }> = ({ onClose }) => {
                      placeholder="Select an option"
                      label="Root Cause Category"
                      options={rcaCategoryOptions}
-                     onChange={(event, option) => setState({ ...state, rootCauseCategory: option?.text as string | undefined })}
+                     onChange={(event, option) => {
+                        if (option?.key === "rca_clear") {                           
+                           option = undefined;
+                        }
+                        setState({ ...state, rootCauseCategory: option?.text as string | undefined })
+                     }}
                   />
                </Stack>
                <Stack>
