Horde - Removed running of deleted test project HordeServer.Ddc.Tests.csproj

diff --git a/Engine/Source/Programs/Horde/Scripts/test.sh b/Engine/Source/Programs/Horde/Scripts/test.sh
--- a/Engine/Source/Programs/Horde/Scripts/test.sh
+++ b/Engine/Source/Programs/Horde/Scripts/test.sh
@@ -18,7 +18,6 @@ testProjects=(
 	"Source/Programs/Horde/Plugins/Analytics/HordeServer.Analytics.Tests/HordeServer.Analytics.Tests.csproj"
 	"Source/Programs/Horde/Plugins/Build/HordeServer.Build.Tests/HordeServer.Build.Tests.csproj"
 	"Source/Programs/Horde/Plugins/Compute/HordeServer.Compute.Tests/HordeServer.Compute.Tests.csproj"
-	"Source/Programs/Horde/Plugins/Ddc/HordeServer.Ddc.Tests/HordeServer.Ddc.Tests.csproj"
 	"Source/Programs/Horde/Plugins/Experimental/HordeServer.Experimental.Tests/HordeServer.Experimental.Tests.csproj"
 	"Source/Programs/Horde/Plugins/Secrets/HordeServer.Secrets.Tests/HordeServer.Secrets.Tests.csproj"
 	"Source/Programs/Horde/Plugins/Storage/HordeServer.Storage.Tests/HordeServer.Storage.Tests.csproj"
