Horde: Fix clearing root cause change in modal

diff --git a/Engine/Source/Programs/Horde/HordeDashboard/src/components/IssueViewV2.tsx b/Engine/Source/Programs/Horde/HordeDashboard/src/components/IssueViewV2.tsx
--- a/Engine/Source/Programs/Horde/HordeDashboard/src/components/IssueViewV2.tsx
+++ b/Engine/Source/Programs/Horde/HordeDashboard/src/components/IssueViewV2.tsx
@@ -2390,11 +2390,11 @@ const MarkFixedModal: React.FC<{ onClose: () => void }> = ({ onClose }) => {
       } else if (issue.summary.startsWith("(Dupe) ")) {
          summary = issue.summary.replace("(Dupe) ", "");
       }
-
+      
       const request: UpdateIssueRequest = {
          fixChange: fixCL,
          summary: summary,
-         rootCommitId: rootCommitId,
+         rootCommitId: rootCommitId ?? "",
          rootCauseCategory: rootCauseCategory,
          rootCauseSummary: rootCauseSummary,
          ownerId: !issue.ownerInfo ? dashboard.userId : issue.ownerInfo.id,
