Horde: More tweaks to the Issue Modal for new functionality

diff --git a/Engine/Source/Programs/Horde/HordeDashboard/src/components/IssueViewV2.tsx b/Engine/Source/Programs/Horde/HordeDashboard/src/components/IssueViewV2.tsx
--- a/Engine/Source/Programs/Horde/HordeDashboard/src/components/IssueViewV2.tsx
+++ b/Engine/Source/Programs/Horde/HordeDashboard/src/components/IssueViewV2.tsx
@@ -1663,7 +1663,7 @@ const IssueCommandBar: React.FC = () => {
 
    const [assignToOtherShown, setAssignToOtherShown] = useState<{ shown?: boolean, defaultUser?: GetUserResponse }>({});
 
-   const [markFixedShown, setMarkFixedShown] = useState(!!query.has("modify"));
+   const [markFixedShown, setMarkFixedShown] = useState<{shown?:boolean, autoFix?:boolean}>({shown: !!query.has("modify"), autoFix: !!query.has("markfix")});
    const [declineShown, setDeclineShown] = useState(false);
    const [ackShown, setAckShown] = useState(false);
    const [externalIssueLinkShown, setExternalIssueLinkShown] = useState(false);
@@ -1761,7 +1761,7 @@ const IssueCommandBar: React.FC = () => {
       key: 'mark_fixed',
       iconProps: { iconName: issue.resolvedAt ? "IssueNew" : "IssueClosed" },
       text: issue.resolvedAt ? 'Clear Fixed' : "Mark Fixed",
-      onClick: () => { setMarkFixedShown(true); }
+      onClick: () => { setMarkFixedShown({shown: true, autoFix: true}); }
    }];
 
    const statusItems: ICommandBarItemProps[] = [
@@ -1825,7 +1825,7 @@ const IssueCommandBar: React.FC = () => {
          iconProps: { iconName: "Edit" },
          text: `Modify Issue`,
          style: { color: "unset" },
-         onClick: () => { setMarkFixedShown(true); }
+         onClick: () => { setMarkFixedShown({shown: true}); }
       },      
       {
          key: 'more_history_link',
@@ -1883,7 +1883,7 @@ const IssueCommandBar: React.FC = () => {
 
    return <Stack >
       {assignToOtherShown.shown && <AssignToOtherModal defaultUser={assignToOtherShown.defaultUser} onClose={() => { setAssignToOtherShown({}) }} />}
-      {markFixedShown && <MarkFixedModal onClose={() => { setMarkFixedShown(false) }} />}
+      {markFixedShown.shown && <MarkFixedModal autoFix={markFixedShown.autoFix} onClose={() => { setMarkFixedShown({}) }} />}
       {ackShown && <AckIssueModal onClose={() => { setAckShown(false) }} />}
       {declineShown && <DeclineIssueModal onClose={() => { setDeclineShown(false) }} />}
       {externalIssueLinkShown && <LinkExternalIssueModal onClose={() => { setExternalIssueLinkShown(false) }} />}
@@ -2324,12 +2324,12 @@ const rcaCategoryOptions: IDropdownOption[] = [
    { key: 'rca_clear', text: 'Clear Selection' }
 ];
 
-const MarkFixedModal: React.FC<{ onClose: () => void }> = ({ onClose }) => {
+const MarkFixedModal: React.FC<{ onClose: () => void, autoFix?:boolean } > = ({ onClose, autoFix }) => {
 
    const issue = details.issue!;
 
    const [state, setState] = useState<{ error?: string, fixCL?: number, rootCommitId?: CommitId, rootCauseSummary?: string, rootCauseCategory?: string, submitting?: boolean, dupeCheck?: boolean, dupeId?: string }>({
-      fixCL: issue.fixChange,
+      fixCL: issue.fixChange ?? (autoFix ? 0 : undefined),
       rootCommitId: issue.rootCommitId,
       dupeId: issue.duplicateIssueId ? (issue.duplicateIssueId.toString() !== "0" ? issue.duplicateIssueId.toString() : undefined) : undefined,
       dupeCheck: issue.duplicateIssueId === 0 || !!issue.duplicateIssueId,
@@ -2351,6 +2351,7 @@ const MarkFixedModal: React.FC<{ onClose: () => void }> = ({ onClose }) => {
       const url = new URL(window.location.href);
       if (url.searchParams.has("modify")) {
          url.searchParams.delete("modify");
+         url.searchParams.delete("markfix");
          navigate(url.pathname + url.search, { replace: true });
       }
       onClose();
@@ -2427,6 +2428,8 @@ const MarkFixedModal: React.FC<{ onClose: () => void }> = ({ onClose }) => {
 
    const selectedCatKey = rcaCategoryOptions.find(c => c.text === state.rootCauseCategory)?.key ?? null;
 
+   const disableRoot = !!state.dupeCheck;
+
    return <Modal className={hordeClasses.modal} isOpen={true} styles={{ main: { padding: 8, width: 540, height: height, minHeight: height } }}>
       <Stack horizontal styles={{ root: { padding: 8 } }}>
          <Stack.Item grow={2}>
@@ -2457,10 +2460,11 @@ const MarkFixedModal: React.FC<{ onClose: () => void }> = ({ onClose }) => {
             </Stack>
             <Stack tokens={{ childrenGap: 18 }}>
                <Stack>
-                  <CommitIdField label="Root Cause Change" defaultCommitId={state.rootCommitId} placeholder="Change which caused this issue, please only enter if certain on root cause" onChange={(commitId => { setState({ ...state, rootCommitId: commitId }) })} />
+                  <CommitIdField disabled={disableRoot} label="Root Cause Change" defaultCommitId={state.rootCommitId} placeholder="Change which caused this issue, please only enter if certain on root cause" onChange={(commitId => { setState({ ...state, rootCommitId: commitId }) })} />
                </Stack>
                <Stack>
                   <Dropdown
+                     disabled={disableRoot}
                      selectedKey={selectedCatKey}
                      placeholder="Select an option"
                      label="Root Cause Category"
@@ -2474,7 +2478,7 @@ const MarkFixedModal: React.FC<{ onClose: () => void }> = ({ onClose }) => {
                   />
                </Stack>
                <Stack>
-                  <TextField spellCheck={false} autoComplete="off" label="Root Cause Summary" placeholder="Please enter a brief summary describing the root cause of this issue." value={state.rootCauseSummary} style={{ minHeight: 120, maxHeight: 120 }} multiline resizable={false} onChange={(ev, newValue) => {
+                  <TextField disabled={disableRoot} spellCheck={false} autoComplete="off" label="Root Cause Summary" placeholder="Please enter a brief summary describing the root cause of this issue." value={state.rootCauseSummary} style={{ minHeight: 120, maxHeight: 120 }} multiline resizable={false} onChange={(ev, newValue) => {
                      ev.preventDefault();
                      setState({ ...state, rootCauseSummary: newValue ? newValue : "" })
                   }} />
@@ -2484,6 +2488,17 @@ const MarkFixedModal: React.FC<{ onClose: () => void }> = ({ onClose }) => {
                   <Checkbox label="Duplicate Issue" checked={state.dupeCheck} onChange={(_, checked) => setState({ ...state, dupeCheck: checked })} />
                   <TextField disabled={!state.dupeCheck} spellCheck={false} autoComplete="off" placeholder="Optional Issue ID" value={state.dupeId} onChange={(ev, newValue) => {
                      ev.preventDefault();
+
+                     newValue = newValue?.trim();
+
+                     // support putting in a link
+                     if (newValue?.startsWith("http")) {
+                        const lastIndex = newValue.lastIndexOf("/");
+                        const issueId = parseInt( newValue.slice(lastIndex + 1));
+                        if (typeof(issueId) === "number" && !isNaN(issueId)) {
+                           newValue = issueId.toString();
+                        }
+                     }
                      setState({ ...state, dupeId: newValue ? newValue?.trim() : "" })
                   }} />
                </Stack>
diff --git a/Engine/Source/Programs/Horde/HordeDashboard/src/components/commits/commitField.tsx b/Engine/Source/Programs/Horde/HordeDashboard/src/components/commits/commitField.tsx
--- a/Engine/Source/Programs/Horde/HordeDashboard/src/components/commits/commitField.tsx
+++ b/Engine/Source/Programs/Horde/HordeDashboard/src/components/commits/commitField.tsx
@@ -10,6 +10,7 @@ export interface ICommitIdField {
 export interface ICommitIdFieldProps {
     defaultCommitId?: CommitId;
     label?: string;
+    disabled?: boolean;
     placeholder?: string;
     onChange?: (commitId: CommitId | undefined) => void;
 }
@@ -23,7 +24,8 @@ export const CommitIdField = React.forwardRef<ICommitIdField, ICommitIdFieldProp
     }));
 
     return <Stack>
-        <TextField label={props?.label}
+        <TextField disabled={props?.disabled} label={props?.label}
+            spellCheck={false} autoComplete="off"
             placeholder={props?.placeholder}
             value={commitId ?? ""}
 
