Horde: Tweak wordng of optional duplicate issue id entry
#rnx

diff --git a/Engine/Source/Programs/Horde/HordeDashboard/src/components/IssueViewV2.tsx b/Engine/Source/Programs/Horde/HordeDashboard/src/components/IssueViewV2.tsx
--- a/Engine/Source/Programs/Horde/HordeDashboard/src/components/IssueViewV2.tsx
+++ b/Engine/Source/Programs/Horde/HordeDashboard/src/components/IssueViewV2.tsx
@@ -2486,7 +2486,7 @@ const MarkFixedModal: React.FC<{ onClose: () => void, autoFix?:boolean } > = ({
 
                <Stack horizontal verticalFill verticalAlign="center" tokens={{ childrenGap: 24 }}>
                   <Checkbox label="Duplicate Issue" checked={state.dupeCheck} onChange={(_, checked) => setState({ ...state, dupeCheck: checked })} />
-                  <TextField disabled={!state.dupeCheck} spellCheck={false} autoComplete="off" placeholder="Optional Issue ID" value={state.dupeId} onChange={(ev, newValue) => {
+                  <TextField disabled={!state.dupeCheck} spellCheck={false} autoComplete="off" placeholder="Optional Issue ID of the source issue this issue duplicates" value={state.dupeId} onChange={(ev, newValue) => {
                      ev.preventDefault();
 
                      newValue = newValue?.trim();
