Fix Info event level filtering

#rnx

diff --git a/Engine/Source/Programs/Horde/HordeDashboard/plugins/test-automation/testAutomationEvents.tsx b/Engine/Source/Programs/Horde/HordeDashboard/plugins/test-automation/testAutomationEvents.tsx
--- a/Engine/Source/Programs/Horde/HordeDashboard/plugins/test-automation/testAutomationEvents.tsx
+++ b/Engine/Source/Programs/Horde/HordeDashboard/plugins/test-automation/testAutomationEvents.tsx
@@ -252,7 +252,7 @@ export const PhaseEventsPane: React.FC<{phase: TestPhaseStatus, onEventsLoaded?:
         if (filterWarning && event.level === EventLevel.Warning) {
             return true;
         }
-        if (filterInfo) {
+        if (filterInfo && event.level === EventLevel.Information) {
             return true;
         }
 
