Horde: Update release notes

#rnx

diff --git a/Engine/Source/Programs/Horde/Docs/ReleaseNotes.md b/Engine/Source/Programs/Horde/Docs/ReleaseNotes.md
--- a/Engine/Source/Programs/Horde/Docs/ReleaseNotes.md
+++ b/Engine/Source/Programs/Horde/Docs/ReleaseNotes.md
@@ -1,5 +1,65 @@
 # Release Notes
 
+## 2025-10-02
+
+* Root Cause Analysis, dashboard side (46444423)
+* Avoid duplicate id setting recursion (46432719)
+* Adding support for Build Health duplicate issue flagging, server side (46432088)
+* Added a fix where we would not delete the cleanup batch file if there is a timeout. (46423682)
+* Improve tracing and logging for issue locks
+* Root Cause Analysis server side (46416008)
+* Added root commit to issues (46383352)
+* Show changes from restricted in release tool (46372449)
+* Linked selection mode, clear all filters functionality, and hierarchical clearing bug (edge case of last item). (46335934)
+* Adding symbolic link support to vite configuration (46294583)
+* Improve logging of AWS EC2 errors in fleet manager (46259788)
+* Optimize CreateAgentObjectAsync to only fetch agent document if update fails (46254913)
+* Updating the Horde retry warnings to no longer create Build Health issues. (46211569)
+* Remove stream salt from Gauntlet issues as they are generating too many individual issues (46203692)
+* JobLabelSummary telemetry (46137130)
+* Added columns to CSV export to help with diagnosis. (46070819)
+* Getting Horde Agent to work with an "Unreal root" hierarchy in a subfolder of the workspace root (46066712)
+* Respect Horde-based views in PerforceMaterializer (46059466)
+* EpicGames.Perforce: Add client name for stream view lookups (46059408)
+* Remove Win64 agent type fallback for installer tools (46055811)
+* Add ellipsis to job operation menu items that seem risky (46020531)
+* Fixing telemetry to use a string version of the commit ID to avoid complex objects. (46007646)
+* Quick UI fixes (46006113)
+* Forward schedule-based claims when scheduled jobs are created (46000083)
+* Add missing ExamplePluginController (45964397)
+* Avoid logging entire stack trace for Horde-related compute allocation errors (45952021)
+* Minor fix to missing relay error message (45951809)
+* Add setting for overriding available RAM in agent (45907276)
+* Reduce cookie size by removing tokens for OIDC and Okta auth (45906522)
+* Text filtering & combo box conversion (45895296)
+* Fixed coding standard violations from CL 45881300 (45883055)
+* Added feature to Horde's issues to be grouped based on customizable rules (45881300)
+* Temp Storage to allow the maximum number of exports to write in a single request to be configurable from the agent settings (45862422)
+* Styling pass V3 & general UI improvements (45849344)
+* Improve retry mechanism for Slack API client (45782122)
+* Export to csv button. (45763412)
+* Filter more files and improve description generation in release notes tool (45729945)
+* Summary column & hours label for title & improving the console.log for change summary request count so we can see the *actual* request (45714224)
+* JobSummary telemetry (45677295)
+* Add flag to report diagnostics during execution of a step (45619267)
+* Remove left-click to add highlighted text to search box in log view (45606471)
+* Horde Dashboard: Put error handling in own component to keep root view from needing to be an observer (45543682)
+* Increase title size to match changelist/date anchors & checkbox to hide date anchors (45535364)
+* Move Job/Step logging to information level (45494424)
+* Change Headers anchored by submit times (45494125)
+* Expand streams/{streamId}/changes api to support batch semantics. It can take a manually specified set of changes to obtain. (45494088)
+* Add using declaration to new JsonDocument reference to ensure disposal (45482577)
+* Resolve workspace materializer through workspace config/type (45478698)
+* Add support for file structured logging property to hashed issue handler (45441165)
+* Add support for LatestChangeNumber constant in PerforceMaterializer (45437537)
+* Add new PerforceMaterializer (45432379)
+* EpicGames.Perforce: Add a secondary release stream to fixtures (45431610)
+* Fix the DatabaseIntegrationTest class to properly forward dispose calls to the mongoDBRunner. (45403945)
+* UseEffect in order to calculate proper height. (45387504)
+* Selective dropdown updates on change & Step filters no longer being applied across streams (45340847)
+* Backing out JobExecutor killing jobs on UAT process hangs, we need to tweak this logic as it's killing jobs too early. (45313353)
+* Minor UI enhancements for 5.7 (45301672, 45280879)
+
 ## 2025-09-18
 
 * Add example backend api calls for dashboard (45941624)
