Horde: Downgrade OpenTelemetry.Instrumentation.StackExchangeRedis to match OTel version

* 1.0.0-rc9.15 requires OpenTelemetry 1.9.0+
* Everything compiles fine but fails silently at runtime unless you enable diagnostics logs for OTel

diff --git a/Engine/Source/Programs/Horde/HordeServer/HordeServer.csproj b/Engine/Source/Programs/Horde/HordeServer/HordeServer.csproj
--- a/Engine/Source/Programs/Horde/HordeServer/HordeServer.csproj
+++ b/Engine/Source/Programs/Horde/HordeServer/HordeServer.csproj
@@ -66,7 +66,7 @@
     <PackageReference Include="OpenTelemetry.Instrumentation.AspNetCore" Version="1.8.1" />
     <PackageReference Include="OpenTelemetry.Instrumentation.GrpcNetClient" Version="1.8.0-beta.1" />
     <PackageReference Include="OpenTelemetry.Instrumentation.Http" Version="1.8.1" />
-    <PackageReference Include="OpenTelemetry.Instrumentation.StackExchangeRedis" Version="1.0.0-rc9.15" />
+    <PackageReference Include="OpenTelemetry.Instrumentation.StackExchangeRedis" Version="1.0.0-rc9.14" />
     <PackageReference Include="Polly" Version="8.2.1" />
     <PackageReference Include="Polly.Extensions.Http" Version="3.0.0" />
     <PackageReference Include="Serilog.AspNetCore" Version="8.0.3" />
