Horde: Updating schema documentation

#rnx

diff --git a/Engine/Source/Programs/Horde/Docs/Config/Schema/AclActions.md b/Engine/Source/Programs/Horde/Docs/Config/Schema/AclActions.md
--- a/Engine/Source/Programs/Horde/Docs/Config/Schema/AclActions.md
+++ b/Engine/Source/Programs/Horde/Docs/Config/Schema/AclActions.md
@@ -746,164 +746,6 @@
 | `UpdateReplicator` |  |
 | `ViewReplicator` |  |
 
-## ServiceAccounts
-
-| Name | Description |
-| ---- | ----------- |
-| `CreateAccount` |  |
-| `UpdateAccount` |  |
-| `DeleteAccount` |  |
-| `ViewAccount` |  |
-
-## Sessions
-
-| Name | Description |
-| ---- | ----------- |
-| `CreateSession` |  |
-| `ViewSession` |  |
-
-## Software
-
-| Name | Description |
-| ---- | ----------- |
-| `UploadSoftware` |  |
-| `DownloadSoftware` |  |
-| `DeleteSoftware` |  |
-
-## Streams
-
-| Name | Description |
-| ---- | ----------- |
-| `CreateStream` |  |
-| `UpdateStream` |  |
-| `DeleteStream` |  |
-| `ViewStream` |  |
-| `ViewChanges` |  |
-| `ViewTemplate` |  |
-
-## Accounts
-
-| Name | Description |
-| ---- | ----------- |
-| `CreateAccount` |  |
-| `UpdateAccount` |  |
-| `DeleteAccount` |  |
-| `ViewAccount` |  |
-
-## Agents
-
-| Name | Description |
-| ---- | ----------- |
-| `CreateAgent` |  |
-| `CreateWorkstationAgent` |  |
-| `UpdateAgent` |  |
-| `DeleteAgent` |  |
-| `ViewAgent` |  |
-| `ListAgents` |  |
-
-## Artifacts
-
-| Name | Description |
-| ---- | ----------- |
-| `ReadArtifact` |  |
-| `WriteArtifact` |  |
-| `DeleteArtifact` |  |
-| `UploadArtifact` |  |
-| `DownloadArtifact` |  |
-
-## Bisect
-
-| Name | Description |
-| ---- | ----------- |
-| `CreateBisectTask` |  |
-| `UpdateBisectTask` |  |
-| `ViewBisectTask` |  |
-
-## Compute
-
-| Name | Description |
-| ---- | ----------- |
-| `AddComputeTasks` |  |
-| `GetComputeTasks` |  |
-| `UbaCacheRead` |  |
-| `UbaCacheWrite` |  |
-
-## Devices
-
-| Name | Description |
-| ---- | ----------- |
-| `DeviceRead` |  |
-| `DeviceWrite` |  |
-
-## Jobs
-
-| Name | Description |
-| ---- | ----------- |
-| `CreateJob` |  |
-| `UpdateJob` |  |
-| `DeleteJob` |  |
-| `ExecuteJob` |  |
-| `RetryJobStep` |  |
-| `ViewJob` |  |
-
-## Leases
-
-| Name | Description |
-| ---- | ----------- |
-| `ViewLeases` |  |
-| `ViewLeaseTasks` |  |
-
-## Logs
-
-| Name | Description |
-| ---- | ----------- |
-| `CreateLog` |  |
-| `UpdateLog` |  |
-| `ViewLog` |  |
-| `WriteLogData` |  |
-| `CreateEvent` |  |
-| `ViewEvent` |  |
-
-## Notices
-
-| Name | Description |
-| ---- | ----------- |
-| `CreateNotice` |  |
-| `UpdateNotice` |  |
-| `DeleteNotice` |  |
-
-## Notifications
-
-| Name | Description |
-| ---- | ----------- |
-| `CreateSubscription` |  |
-
-## Pools
-
-| Name | Description |
-| ---- | ----------- |
-| `CreatePool` |  |
-| `UpdatePool` |  |
-| `DeletePool` |  |
-| `ViewPool` |  |
-| `ListPools` |  |
-
-## Projects
-
-| Name | Description |
-| ---- | ----------- |
-| `CreateProject` |  |
-| `DeleteProject` |  |
-| `UpdateProject` |  |
-| `ViewProject` |  |
-
-## Replicators
-
-| Name | Description |
-| ---- | ----------- |
-| `UpdateReplicator` |  |
-| `ViewReplicator` |  |
-
 ## Secrets
 
 | Name | Description |
@@ -1299,11 +1141,12 @@
 | `ViewChanges` |  |
 | `ViewTemplate` |  |
 
-## Symbols
+## Tools
 
 | Name | Description |
 | ---- | ----------- |
-| `ReadSymbols` | Ability to download symbols |
+| `DownloadTool` | Ability to download a tool |
+| `UploadTool` | Ability to upload new tool versions |
 
 ## Accounts
 
@@ -1486,14 +1329,14 @@
 
 | Name | Description |
 | ---- | ----------- |
-| `ReadSymbols` |  |
+| `ReadSymbols` | Ability to download symbols |
 
 ## Tools
 
 | Name | Description |
 | ---- | ----------- |
-| `DownloadTool` | Ability to download a tool |
-| `UploadTool` | Ability to upload new tool versions |
+| `DownloadTool` |  |
+| `UploadTool` |  |
 
 ## Accounts
 
diff --git a/Engine/Source/Programs/Horde/Docs/Config/Schema/Globals.md b/Engine/Source/Programs/Horde/Docs/Config/Schema/Globals.md
--- a/Engine/Source/Programs/Horde/Docs/Config/Schema/Globals.md
+++ b/Engine/Source/Programs/Horde/Docs/Config/Schema/Globals.md
@@ -66,8 +66,8 @@ Name | Description
 `experimental` | [ExperimentalConfig](#experimentalconfig)<br>Configuration for the experimental plugin
 `secrets` | [SecretsConfig](#secretsconfig)<br>Configuration for the secrets plugin
 `storage` | [StorageConfig](#storageconfig)<br>Configuration for the storage plugin
-`symbols` | [SymbolsConfig](#symbolsconfig)<br>Configuration for the symbols plugin
 `tools` | [ToolsConfig](#toolsconfig)<br>Configuration for the tools plugin
+`symbols` | [SymbolsConfig](#symbolsconfig)<br>Configuration for the symbols plugin
 
 ## AnalyticsConfig
 
@@ -430,11 +430,6 @@ Name | Description
 `description` | `string`<br>Human-readable description
 `computeId` | `string`<br>Compute ID for this network (used when allocating compute resources)
 
-## EmptyPluginConfig
-
-Empty implementation of
-
-
 ## ExamplePluginConfig
 
 Configuration for the ExamplePlugin plugin
@@ -605,25 +600,6 @@ Name | Description
 `enableAliases` | `boolean`<br>Support querying exports by their aliases
 `acl` | [AclConfig](#aclconfig)<br>Access list for this namespace
 
-## SymbolsConfig
-
-Configuration for the tools system
-
-Name | Description
----- | -----------
-`stores` | [SymbolStoreConfig](#symbolstoreconfig)`[]`<br>List of symbol stores
-
-## SymbolStoreConfig
-
-Configuration for a symbol store
-
-Name | Description
----- | -----------
-`id` | `string`<br>Identifier for this store
-`namespaceId` | `string`<br>Configuration for the symbol store backend
-`public` | `boolean`<br>Whether to make this store available without auth
-`acl` | [AclConfig](#aclconfig)<br>Access to the symbol store
-
 ## ToolsConfig
 
 Configuration for the tools system
@@ -651,3 +627,22 @@ Name | Description
 `metadata` | `string` `->` `string`<br>Metadata for this tool
 `namespaceId` | `string`<br>Default namespace for new deployments of this tool
 `acl` | [AclConfig](#aclconfig)<br>Permissions for the tool
+
+## SymbolsConfig
+
+Configuration for the tools system
+
+Name | Description
+---- | -----------
+`stores` | [SymbolStoreConfig](#symbolstoreconfig)`[]`<br>List of symbol stores
+
+## SymbolStoreConfig
+
+Configuration for a symbol store
+
+Name | Description
+---- | -----------
+`id` | `string`<br>Identifier for this store
+`namespaceId` | `string`<br>Configuration for the symbol store backend
+`public` | `boolean`<br>Whether to make this store available without auth
+`acl` | [AclConfig](#aclconfig)<br>Access to the symbol store
