Horde: Instrument CachedCommitSource.FindAsync

diff --git a/Engine/Source/Programs/Horde/Plugins/Build/HordeServer.Build/VersionControl/Perforce/PerforceService.cs b/Engine/Source/Programs/Horde/Plugins/Build/HordeServer.Build/VersionControl/Perforce/PerforceService.cs
--- a/Engine/Source/Programs/Horde/Plugins/Build/HordeServer.Build/VersionControl/Perforce/PerforceService.cs
+++ b/Engine/Source/Programs/Horde/Plugins/Build/HordeServer.Build/VersionControl/Perforce/PerforceService.cs
@@ -1235,7 +1235,7 @@ namespace HordeServer.VersionControl.Perforce
 				}
 
 				using TelemetrySpan span = PerforceService._tracer.StartActiveSpan($"{nameof(Perforce.PerforceService)}.{nameof(PerforceService.CommitSource)}.{nameof(FindAsync)}");
-				span.SetAttribute("stream", StreamConfig.ClusterName);
+				span.SetAttribute("stream", StreamConfig.Name);
 				span.SetAttribute("minChange", minChange ?? -2);
 				span.SetAttribute("maxChange", maxChange ?? -2);
 				span.SetAttribute("maxResults", maxResults ?? -1);
diff --git a/Engine/Source/Programs/Horde/Plugins/Build/HordeServer.Build/VersionControl/Perforce/PerforceServiceCache.cs b/Engine/Source/Programs/Horde/Plugins/Build/HordeServer.Build/VersionControl/Perforce/PerforceServiceCache.cs
--- a/Engine/Source/Programs/Horde/Plugins/Build/HordeServer.Build/VersionControl/Perforce/PerforceServiceCache.cs
+++ b/Engine/Source/Programs/Horde/Plugins/Build/HordeServer.Build/VersionControl/Perforce/PerforceServiceCache.cs
@@ -632,6 +632,12 @@ namespace HordeServer.VersionControl.Perforce
 
 				int numResults = 0;
 				_owner._logger.LogDebug("Querying Perforce cache for {StreamId} commits from {MinChange} to {MaxChange} (max: {MaxResults}, tags: {Tags})", StreamConfig.Id, minChange ?? -2, maxChange ?? -2, maxResults ?? -1, (tags == null || tags.Count == 0) ? "none" : String.Join("/", tags.Select(x => x.ToString())));
+				
+				using TelemetrySpan span = _owner._tracer.StartActiveSpan($"{nameof(Perforce.PerforceService)}.{nameof(CachedCommitSource)}.{nameof(FindAsync)}");
+				span.SetAttribute("stream", StreamConfig.Name);
+				span.SetAttribute("minChange", minChange ?? -2);
+				span.SetAttribute("maxChange", maxChange ?? -2);
+				span.SetAttribute("maxResults", maxResults ?? -1);
 
 				CacheState state = await _owner._mongoService.GetSingletonAsync<CacheState>(cancellationToken);
 				if (state.Clusters.TryGetValue(StreamConfig.ClusterName, out ClusterState? clusterState))
