Horde: Address issue update, issues

diff --git a/Engine/Source/Programs/Horde/Plugins/Build/HordeServer.Build.Tests/Issues/IssueServiceTests.cs b/Engine/Source/Programs/Horde/Plugins/Build/HordeServer.Build.Tests/Issues/IssueServiceTests.cs
--- a/Engine/Source/Programs/Horde/Plugins/Build/HordeServer.Build.Tests/Issues/IssueServiceTests.cs
+++ b/Engine/Source/Programs/Horde/Plugins/Build/HordeServer.Build.Tests/Issues/IssueServiceTests.cs
@@ -616,6 +616,9 @@ namespace HordeServer.Tests.Issues
 				IUser john = await UserCollection.FindOrAddUserByLoginAsync("John");
 				_perforce.AddChange(MainStreamId, 160, john, "Description", new string[] { "Engine/Foo/Baz.txt" });
 
+				IUser anne = await UserCollection.FindOrAddUserByLoginAsync("Anne");
+				_perforce.AddChange(MainStreamId, 170, anne, "Description", new string[] { "Engine/Foo/Baz.txt" });
+
 				IJob job = CreateJob(MainStreamId, 170, "Test Build", Graph);
 				await using (TestJsonLogger logger = await CreateLoggerAsync(job, 0, 0))
 				{
@@ -669,6 +672,9 @@ namespace HordeServer.Tests.Issues
 				IUser john = await UserCollection.FindOrAddUserByLoginAsync("John");
 				_perforce.AddChange(MainStreamId, 160, john, "Description", new string[] { "Engine/Foo/Baz.txt" });
 
+				IUser anne = await UserCollection.FindOrAddUserByLoginAsync("Anne");
+				_perforce.AddChange(MainStreamId, 170, anne, "Description", new string[] { "Engine/Foo/Baz.txt" });
+
 				IJob job = CreateJob(MainStreamId, 170, "Test Build", Graph);
 				await ParseEventsAsync(job, 0, 0, lines);
 				await UpdateCompleteStepAsync(job, 0, 0, JobStepOutcome.Failure);
@@ -1038,7 +1044,7 @@ namespace HordeServer.Tests.Issues
 		}
 
 		[TestMethod]
-		public async Task EnsureWarningTestAsync()
+		public async Task EnsureNoWarningTestAsync()
 		{
 			// #1
 			// Scenario: Job step completes successfully at CL 105
@@ -1070,7 +1076,7 @@ namespace HordeServer.Tests.Issues
 				Assert.AreEqual(1, anchors[0].LineCount);
 
 				IReadOnlyList<IIssue> issues = await IssueCollection.FindIssuesAsync();
-				Assert.AreEqual(1, issues.Count);
+				Assert.AreEqual(0, issues.Count);
 			}
 		}
 
diff --git a/Engine/Source/Programs/Horde/Plugins/Build/HordeServer.Build/Issues/IssueService.cs b/Engine/Source/Programs/Horde/Plugins/Build/HordeServer.Build/Issues/IssueService.cs
--- a/Engine/Source/Programs/Horde/Plugins/Build/HordeServer.Build/Issues/IssueService.cs
+++ b/Engine/Source/Programs/Horde/Plugins/Build/HordeServer.Build/Issues/IssueService.cs
@@ -191,7 +191,7 @@ namespace HordeServer.Issues
 		/// <summary>
 		/// Maximum number of changes to query from Perforce in one go
 		/// </summary>
-		const int MaxChanges = 1000;
+		const int MaxChanges = 500;
 
 		record class CompleteStep(JobId JobId, JobStepBatchId BatchId, JobStepId StepId);
 
@@ -1374,7 +1374,7 @@ namespace HordeServer.Issues
 		/// <summary>
 		/// Maximum number of files to look at when looking for suspects
 		/// </summary>
-		const int MaxFilesForSuspects = 1000;
+		const int MaxFilesForSuspects = 100;
 
 		/// <summary>
 		/// Find suspects that can match a given span
@@ -1390,103 +1390,137 @@ namespace HordeServer.Issues
 			using TelemetrySpan span = _tracer.StartActiveSpan($"{nameof(IssueService)}.{nameof(FindSuspectsForSpanAsync)}");
 			
 			List<NewIssueSpanSuspectData> suspects = new List<NewIssueSpanSuspectData>();
-			if (!String.IsNullOrEmpty(fingerprint.ChangeFilter))
-			{
-				_logger.LogDebug("Querying for changes in {StreamName} between {LastSuccess} and {FirstFailure}", streamConfig.Name, lastSuccess, firstFailure);
 
-				// Use fingerprint type to get commit tags and narrow search window
-				List<CommitTag>? tags = null;
-				if (fingerprint.Type.StartsWith("Compile", StringComparison.OrdinalIgnoreCase))
+			try
+			{
+				if (!String.IsNullOrEmpty(fingerprint.ChangeFilter))
 				{
-					tags = [CommitTag.Code];
-				}
+					ICommit? firstFailureCommit = await _commitService.GetCollection(streamConfig).GetAsync(firstFailure, cancellationToken);
+					ICommit? lastSuccessCommit = await _commitService.GetCollection(streamConfig).GetAsync(lastSuccess, cancellationToken);
+
+					if (firstFailureCommit == null || lastSuccessCommit == null)
+					{
+						_logger.LogError("Failed to get commits for {StreamName}: {LastSuccess} and {FirstFailure}", streamConfig.Name, lastSuccess, firstFailure);
+						return suspects;
+					}
 
-				// Get the submitted changes before this job
-				List<ICommit> changes = await _commitService.GetCollection(streamConfig).FindAsync(lastSuccess, includeMinCommit: false, firstFailure, includeMaxCommit: true, MaxChanges, tags, cancellationToken: cancellationToken).ToListAsync(cancellationToken);
-				_logger.LogDebug("Found {NumResults} changes", changes.Count);
+					TimeSpan delta = lastSuccessCommit.DateUtc - firstFailureCommit.DateUtc;
+					
+					if (delta.TotalDays >= 7)
+					{
+						_logger.LogInformation("Ignoring commits for {StreamName}: {LastSuccess} and {FirstFailure}, delta is {DeltaDays}", streamConfig.Name, lastSuccess, firstFailure.Name, delta.TotalDays);
+						return suspects;
+					}
 
-				// Get all the parameters used to rank suspects
-				FileFilter filter = new FileFilter(fingerprint.ChangeFilter.Split(";", StringSplitOptions.RemoveEmptyEntries | StringSplitOptions.TrimEntries));
+					_logger.LogDebug("Querying for changes in {StreamName} between {LastSuccess} and {FirstFailure}", streamConfig.Name, lastSuccess, firstFailure);
 
-				// Build a set of all the files to include as suspects
-				HashSet<string> suspectFiles = new HashSet<string>(StringComparer.OrdinalIgnoreCase);
-				foreach (IssueKey key in fingerprint.Keys)
-				{
-					if (key.Type == IssueKeyType.File)
+					// Use fingerprint type to get commit tags and narrow search window
+					List<CommitTag>? tags = null;
+					if (fingerprint.Type.StartsWith("Compile", StringComparison.OrdinalIgnoreCase))
 					{
-						suspectFiles.Add(key.Name);
+						tags = [CommitTag.Code];
 					}
-				}
 
-				if (fingerprint.Metadata != null)
-				{
-					foreach (IssueMetadata metadata in fingerprint.Metadata)
+					// Get the submitted changes before this job
+					List<ICommit> changes = await _commitService.GetCollection(streamConfig).FindAsync(lastSuccess, includeMinCommit: false, firstFailure, includeMaxCommit: true, MaxChanges, tags, cancellationToken: cancellationToken).ToListAsync(cancellationToken);
+					_logger.LogDebug("Found {NumResults} changes", changes.Count);
+
+					// Get all the parameters used to rank suspects
+					FileFilter filter = new FileFilter(fingerprint.ChangeFilter.Split(";", StringSplitOptions.RemoveEmptyEntries | StringSplitOptions.TrimEntries));
+
+					// Build a set of all the files to include as suspects
+					HashSet<string> suspectFiles = new HashSet<string>(StringComparer.OrdinalIgnoreCase);
+					foreach (IssueKey key in fingerprint.Keys)
 					{
-						if (metadata.Key == "File")
+						if (key.Type == IssueKeyType.File)
 						{
-							suspectFiles.Add(metadata.Value);
+							suspectFiles.Add(key.Name);
 						}
 					}
-				}
 
-				// Build a set of 'notes' to include as suspects
-				HashSet<string> suspectNotes = new HashSet<string>(StringComparer.OrdinalIgnoreCase);
-				foreach (IssueKey key in fingerprint.Keys)
-				{
-					if (key.Type == IssueKeyType.Note)
+					if (fingerprint.Metadata != null)
 					{
-						suspectNotes.Add(key.Name);
+						foreach (IssueMetadata metadata in fingerprint.Metadata)
+						{
+							if (metadata.Key == "File")
+							{
+								suspectFiles.Add(metadata.Value);
+							}
+						}
 					}
-				}
 
-				// Build a set of symbols to include as suspects
-				HashSet<string> suspectSymbols = new HashSet<string>(StringComparer.OrdinalIgnoreCase);
-				foreach (IssueKey key in fingerprint.Keys)
-				{
-					if (key.Type == IssueKeyType.Symbol)
+					// Build a set of 'notes' to include as suspects
+					HashSet<string> suspectNotes = new HashSet<string>(StringComparer.OrdinalIgnoreCase);
+					foreach (IssueKey key in fingerprint.Keys)
 					{
-						suspectSymbols.UnionWith(key.Name.Split("::", StringSplitOptions.RemoveEmptyEntries));
+						if (key.Type == IssueKeyType.Note)
+						{
+							suspectNotes.Add(key.Name);
+						}
 					}
-				}
 
-				// Get the filter for changes to include in the suspect list
-				int maxRank = -1;
-				foreach (ICommit commit in changes)
-				{
-					int rank = -1;
-					try
+					// Build a set of symbols to include as suspects
+					HashSet<string> suspectSymbols = new HashSet<string>(StringComparer.OrdinalIgnoreCase);
+					foreach (IssueKey key in fingerprint.Keys)
 					{
-						rank = await GetChangeScoreAsync(commit, filter, suspectFiles, suspectNotes, suspectSymbols, _tracer);
+						if (key.Type == IssueKeyType.Symbol)
+						{
+							suspectSymbols.UnionWith(key.Name.Split("::", StringSplitOptions.RemoveEmptyEntries));
+						}
 					}
-					catch (NullReferenceException nullRefEx)
-					{
-						_logger.LogError("Exception ranking suspect change {Commit}: {Message} - {Callstack}", commit.Id, nullRefEx.Message, nullRefEx.StackTrace);
 
-					}
-					catch (EpicGames.Perforce.PerforceException perforceEx)
+					DateTime scoringStart = DateTime.Now;
+
+					// Get the filter for changes to include in the suspect list
+					int maxRank = -1;
+					foreach (ICommit commit in changes)
 					{
-						_logger.LogError("Exception ranking suspect change {Commit}: {Message} - {Callstack}", commit.Id, perforceEx.Message, perforceEx.StackTrace);
+						TimeSpan scoringDuration = DateTime.Now - scoringStart;
+						if (scoringDuration.TotalSeconds >= 10)
+						{
+							_logger.LogWarning("Scoring suspect commits took longer than 10 seconds, giving up");
+							break;
+						}
 
-						// re-throw any perforce exceptions that aren't known or expected here
-						if (!perforceEx.Message.Contains("does not contain any files in", System.StringComparison.OrdinalIgnoreCase))
+						int rank = -1;
+						try
 						{
-							throw;
+							rank = await GetChangeScoreAsync(commit, filter, suspectFiles, suspectNotes, suspectSymbols, _tracer);
 						}
-					}
-					if (rank > maxRank)
-					{
-						suspects.Clear();
-						maxRank = rank;
-					}
+						catch (NullReferenceException nullRefEx)
+						{
+							_logger.LogError("Exception ranking suspect change {Commit}: {Message} - {Callstack}", commit.Id, nullRefEx.Message, nullRefEx.StackTrace);
 
-					if (rank == maxRank)
-					{
-						NewIssueSpanSuspectData suspect = new NewIssueSpanSuspectData(commit.Id, commit.AuthorId == UserId.Empty || commit.AuthorId == UserId.Anonymous ? commit.OwnerId : commit.AuthorId);
-						suspect.SourceCommitId = commit.OriginalCommitId;
-						suspects.Add(suspect);
+						}
+						catch (EpicGames.Perforce.PerforceException perforceEx)
+						{
+							_logger.LogError("Exception ranking suspect change {Commit}: {Message} - {Callstack}", commit.Id, perforceEx.Message, perforceEx.StackTrace);
+
+							// re-throw any perforce exceptions that aren't known or expected here
+							if (!perforceEx.Message.Contains("does not contain any files in", System.StringComparison.OrdinalIgnoreCase))
+							{
+								throw;
+							}
+						}
+						if (rank > maxRank)
+						{
+							suspects.Clear();
+							maxRank = rank;
+						}
+
+						if (rank == maxRank)
+						{
+							NewIssueSpanSuspectData suspect = new NewIssueSpanSuspectData(commit.Id, commit.AuthorId == UserId.Empty || commit.AuthorId == UserId.Anonymous ? commit.OwnerId : commit.AuthorId);
+							suspect.SourceCommitId = commit.OriginalCommitId;
+							suspects.Add(suspect);
+						}
 					}
 				}
 			}
+			catch (Exception ex)
+			{
+				_logger.LogError(ex, "Exception finding suspects for {StreamName}: {LastSuccess} and {FirstFailure}", streamConfig.Name, lastSuccess, firstFailure.Name);
+			}
 			return suspects;
 		}
 
