Horde: Fix issue with issue generation delta comparison

diff --git a/Engine/Source/Programs/Horde/Plugins/Build/HordeServer.Build/Issues/IssueService.cs b/Engine/Source/Programs/Horde/Plugins/Build/HordeServer.Build/Issues/IssueService.cs
--- a/Engine/Source/Programs/Horde/Plugins/Build/HordeServer.Build/Issues/IssueService.cs
+++ b/Engine/Source/Programs/Horde/Plugins/Build/HordeServer.Build/Issues/IssueService.cs
@@ -1404,9 +1404,9 @@ namespace HordeServer.Issues
 						return suspects;
 					}
 
-					TimeSpan delta = lastSuccessCommit.DateUtc - firstFailureCommit.DateUtc;
+					TimeSpan delta = firstFailureCommit.DateUtc - lastSuccessCommit.DateUtc;
 					
-					if (delta.TotalDays >= 7)
+					if (Math.Abs(delta.TotalDays) >= 7)
 					{
 						_logger.LogInformation("Ignoring commits for {StreamName}: {LastSuccess} and {FirstFailure}, delta is {DeltaDays}", streamConfig.Name, lastSuccess, firstFailure.Name, delta.TotalDays);
 						return suspects;
