[Horde][BuildHealth][Dashboard] Collection fo QoL UI features for interacting with dropdowns:
- clear all + select all (for each group)
-- filter context is applied
- sticky search field
- Jump to top/bottom of scrollable
- Jump up group and down group
- optimizations to dropdown to reduce load on heavy case (removal of ToolTipHost)
- Compression of step names in query param (remove duplicates of data across all stepnames)
-- query param schema version & auto-upgrades to v1
- NaN fix for summary where no completed steps exist

diff --git a/Engine/Source/Programs/Horde/HordeDashboard/src/components/buildhealth/BuildHealthDataHandler.ts b/Engine/Source/Programs/Horde/HordeDashboard/src/components/buildhealth/BuildHealthDataHandler.ts
--- a/Engine/Source/Programs/Horde/HordeDashboard/src/components/buildhealth/BuildHealthDataHandler.ts
+++ b/Engine/Source/Programs/Horde/HordeDashboard/src/components/buildhealth/BuildHealthDataHandler.ts
@@ -3,10 +3,10 @@
 import backend from "horde/backend";
 import { JobData, JobStreamQuery, ProjectData, StreamData } from "horde/backend/Api";
 import { makeObservable, observable, runInAction } from "mobx";
-import { BuildHealthOptionsController, BuildHealthOptionsState, IHierarchyValidator } from "./BuildHealthOptions";
+import { BuildHealthOptionsState } from "./BuildHealthOptions";
 import { StepRefData, TemplateRefData } from "./BuildHealthDataTypes";
-import { decodeTemplateKey, encodeTemplateKey, encodeTemplateKeyFromStrings } from "./BuildHealthUtilities";
-import { find } from "@fluentui/react";
+import { decodeTemplateKey, encodeTemplateKeyFromStrings } from "./BuildHealthUtilities";
+import { BuildHealthOptionsController, IHierarchyValidator } from "./BuildHealthOptionsController";
 
 // #region -- API Data Types & Utilities --
 
@@ -46,6 +46,14 @@ export function expandRefreshRequest(request: DataHandlerRefreshRequest): DataHa
 
 // #endregion -- API Data Types & Utilities --
 
+/**
+ * Represents a cache entry for step ref data.
+ */
+interface StepCacheEntry {
+    timestamp: number;
+    steps: StepRefData[];
+}
+
 /**
  * Data handler for all build health data used for dropdowns.
  * @todo Data sharing & caching would ideally occur between the BuildHealth View & Step Outcome.
@@ -60,7 +68,7 @@ export class BuildHealthDataHandler implements IHierarchyValidator {
     private buildHealthOptions: BuildHealthOptionsState;
     private buildHealthOptionsController: BuildHealthOptionsController;
     private cachedStreamToTemplatesData: Map<string, TemplateRefData[]> = new Map<string, TemplateRefData[]>();
-    private cachedStreamTempalteIdToSteps: Map<string, StepRefData[]> = new Map<string, StepRefData[]>();
+    private cachedStreamTempalteIdToSteps: Map<string, StepCacheEntry> = new Map<string, StepCacheEntry>();
 
     // #endregion -- Private Members --
 
@@ -114,11 +122,13 @@ export class BuildHealthDataHandler implements IHierarchyValidator {
      * Requests a hierarchical data refresh. Refreshw ill occur if the options have changed since last refresh.
      * @param request The data refresh request message to use in order to control refresh depth.
      * @param force Whether to force the teh refresh or not. If false, will only refresh if the options have changed since last refresh.
+     * @param onComplete Delegate that will be invoked upon completion of the hierarchical refresh.
      */
-    async requestHierarchicalRefresh(request: DataHandlerRefreshRequest = DEFAULT_REFRESH_REQUEST, force?: boolean) {
+    async requestHierarchicalRefresh(request: DataHandlerRefreshRequest = DEFAULT_REFRESH_REQUEST, force?: boolean, onComplete?: (request: DataHandlerRefreshRequest) => void) {
         if (this.lastRefresh < this.buildHealthOptionsController.optionsChangeVersion || force) {
             this.lastRefresh = this.buildHealthOptionsController.optionsChangeVersion;
             await this.refreshData(request);
+            onComplete?.(request);
         }
     }
 
@@ -203,9 +213,10 @@ export class BuildHealthDataHandler implements IHierarchyValidator {
      */
     async getStepData(request: DataHandlerRefreshRequest) {
         const entries = Object.entries(this.buildHealthOptions.enabledTemplates);
-        this.cachedStreamTempalteIdToSteps.clear();
         if (request.steps && entries.length > 0) {
-            const results = await Promise.all(entries.map(async ([key]) => {
+            const enabledIds = new Set(entries.map(([id]) => id));
+
+            await Promise.all(entries.filter(([key, _value]) => !this.cachedStreamTempalteIdToSteps.has(key)).map(async ([key]) => {
                 const templateData = decodeTemplateKey(key);
 
                 const query: JobStreamQuery = {
@@ -230,23 +241,54 @@ export class BuildHealthDataHandler implements IHierarchyValidator {
                 }
 
                 const distinctSteps = [...stepMap.values()];
-                this.cachedStreamTempalteIdToSteps.set(key, distinctSteps);
+
+                this.cachedStreamTempalteIdToSteps.set(key, { timestamp: Date.now(), steps: distinctSteps });
 
                 return distinctSteps;
             }));
 
-            const allSteps = results.flat();
-
             runInAction(() => {
-                this.stepData = [...allSteps.values()];
+                this.stepData = Array.from(this.cachedStreamTempalteIdToSteps.entries())
+                    .filter(([templateId]) => enabledIds.has(templateId))
+                    .flatMap(([_, entry]) => entry.steps);
             });
-        } else if (request.steps) {  // We have no enabled templates, so we should have no step data showing.
+        } else if (request.steps) { // We have no enabled templates, so we should have no step data showing.
+
             runInAction(() => {
                 this.stepData = [];
             });
         }
     }
 
+    /**
+     * Obtains all the currently known step names for a given stream & template.
+     * @param streamId The stream.
+     * @param templateId The template.
+     * @returns An array representing all the known step names for a given stream & template pair.
+     */
+    getCachedTemplateStepNames(streamId: string, templateId: string): string[] {
+        let returnArray: string[] = [];
+
+        let encodedTemplateKey: string = encodeTemplateKeyFromStrings(streamId, templateId);
+
+        if (this.cachedStreamTempalteIdToSteps.has(encodedTemplateKey)) {
+            returnArray = this.cachedStreamTempalteIdToSteps.get(encodedTemplateKey)!.steps.map(x => x.name);
+        }
+
+        return returnArray;
+    }
+
+    /**
+     * Obtains all the StepRefData for a given stream & template.
+     * @param streamId The stream.
+     * @param templateId The template.
+     * @returns An array representing all the StepRefData for a given stream & template pair.
+     */
+    getCachedTemplateStepData(streamId: string, templateId: string): StepRefData[] | undefined {
+        let encodedTemplateKey: string = encodeTemplateKeyFromStrings(streamId, templateId);
+        return this.cachedStreamTempalteIdToSteps.get(encodedTemplateKey)?.steps;
+    }
+
     /**
      * @inheritdoc
      */
@@ -268,13 +310,14 @@ export class BuildHealthDataHandler implements IHierarchyValidator {
         if (!this.cachedStreamToTemplatesData.has(streamId)) {
             return false;
         }
+
         let template = this.cachedStreamToTemplatesData.get(streamId)?.find(x => x.id === templateId)
 
         if (template === undefined) {
             return false;
         }
         let encodedTemplateKey: string = encodeTemplateKeyFromStrings(streamId, templateId);
-        if (!this.cachedStreamTempalteIdToSteps.get(encodedTemplateKey)?.find(x => x.name == stepName)) {
+        if (!this.cachedStreamTempalteIdToSteps.get(encodedTemplateKey)?.steps.find(x => x.name == stepName)) {
             return false;
         }
 
diff --git a/Engine/Source/Programs/Horde/HordeDashboard/src/components/buildhealth/BuildHealthDataTypes.ts b/Engine/Source/Programs/Horde/HordeDashboard/src/components/buildhealth/BuildHealthDataTypes.ts
--- a/Engine/Source/Programs/Horde/HordeDashboard/src/components/buildhealth/BuildHealthDataTypes.ts
+++ b/Engine/Source/Programs/Horde/HordeDashboard/src/components/buildhealth/BuildHealthDataTypes.ts
@@ -2,6 +2,7 @@
 
 import { GetStepResponse, GetTemplateRefResponse } from "horde/backend/Api";
 
+export const QUERY_SCHEMA_PARAM_NAME: string = "query_schema";
 export const PROJECTS_URL_PARAM_NAME: string = "projects";
 export const STREAMS_URL_PARAM_NAME: string = "streams";
 export const TEMPLATES_URL_PARAM_NAME: string = "templates";
@@ -11,6 +12,8 @@ export const TIME_SPAN_URL_PARAM: string = "lastTimeRange";
 export const INCLUDE_PREFLIGHT_PARAM: string = "includePreflight";
 export const KEY_SEPARATOR: string = "::";
 
+export const STEPS_URL_PARAM_ALL_STEPS: string = "stepNames_AllKnownValues";
+
 /**
  * Convenience type to fold in extra metadata relevant for the Template. 
  */
diff --git a/Engine/Source/Programs/Horde/HordeDashboard/src/components/buildhealth/BuildHealthOptions.ts b/Engine/Source/Programs/Horde/HordeDashboard/src/components/buildhealth/BuildHealthOptions.ts
--- a/Engine/Source/Programs/Horde/HordeDashboard/src/components/buildhealth/BuildHealthOptions.ts
+++ b/Engine/Source/Programs/Horde/HordeDashboard/src/components/buildhealth/BuildHealthOptions.ts
@@ -1,11 +1,12 @@
 // Copyright Epic Games, Inc. All Rights Reserved.
 
-import { action, makeObservable, observable, runInAction } from "mobx";
-import { INCLUDE_PREFLIGHT_PARAM, PROJECTS_URL_PARAM_NAME, STEPS_URL_PARAM_NAME, STREAMS_URL_PARAM_NAME, TEMPLATES_URL_PARAM_NAME, TIME_SPAN_URL_PARAM } from "../buildhealth/BuildHealthDataTypes";
+import { makeObservable, observable, runInAction } from "mobx";
+import { INCLUDE_PREFLIGHT_PARAM, PROJECTS_URL_PARAM_NAME, QUERY_SCHEMA_PARAM_NAME, STEPS_URL_PARAM_ALL_STEPS, STEPS_URL_PARAM_NAME, STREAMS_URL_PARAM_NAME, TEMPLATES_URL_PARAM_NAME, TIME_SPAN_URL_PARAM } from "../buildhealth/BuildHealthDataTypes";
 import { BuildHealthDataHandler } from "./BuildHealthDataHandler";
 import { decodeStepKey, decodeTemplateKey, encodeStepKey, encodeStepNameFromStringsPreserveCase, encodeTemplateKey, encodeTemplateKeyFromStrings } from "./BuildHealthUtilities";
 import { Location } from 'react-router-dom';
 import LZString from "lz-string"
+import { BuildHealthOptionsController, IBuildHealthOptionWriter } from "./BuildHealthOptionsController";
 
 // #region -- Options Utilities & Data Types --
 
@@ -19,6 +20,7 @@ type BuildHealthQueryParams = {
     stepNames: string[];
     timeSpanKey: string;
     includePreflight: boolean;
+    querySchemaVersion: number;
 };
 
 /**
@@ -62,13 +64,121 @@ export const JobHistoryTimeSpans: TimeSpan[] = [
 
 export const DEFAULT_JOB_HISTORY_TIME_SPAN: TimeSpan = JobHistoryTimeSpans[1];
 
+/**
+ * Obtains a URLSearchParams representation of the current options.
+ * @returns The URLSearchParams representation of the current options.
+ */
+export function toNavigationQuery(optionsState: BuildHealthOptionsState, querySchemaVersion: number, buildHealthDataHandler: BuildHealthDataHandler): URLSearchParams {
+
+    // Internal helper method to compress keys to the comrpessed representation.
+    // e.g. stream-id::template-id::stepname-1,stream-id::template-id::stepname-2 -> stream-id::template-id(stepname-1;stepname-2)
+    function compressKeys(keys: string[], encode: (k: string) => { prefix: string, suffix: string }, groupSummarizer?: (prefix: string, suffixes: string[]) => string | undefined): string {
+        const groups: Record<string, string[]> = {};
+
+        for (const key of keys) {
+            const { prefix, suffix } = encode(key);
+            if (!groups[prefix]) {
+                groups[prefix] = [];
+            }
+            groups[prefix].push(suffix);
+        }
+
+        return Object.entries(groups)
+            .map(([prefix, suffixes]) => {
+                let summarizedSuffix = groupSummarizer ? groupSummarizer(prefix, suffixes) : undefined;
+                let finalSuffix = summarizedSuffix ?? suffixes.join(",");
+
+                return `${prefix}(${finalSuffix})`;
+            })
+            .join(";");
+    }
+
+    const params = new URLSearchParams(location.search);
+    let projects: string = "";
+    let streams: string = "";
+    let jobs: string = "";
+    let steps: string = "";
+
+    projects = Object.keys(optionsState.enabledProjects).join(",");
+    streams = optionsState.stepOutcomeEnabledStreamKeys.join(",");
+    const lastTimeRange = optionsState.jobHistoryTimeSpan.key.toString();
+    const includePreflight = optionsState.includePreflight;
+
+    if (querySchemaVersion === 1) {
+        jobs = optionsState.stepOutcomeEnabledJobKeys.join(",");
+        steps = optionsState.stepOutcomeEnabledStepKeys.join(",");
+    }
+    else if (querySchemaVersion === 2) {
+
+        jobs = compressKeys(optionsState.stepOutcomeEnabledJobKeys, (key: string) => {
+            let { streamId, templateId } = decodeTemplateKey(key);
+
+            return { prefix: streamId, suffix: templateId };
+        });
+
+        steps = compressKeys(optionsState.stepOutcomeEnabledStepKeys,
+            (key: string) => {
+                let { streamId, templateId, stepName } = decodeStepKey(key);
+
+                return { prefix: encodeTemplateKeyFromStrings(streamId, templateId), suffix: stepName };
+            },
+            (prefix: string, suffixes: string[]) => {
+                let { streamId, templateId } = decodeTemplateKey(prefix);
+                let summarized = buildHealthDataHandler.getCachedTemplateStepData(streamId, templateId)?.length == suffixes.length ? STEPS_URL_PARAM_ALL_STEPS : undefined;
+
+                return summarized;
+            });
+    }
+
+    params.set(PROJECTS_URL_PARAM_NAME, LZString.compressToEncodedURIComponent(projects));
+    params.set(STREAMS_URL_PARAM_NAME, LZString.compressToEncodedURIComponent(streams));
+    params.set(TEMPLATES_URL_PARAM_NAME, LZString.compressToEncodedURIComponent(jobs));
+    params.set(STEPS_URL_PARAM_NAME, LZString.compressToEncodedURIComponent(steps));
+    params.set(TIME_SPAN_URL_PARAM, LZString.compressToEncodedURIComponent(lastTimeRange));
+    params.set(INCLUDE_PREFLIGHT_PARAM, LZString.compressToEncodedURIComponent(String(includePreflight)));
+    params.set(QUERY_SCHEMA_PARAM_NAME, querySchemaVersion.toString());
+
+    return params;
+}
+
 /**
  * Parses the Build Health options from the query params.
  * @param location Location object.
  * @returns Parsed Build Health Query Params.
  */
 export function parseBuildHealthQueryParams(location: Location): BuildHealthQueryParams {
+
+    // Internal helper method to expand keys from the comrpessed representation.
+    // e.g. stream-id::template-id(stepname-1;stepname-2) -> stream-id::template-id::stepname-1,stream-id::template-id::stepname-2
+    function expandKeys(encoded: string, decode: (prefix: string, suffix: string) => string): string[] {
+        if (!encoded || !encoded.trim()) return [];
+
+        const result: string[] = [];
+
+        const groups = encoded.split(";");
+
+        for (const group of groups) {
+            const match = group.match(/^([^()]+)\((.*)\)$/);
+            if (match) {
+                const prefix = match[1];
+                const suffixes = match[2].split(",");
+                for (const suffix of suffixes) {
+                    result.push(decode(prefix, suffix));
+                }
+            } else {
+                console.warn("Failed to decode parameter: " + group);
+            }
+        }
+
+        return result;
+    }
+
     const searchParams = new URLSearchParams(location.search);
+    let projectIds: string[] = [];
+    let streamIds: string[] = [];
+    let templateIds: string[] = [];
+    let stepNames: string[] = [];
+
     const getList = (key: string) => {
         const raw = searchParams.get(key);
         if (!raw) return [];
@@ -79,10 +189,39 @@ export function parseBuildHealthQueryParams(location: Location): BuildHealthQuer
         return decompressed.split(",").filter(Boolean);
     };
 
-    const projectIds = getList(PROJECTS_URL_PARAM_NAME);
-    const streamIds = getList(STREAMS_URL_PARAM_NAME);
-    const templateIds = getList(TEMPLATES_URL_PARAM_NAME);
-    const stepNames = getList(STEPS_URL_PARAM_NAME);
+    const querySchemaVersion = Number(searchParams.get(QUERY_SCHEMA_PARAM_NAME)) || 1;
+    projectIds = getList(PROJECTS_URL_PARAM_NAME);
+    streamIds = getList(STREAMS_URL_PARAM_NAME);
+
+    if (querySchemaVersion === 1) {
+        templateIds = getList(TEMPLATES_URL_PARAM_NAME);
+        stepNames = getList(STEPS_URL_PARAM_NAME);
+    }
+    else if (querySchemaVersion === 2) {
+        const rawTemplates = searchParams.get(TEMPLATES_URL_PARAM_NAME);
+        if (!rawTemplates) {
+            templateIds = [];
+        }
+        else {
+            let decompressedTemplates = LZString.decompressFromEncodedURIComponent(rawTemplates);
+            templateIds = expandKeys(decompressedTemplates, (prefix, suffix) => {
+                return encodeTemplateKeyFromStrings(prefix, suffix);
+            });
+        }
+
+        const rawSteps = searchParams.get(STEPS_URL_PARAM_NAME);
+        if (!rawSteps) {
+            stepNames = [];
+        }
+        else {
+            let decompressedSteps = LZString.decompressFromEncodedURIComponent(rawSteps);
+            stepNames = expandKeys(decompressedSteps, (prefix, suffix) => {
+                let { streamId, templateId } = decodeTemplateKey(prefix);
+
+                return encodeStepNameFromStringsPreserveCase(streamId, templateId, suffix);
+            });
+        }
+    }
 
     let timeSpanKey: string = (() => {
         const raw = searchParams.get(TIME_SPAN_URL_PARAM);
@@ -96,7 +235,7 @@ export function parseBuildHealthQueryParams(location: Location): BuildHealthQuer
     // Read from URLSearchParams
     const includePreflight = (searchParams.get(INCLUDE_PREFLIGHT_PARAM) ? LZString.decompressFromEncodedURIComponent(searchParams.get(INCLUDE_PREFLIGHT_PARAM)!) : "false") === "true";
 
-    return { projectIds, streamIds, templateIds: templateIds, stepNames, timeSpanKey, includePreflight };
+    return { projectIds, streamIds, templateIds: templateIds, stepNames, timeSpanKey, includePreflight, querySchemaVersion };
 }
 
 /**
@@ -109,14 +248,15 @@ export async function loadBuildHealthOptionsFromParams(
     params: BuildHealthQueryParams,
     options: BuildHealthOptionsController,
     handler: BuildHealthDataHandler) {
+    let optionsWriter: IBuildHealthOptionWriter = options.getTransactionSession();
 
     runInAction(() => {
-        options.setIncludePreflight(params.includePreflight);
+        optionsWriter.setIncludePreflight(params.includePreflight);
     })
 
     if (params.timeSpanKey) {
         runInAction(() => {
-            options.setJobHistoryTimeSpan(JobHistoryTimeSpans.find(ts => ts.key === params.timeSpanKey)!);
+            optionsWriter.setJobHistoryTimeSpan(JobHistoryTimeSpans.find(ts => ts.key === params.timeSpanKey)!);
         });
     }
 
@@ -124,11 +264,11 @@ export async function loadBuildHealthOptionsFromParams(
     await handler.requestHierarchicalRefresh();
 
     runInAction(() => {
-        options.clearProjects();
+        optionsWriter.clearProjects();
         params.projectIds.forEach(pid => {
             const project = handler.projectsData.find(p => p.id === pid);
             if (project) {
-                options.toggleProject(pid, project.name, true);
+                optionsWriter.toggleProject(pid, project.name, true);
             }
             else {
                 console.warn(`Project parameter provided in URL query (${pid}) that could not be matched to underlying data - discarding.`);
@@ -139,11 +279,11 @@ export async function loadBuildHealthOptionsFromParams(
     await handler.getStreamData({ streams: true });
 
     runInAction(() => {
-        options.clearStreams();
+        optionsWriter.clearStreams();
         params.streamIds.forEach(sid => {
             const stream = handler.streamsData.find(s => s.id === sid);
             if (stream) {
-                options.toggleStream(sid, stream.name, true);
+                optionsWriter.toggleStream(sid, stream.name, true);
             }
             else {
                 console.warn(`Stream parameter provided in URL query (${sid}) that could not be matched to underlying data - discarding.`);
@@ -154,13 +294,13 @@ export async function loadBuildHealthOptionsFromParams(
     await handler.getTemplateData({ jobs: true });
 
     runInAction(() => {
-        options.clearTemplates();
+        optionsWriter.clearTemplates();
         const decodedTemplates = params.templateIds.map(templateKey => decodeTemplateKey(templateKey));
 
         decodedTemplates.forEach(tid => {
             const template = handler.templatesData.find(t => t.id === tid.templateId && t.streamId == tid.streamId);
             if (template) {
-                options.toggleTemplate(encodeTemplateKey(template), template.name, true);
+                optionsWriter.toggleTemplate(encodeTemplateKey(template), template.name, true);
             }
             else {
                 console.warn(`Template parameter provided in URL query (${tid}) that could not be matched to underlying data - discarding.`);
@@ -171,20 +311,24 @@ export async function loadBuildHealthOptionsFromParams(
     await handler.getStepData({ steps: true });
 
     runInAction(() => {
-        options.clearSteps();
+        optionsWriter.clearSteps();
         const decodedStepNames = params.stepNames.map(stepNameKey => decodeStepKey(stepNameKey));
+
         decodedStepNames.forEach(stepData => {
             const step = handler.stepData.find(s => s.name.toLowerCase() === stepData.stepName.toLowerCase() && s.streamId == stepData.streamId && s.templateId == stepData.templateId);
             if (step) {
-                options.toggleStep(encodeStepKey(step), step.name, true);
+                optionsWriter.toggleStep(encodeStepKey(step), step.name, true);
+            } else if (stepData.stepName === STEPS_URL_PARAM_ALL_STEPS) {
+                let stepNames = handler.getCachedTemplateStepNames(stepData.streamId, stepData.templateId)
+                optionsWriter.toggleAllSteps(stepData.streamId, stepData.templateId, stepNames, true);
             }
             else {
-                console.warn(`Step name parameter provided in URL query (${name}) that could not be matched to underlying data - discarding.`);
+                console.warn(`Step name parameter provided in URL query (${stepData.stepName}) that could not be matched to underlying data - discarding.`);
             }
         });
-
-        options.synchronizeDerivedKeys();
     });
+
+    options.commitTransactionSession();
 }
 
 // #endregion -- Options Utilities & Data Types --
@@ -201,6 +345,95 @@ export class BuildHealthUIOptionsState {
     }
 }
 
+/**
+ * Interface that describes the difference between a property of type T.
+ */
+export interface FlagDiff<T> {
+    old: T;
+    new: T;
+}
+
+/**
+ * Interface that describes the difference between two option states.
+ */
+export interface BuildHealthOptionStateDiff {
+    added?: BuildHealthOptionStateReceipt;
+    removed?: BuildHealthOptionStateReceipt;
+    enabledPreflight?: FlagDiff<boolean>;
+    timespan?: FlagDiff<TimeSpan>;
+}
+
+/**
+ * A snapshot of the build health options.
+ */
+export class BuildHealthOptionStateReceipt {
+    enabledProjects: ReadonlySet<string>;
+    enabledStreams: ReadonlySet<string>;
+    enabledTemplates: ReadonlySet<string>;
+    enabledSteps: ReadonlySet<string>;
+    enabledPreflight: boolean;
+    timespan: TimeSpan;
+
+    constructor(
+        projects: Iterable<string> = [],
+        streams: Iterable<string> = [],
+        templates: Iterable<string> = [],
+        steps: Iterable<string> = [],
+        enabledPreflight: boolean,
+        timespan: TimeSpan
+    ) {
+        this.enabledProjects = new Set(projects);
+        this.enabledStreams = new Set(streams);
+        this.enabledTemplates = new Set(templates);
+        this.enabledSteps = new Set(steps);
+        this.enabledPreflight = enabledPreflight;
+        this.timespan = timespan;
+    }
+
+    /**
+     * Produces a build health option state diff between two provided receipts.
+     * @param prev The previous option state.
+     * @param next The next option state.
+     * @returns The diff between the two receipts.
+     */
+    static diff(
+        prev: BuildHealthOptionStateReceipt,
+        next: BuildHealthOptionStateReceipt
+    ): BuildHealthOptionStateDiff {
+
+        const added = new BuildHealthOptionStateReceipt(
+            BuildHealthOptionStateReceipt.difference(next.enabledProjects, prev.enabledProjects),
+            BuildHealthOptionStateReceipt.difference(next.enabledStreams, prev.enabledStreams),
+            BuildHealthOptionStateReceipt.difference(next.enabledTemplates, prev.enabledTemplates),
+            BuildHealthOptionStateReceipt.difference(next.enabledSteps, prev.enabledSteps),
+            next.enabledPreflight,
+            next.timespan
+        );
+
+        const removed = new BuildHealthOptionStateReceipt(
+            BuildHealthOptionStateReceipt.difference(prev.enabledProjects, next.enabledProjects),
+            BuildHealthOptionStateReceipt.difference(prev.enabledStreams, next.enabledStreams),
+            BuildHealthOptionStateReceipt.difference(prev.enabledTemplates, next.enabledTemplates),
+            BuildHealthOptionStateReceipt.difference(prev.enabledSteps, next.enabledSteps),
+            prev.enabledPreflight,
+            prev.timespan
+        );
+
+        const enabledPreflight: FlagDiff<boolean> = { old: prev.enabledPreflight, new: next.enabledPreflight };
+        const timespan: FlagDiff<TimeSpan> = { old: prev.timespan, new: next.timespan };
+
+        return { added, removed, enabledPreflight, timespan };
+    }
+
+    private static difference<T>(a: ReadonlySet<T>, b: ReadonlySet<T>): Set<T> {
+        const out = new Set<T>();
+        for (const item of a) {
+            if (!b.has(item)) out.add(item);
+        }
+        return out;
+    }
+}
+
 /**
  * Build Health Filter Options.
  */
@@ -240,400 +473,24 @@ export class BuildHealthOptionsState {
     }
 
     // #endregion -- Constructor --
-}
-
-/**
- * An interface that can provide hierarchy checks on stream, template, and step relationships.
- */
-export interface IHierarchyValidator {
-    /**
-     * Determines whether a stream & template combination is valid (in that it exists).
-     * @param streamId The stream.
-     * @param templateId The template.
-     * @returns True if this stream has a template of the given id, false otherwise.
-     */
-    isValidTemplate(streamId: string, templateId: string): boolean;
-
-/**
-     * Determines whether a stream, template, and step combination is valid (in that it exists).
-     * @param streamId The stream.
-     * @param templateId The template.
-     * @param stepName The step name.
-     * @returns True if this stream has a template of the given id, which has a step of the given name, false otherwise.
-     */
-    isValidStep(streamId: string, templateId: string, stepName: string): boolean;
-}
-
-/**
- * Controller responsible for modifying the underlying BuildHealthOptionsState.
- */
-export class BuildHealthOptionsController {
-    // #region -- Private Members --
-
-    readonly state: BuildHealthOptionsState;
-    readonly uiState: BuildHealthUIOptionsState;
-    readonly handler: BuildHealthDataHandler;
-    private hierarchyValidator?: IHierarchyValidator;
-    private lastSynchronize = -1;
-
-    // #endregion -- Private Members --
-
-    // #region -- Public Members --
-
-    optionsChangeVersion = 0;
-
-    // #endregion -- Public Members --
-
-    // #region -- Constructor --
-
-    constructor(state: BuildHealthOptionsState, uiState: BuildHealthUIOptionsState) {
-        this.state = state;
-        this.uiState = uiState;
-        makeObservable(this);
-    }
-
-    private setOptionsChanged() {
-        this.optionsChangeVersion++;
-    }
-
-    // #endregion -- Constructor --
-
-    // #region -- Private API --
-
-    private toggleTemplateInternal(id: string, name: string, enabled: boolean) {
-        if (enabled) {
-            this.state.enabledTemplates[id] = name;
-        } else {
-            delete this.state.enabledTemplates[id];
-            let { streamId, templateId } = decodeTemplateKey(id);
-            this.clearSteps({ streamId, templateId });
-        }
-    }
-
-    private toggleLinkedTemplates(sourceId: string, name: string, enabled: boolean) {
-        let { streamId, templateId } = decodeTemplateKey(sourceId);
-        Object.keys(this.state.enabledStreams).forEach(streamKey => {
-            if (streamKey !== streamId) {
-                let linkedTempalteId = encodeTemplateKeyFromStrings(streamKey, templateId);
-                if (this.hierarchyValidator?.isValidTemplate(streamKey, templateId)) { // we need to see if this is a valid pair?
-                    this.toggleTemplateInternal(linkedTempalteId, name, enabled);
-                }
-            }
-        });
-    }
-
-    private toggleStepInternal(id: string, name: string, enabled: boolean) {
-        if (enabled) {
-            this.state.enabledSteps[id] = name;
-        } else {
-            delete this.state.enabledSteps[id];
-        }
-    }
-
-    private toggleLinkedSteps(sourceId: string, name: string, enabled: boolean) {
-        let sourceDecodedStepKey = decodeStepKey(sourceId);
-        Object.keys(this.state.enabledTemplates).forEach(key => {
-            let targetDecodedTempalteKey = decodeTemplateKey(key);
-            // Set this step on every other enabled template, except the narrow case of the stream-template-step - we've just come from there.
-            if (targetDecodedTempalteKey.templateId !== sourceDecodedStepKey.templateId || targetDecodedTempalteKey.streamId !== sourceDecodedStepKey.streamId) {
-                let linkedEncodedStepId = encodeStepNameFromStringsPreserveCase(targetDecodedTempalteKey.streamId, targetDecodedTempalteKey.templateId, sourceDecodedStepKey.stepName);
-                if (this.hierarchyValidator?.isValidStep(targetDecodedTempalteKey.streamId, targetDecodedTempalteKey.templateId, sourceDecodedStepKey.stepName)) { // we need to see if this is a valid pair?
-                    this.toggleStepInternal(linkedEncodedStepId, name, enabled);
-                }
-            }
-        });
-    }
-
-    // #endregion -- Private API --
 
     // #region -- Public API --
 
     /**
-     * Sets the hierarchy validator for the controller.
-     * @param validator The hierarchy valildator to us.
-     */
-    setHierarchyValidator(validator: IHierarchyValidator) {
-        this.hierarchyValidator = validator;
-    }
-
-    /**
-     * Toggles the state of the provided project. Will clear all other set projects.
-     * @param id The id of the project to toggle.
-     * @param name The name of the project.
-     * @param enabled Whether it's enabled or not.
-     */
-    @action
-    toggleSingleProject(id: string, name: string) {
-        this.state.enabledProjects = { [id]: name };
-        this.clearStreams(true);
-        this.setOptionsChanged();
-    }
-
-    /**
-     * Clears all currently selected options.
-     */
-    @action
-    clearAll() {
-        Object.keys(this.state.enabledProjects).forEach(key => {
-            delete this.state.enabledProjects[key];
-        });
-        this.clearStreams(true);
-        this.setOptionsChanged();
-    }
-
-    /**
-     * Toggles the state of the provided project.
-     * @param id The id of the project to toggle.
-     * @param name The name of the project.
-     * @param enabled Whether it's enabled or not.
-     */
-    @action
-    toggleProject(id: string, name: string, enabled: boolean) {
-        if (enabled) {
-            this.state.enabledProjects[id] = name;
-        } else {
-            delete this.state.enabledProjects[id];
-            this.clearStreams(true);
-        }
-        this.setOptionsChanged();
-    }
-
-    /**
-     * Toggles the state of the provided stream. Will clear all other set streams.
-     * @param id The id of the stream to toggle.
-     * @param name The name of the stream.
-     * @param enabled Whether it's enabled or not.
-     */
-    @action
-    toggleSingleStream(id: string, name: string) {
-        this.clearStreams(true);
-        this.state.enabledStreams[id] = name;
-        this.setOptionsChanged();
-    }
-
-    /**
-     * Toggles the state of the provided stream.
-     * @param id The id of the stream to toggle.
-     * @param name The name of the stream.
-     * @param enabled Whether it's enabled or not.
-     */
-    @action
-    toggleStream(id: string, name: string, enabled: boolean) {
-        if (enabled) {
-            this.state.enabledStreams[id] = name;
-        } else {
-            delete this.state.enabledStreams[id];
-            this.clearTemplates(true, id);
-        }
-        this.setOptionsChanged();
-    }
-
-    /**
-     * Toggles the state of the provided template.
-     * @param id The id of the template to toggle.
-     * @param name The name of the template.
-     * @param enabled Whether it's enabled or not.
-     */
-    @action
-    toggleTemplate(id: string, name: string, enabled: boolean) {
-        this.toggleTemplateInternal(id, name, enabled);
-        if (this.state.linkedSelectionModeEnabled) {
-            this.toggleLinkedTemplates(id, name, enabled);
-        }
-        this.setOptionsChanged();
-    }
-
-    /**
-     * Toggles all linkedable templates, for all active streams.
-     * @param enabled True if we want to enable, false otherwise.
-     */
-    @action
-    toggleAllLinkableTemplates(enabled: boolean) {
-        if (this.state.linkedSelectionModeEnabled) {
-            Object.keys(this.state.enabledTemplates).forEach(templateKey => {
-                let templateName = this.state.enabledTemplates[templateKey];
-                this.toggleLinkedTemplates(templateKey, templateName, enabled);
-            });
-            this.setOptionsChanged();
-        }
-    }
-
-    /**
-     * Toggles all linkedable steps, for all active templates.
-     * @param enabled True if we want to enable, false otherwise.
-     */
-    @action
-    toggleAllLinkableSteps(enabled: boolean) {
-        if (this.state.linkedSelectionModeEnabled) {
-            Object.keys(this.state.enabledSteps).forEach(stepKey => {
-                let stepName = this.state.enabledSteps[stepKey];
-                this.toggleLinkedSteps(stepKey, stepName, enabled);
-            });
-            this.setOptionsChanged();
-        }
-    }
-
-    /**
-     * Toggles the state of the provided step.
-     * @param id The id of the step to toggle.
-     * @param name The name of the step.
-     * @param enabled Whether it's enabled or not.
-     */
-    @action
-    toggleStep(id: string, name: string, enabled: boolean) {
-        this.toggleStepInternal(id, name, enabled);
-        if (this.state.linkedSelectionModeEnabled) {
-            this.toggleLinkedSteps(id, name, enabled);
-        }
-        this.setOptionsChanged();
-    }
-
-    /**
-     * Sets the current job history time span to utilize.
-     * @param time The time span to utilize.
-     */
-    @action
-    setJobHistoryTimeSpan(time: TimeSpan) {
-        this.state.jobHistoryTimeSpan = time;
-        this.setOptionsChanged();
-    }
-
-    /**
-     * Sets whether to filter out preflight jobs, or not.
-     * @param includePreflight True to include preflight job instances, false otherwise.
-     */
-    @action
-    setIncludePreflight(includePreflight: boolean) {
-        this.state.includePreflight = includePreflight;
-        this.setOptionsChanged();
-    }
-
-    /**
-     * Sets whether to enter into linked selection mode, or not.
-     * @param linkedSelectionModeEnabled True to enter linked selection mode, false otherwise.
-     */
-    @action
-    setLinkedSelectionMode(linkedSelectionModeEnabled: boolean) {
-        this.state.linkedSelectionModeEnabled = linkedSelectionModeEnabled;
-        this.setOptionsChanged();
-    }
-
-    /**
-     * Sets whether the date anchors should be used when visualizing the table.
-     * @param includeDateAnchors 
-     */
-    @action
-    setHideDateAnchors(includeDateAnchors: boolean) {
-        this.uiState.includeDateAnchors = includeDateAnchors;
-    }
-
-    /**
- * Sets whether summary rows should consider warnings as a success, or failure.
- * @param warningsAsSummaryFailure False to include warnings in the numerator of success ratio, true otherwise.
- */
-    @action
-    setWarningsAsSummaryFailure(warningsAsSummaryFailure: boolean) {
-        this.uiState.warningsAsSummaryFailure = warningsAsSummaryFailure;
-    }
-
-    /**
-      * Clear all the projects tracked by the options.
-      * @param cascade Whether to cascade this down to dependent options.
-      */
-    clearProjects(cascade?: boolean) {
-        Object.keys(this.state.enabledProjects).forEach(key => {
-            delete this.state.enabledProjects[key];
-        });
-
-        if (cascade) {
-            this.clearStreams(cascade);
-        }
-    }
-
-    /**
-      * Clear all the streams tracked by the options.
-      * @param cascade Whether to cascade this down to dependent options.
-      */
-    clearStreams(cascade?: boolean) {
-        Object.keys(this.state.enabledStreams).forEach(key => {
-            delete this.state.enabledStreams[key];
-
-            if (cascade) {
-                this.clearTemplates(cascade, key);
-            }
-        });
-    }
-
-    /**
-      * Clear all the templates tracked by the options.
-      * @param cascade Whether to cascade this down to dependent options.
-      * @param owningStreamId Whether to clear an item if it belongs to a stream. If not provided, will remove any template.
-      */
-    clearTemplates(cascade?: boolean, owningStreamId?: string) {
-        Object.keys(this.state.enabledTemplates).forEach(key => {
-            let { streamId, templateId } = decodeTemplateKey(key);
-            if (!owningStreamId || (streamId === owningStreamId)) {
-                delete this.state.enabledTemplates[key];
-
-                if (cascade) {
-                    this.clearSteps(owningStreamId ? { streamId: owningStreamId!, templateId: templateId } : undefined);
-                }
-            }
-        });
-    }
-
-    /**
-     * Clear all the steps tracked by the options.
-     * @param owningHierarchy Whether to clear an item if it belongs to a stream & template. If not provided, will remove any step.
-     */
-    clearSteps(owningHierarchy?: { streamId: string; templateId: string }) {
-        Object.keys(this.state.enabledSteps).forEach(key => {
-            let { streamId, templateId, stepName } = decodeStepKey(key);
-            if ((!owningHierarchy) || (streamId === owningHierarchy.streamId && templateId === owningHierarchy.templateId)) {
-                delete this.state.enabledSteps[key];
-            }
-        });
-    }
-
-    /**
-     * Sycnrhonizes the options that are bound to UI elements, with the data observed by consumers of the option set.
-     * This is important to keep separate as it allows us to control when we signal/flush a finalized set of options to the consumers.
-     */
-    @action
-    synchronizeDerivedKeys() {
-        if (this.lastSynchronize < this.optionsChangeVersion) {
-            this.state.stepOutcomeEnabledStreamKeys = Object.keys(this.state.enabledStreams);
-            this.state.stepOutcomeEnabledJobKeys = Object.keys(this.state.enabledTemplates);
-            this.state.stepOutcomeEnabledStepKeys = Object.keys(this.state.enabledSteps);
-            this.state.startDate = new Date(new Date().valueOf() - (this.state.jobHistoryTimeSpan.minutes * 60000));
-            this.state.endDate = new Date();
-
-            this.lastSynchronize = this.optionsChangeVersion;
-        }
-    }
-
-    /**
-     * Obtains a URLSearchParams representation of the current options.
-     * @returns The URLSearchParams representation of the current options.
+     * Generates a state receipt for the options in it's current state.
+     * @returns A receipt representation of the current build health options.
      */
-    toNavigationQuery(): URLSearchParams {
-        const params = new URLSearchParams(location.search);
-
-        const projects = Object.keys(this.state.enabledProjects).join(",");
-        const streams = this.state.stepOutcomeEnabledStreamKeys.join(",");
-        const jobs = this.state.stepOutcomeEnabledJobKeys.join(",");
-        const steps = this.state.stepOutcomeEnabledStepKeys.join(",");
-        const lastTimeRange = this.state.jobHistoryTimeSpan.key.toString();
-        const includePreflight = this.state.includePreflight;
-
-        params.set(PROJECTS_URL_PARAM_NAME, LZString.compressToEncodedURIComponent(projects));
-        params.set(STREAMS_URL_PARAM_NAME, LZString.compressToEncodedURIComponent(streams));
-        params.set(TEMPLATES_URL_PARAM_NAME, LZString.compressToEncodedURIComponent(jobs));
-        params.set(STEPS_URL_PARAM_NAME, LZString.compressToEncodedURIComponent(steps));
-        params.set(TIME_SPAN_URL_PARAM, LZString.compressToEncodedURIComponent(lastTimeRange));
-        params.set(INCLUDE_PREFLIGHT_PARAM, LZString.compressToEncodedURIComponent(String(includePreflight)));
-
-        return params;
+    generateStateReceipt(): BuildHealthOptionStateReceipt {
+        let receipt: BuildHealthOptionStateReceipt = new BuildHealthOptionStateReceipt(
+            Object.keys(this.enabledProjects),
+            Object.keys(this.enabledStreams),
+            Object.keys(this.enabledTemplates),
+            Object.keys(this.enabledSteps),
+            this.includePreflight,
+            this.jobHistoryTimeSpan
+        );
+
+        return receipt;
     }
 
     // #endregion -- Public API --
diff --git a/Engine/Source/Programs/Horde/HordeDashboard/src/components/buildhealth/BuildHealthOptionsController.ts b/Engine/Source/Programs/Horde/HordeDashboard/src/components/buildhealth/BuildHealthOptionsController.ts
new file mode 100644
--- /dev/null
+++ b/Engine/Source/Programs/Horde/HordeDashboard/src/components/buildhealth/BuildHealthOptionsController.ts
@@ -0,0 +1,547 @@
+import { action, makeObservable } from "mobx";
+import { BuildHealthOptionsState, BuildHealthOptionStateDiff, BuildHealthOptionStateReceipt, BuildHealthUIOptionsState, TimeSpan } from "./BuildHealthOptions";
+import { decodeStepKey, decodeTemplateKey, encodeStepNameFromStringsPreserveCase, encodeTemplateKeyFromStrings } from "./BuildHealthUtilities";
+import { BuildHealthDataHandler } from "./BuildHealthDataHandler";
+
+/**
+ * An interface that can provide hierarchy checks on stream, template, and step relationships.
+ */
+export interface IHierarchyValidator {
+    /**
+     * Determines whether a stream & template combination is valid (in that it exists).
+     * @param streamId The stream.
+     * @param templateId The template.
+     * @returns True if this stream has a template of the given id, false otherwise.
+     */
+    isValidTemplate(streamId: string, templateId: string): boolean;
+
+    /**
+         * Determines whether a stream, template, and step combination is valid (in that it exists).
+         * @param streamId The stream.
+         * @param templateId The template.
+         * @param stepName The step name.
+         * @returns True if this stream has a template of the given id, which has a step of the given name, false otherwise.
+         */
+    isValidStep(streamId: string, templateId: string, stepName: string): boolean;
+}
+
+/**
+ * Interface that describes all the writeable operations on the build health options.
+ */
+export interface IBuildHealthOptionWriter {
+
+    /**
+     * Produces the set of additions and removals for all options.
+     */
+    produceReceiptDiff(): BuildHealthOptionStateDiff;
+
+    /**
+     * Sets whether to filter out preflight jobs, or not.
+     * @param includePreflight True to include preflight job instances, false otherwise.
+     */
+    setIncludePreflight(includePreflight: boolean): void;
+
+    /**
+     * Sets the current job history time span to utilize.
+     * @param time The time span to utilize.
+     */
+    setJobHistoryTimeSpan(time: TimeSpan): void;
+
+    /**
+     * Sets whether to enter into linked selection mode, or not.
+     * @param linkedSelectionModeEnabled True to enter linked selection mode, false otherwise.
+     */
+    setLinkedSelectionMode(linkedSelectionModeEnabled: boolean): void;
+
+    /**
+     * Toggles the state of the provided project.
+     * @param id The id of the project to toggle.
+     * @param name The name of the project.
+     * @param enabled Whether it's enabled or not.
+     */
+    toggleProject(id: string, name: string, enabled: boolean): void;
+
+    /**
+     * Toggles the state of the provided project. Will clear all other set projects.
+     * @param id The id of the project to toggle.
+     * @param name The name of the project.
+     * @param enabled Whether it's enabled or not.
+     */
+    toggleSingleProject(id: string, name: string): void;
+
+    /**
+     * Toggles the state of the provided stream.
+     * @param id The id of the stream to toggle.
+     * @param name The name of the stream.
+     * @param enabled Whether it's enabled or not.
+     */
+    toggleStream(id: string, name: string, enabled: boolean): void;
+
+    /**
+     * Toggles the state of the provided stream. Will clear all other set streams.
+     * @param id The id of the stream to toggle.
+     * @param name The name of the stream.
+     * @param enabled Whether it's enabled or not.
+     */
+    toggleSingleStream(id: string, name: string): void;
+
+    /**
+     * Toggles all linkedable templates, for all active streams.
+     * @param enabled True if we want to enable, false otherwise.
+     */
+    toggleAllLinkableTemplates(enabled: boolean): void;
+
+    /**
+     * Toggles the state of the provided template.
+     * @param id The id of the template to toggle.
+     * @param name The name of the template.
+     * @param enabled Whether it's enabled or not.
+     */
+    toggleTemplate(id: string, name: string, enabled: boolean): void;
+
+    /**
+     * Toggles all linkedable steps, for all active templates.
+     * @param enabled True if we want to enable, false otherwise.
+     */
+    toggleAllLinkableSteps(enabled: boolean): void;
+
+    /**
+     * Toggles all provided steps for a given stream & template, given they are valid.
+     * @param streamId The stream.
+     * @param templateId The template.
+     * @param stepNames The list of steps to toggle.
+     * @param enabled True to enable, false otherwise.
+     * @param cascadeLinkedAction True to enable linked action propagation (default behaviour), false or otherwise.
+     */
+    toggleAllSteps(streamId: string, templateId: string, stepNames: string[], enabled: boolean, cascadeLinkedAction?: boolean): void;
+
+    /**
+     * Toggles the state of the provided step.
+     * @param id The id of the step to toggle.
+     * @param name The name of the step.
+     * @param enabled Whether it's enabled or not.
+     */
+    toggleStep(id: string, name: string, enabled: boolean): void;
+
+    /**
+     * Clears all currently selected options.
+     */
+    clearAll(): void;
+
+    /**
+      * Clear all the projects tracked by the options.
+      * @param cascade Whether to cascade this down to dependent options.
+      */
+    clearProjects(cascade?: boolean): void;
+
+    /**
+      * Clear all the streams tracked by the options.
+      * @param cascade Whether to cascade this down to dependent options.
+      */
+    clearStreams(cascade?: boolean): void;
+
+    /**
+      * Clear all the templates tracked by the options.
+      * @param cascade Whether to cascade this down to dependent options.
+      * @param owningStreamId Whether to clear an item if it belongs to a stream. If not provided, will remove any template.
+      */
+    clearTemplates(cascade?: boolean, owningStreamId?: string): void;
+
+    /**
+     * Clear all the steps tracked by the options.
+     * @param owningHierarchy Whether to clear an item if it belongs to a stream & template. If not provided, will remove any step.
+     */
+    clearSteps(owningHierarchy?: { streamId: string; templateId: string }): void;
+}
+
+export class BuildHealthOptionsWriter implements IBuildHealthOptionWriter {
+    readonly state: BuildHealthOptionsState;
+    readonly hierarchyValidator?: IHierarchyValidator;
+    private activeTransactionReceipt: BuildHealthOptionStateReceipt;
+
+    constructor(state: BuildHealthOptionsState, hierarhcyValidator?: IHierarchyValidator) {
+        makeObservable(this);
+        this.state = state;
+        this.activeTransactionReceipt = this.state.generateStateReceipt();
+        this.hierarchyValidator = hierarhcyValidator;
+    }
+
+    /**
+     * Produces the current receipt diff between the receipt from the start of the session, and the current receipt.
+     * @returns The diff receipts: added items & deleted items.
+     */
+    produceReceiptDiff(): BuildHealthOptionStateDiff {
+        let transactionStartReceipt = this.state.generateStateReceipt();
+        let returnSet = BuildHealthOptionStateReceipt.diff(this.activeTransactionReceipt, transactionStartReceipt);
+        return returnSet;
+    }
+
+    @action
+    clearAll() {
+        Object.keys(this.state.enabledProjects).forEach(key => {
+            delete this.state.enabledProjects[key];
+        });
+        this.clearStreams(true);
+    }
+
+    /**
+     * @inheritdoc
+     */
+    @action
+    toggleSingleProject(id: string, name: string) {
+        this.state.enabledProjects = { [id]: name };
+        this.clearStreams(true);
+    }
+
+    /**
+     * @inheritdoc
+     */
+    @action
+    toggleProject(id: string, name: string, enabled: boolean) {
+        if (enabled) {
+            this.state.enabledProjects[id] = name;
+        } else {
+            delete this.state.enabledProjects[id];
+            this.clearStreams(true);
+        }
+    }
+
+    /**
+     * @inheritdoc
+     */
+    @action
+    toggleSingleStream(id: string, name: string) {
+        this.clearStreams(true);
+        this.state.enabledStreams[id] = name;
+    }
+
+    /**
+     * @inheritdoc
+     */
+    @action
+    toggleStream(id: string, name: string, enabled: boolean) {
+        if (enabled) {
+            this.state.enabledStreams[id] = name;
+        } else {
+            delete this.state.enabledStreams[id];
+            this.clearTemplates(true, id);
+        }
+    }
+
+    /**
+     * @inheritdoc
+     */
+    @action
+    toggleTemplate(id: string, name: string, enabled: boolean) {
+        this.toggleTemplateInternal(id, name, enabled);
+        if (this.state.linkedSelectionModeEnabled) {
+            this.toggleLinkedTemplates(id, name, enabled);
+        }
+    }
+
+    /**
+     * @inheritdoc
+     */
+    @action
+    toggleAllLinkableTemplates(enabled: boolean) {
+        if (this.state.linkedSelectionModeEnabled) {
+            Object.keys(this.state.enabledTemplates).forEach(templateKey => {
+                let templateName = this.state.enabledTemplates[templateKey];
+                this.toggleLinkedTemplates(templateKey, templateName, enabled);
+            });
+        }
+    }
+
+    /**
+     * @inheritdoc
+     */
+    @action
+    toggleAllLinkableSteps(enabled: boolean) {
+        if (this.state.linkedSelectionModeEnabled) {
+            Object.keys(this.state.enabledSteps).forEach(stepKey => {
+                let stepName = this.state.enabledSteps[stepKey];
+                this.toggleLinkedSteps(stepKey, stepName, enabled);
+            });
+        }
+    }
+
+    /**
+     * @inheritdoc
+     */
+    @action
+    toggleStep(id: string, name: string, enabled: boolean) {
+        this.toggleStepInternal(id, name, enabled);
+        if (this.state.linkedSelectionModeEnabled) {
+            this.toggleLinkedSteps(id, name, enabled);
+        }
+    }
+
+    /**
+     * @inheritdoc
+     */
+    @action
+    toggleAllSteps(streamId: string, templateId: string, stepNames: string[], enabled: boolean, cascadeLinkedAction: boolean = true) {
+        for (let idx: number = 0; idx < stepNames.length; ++idx) {
+            let stepName = stepNames[idx];
+            let toggleAllStepId = encodeStepNameFromStringsPreserveCase(streamId, templateId, stepName);
+            if (cascadeLinkedAction) {
+                this.toggleStep(toggleAllStepId, stepName, enabled);
+            } else {
+                this.toggleStepInternal(toggleAllStepId, stepName, enabled);
+            }
+        }
+    }
+
+    /**
+     * @inheritdoc
+     */
+    @action
+    setJobHistoryTimeSpan(time: TimeSpan) {
+        this.state.jobHistoryTimeSpan = time;
+    }
+
+    /**
+     * @inheritdoc
+     */
+    @action
+    setIncludePreflight(includePreflight: boolean) {
+        this.state.includePreflight = includePreflight;
+    }
+
+    /**
+     * @inheritdoc
+     */
+    @action
+    setLinkedSelectionMode(linkedSelectionModeEnabled: boolean) {
+        this.state.linkedSelectionModeEnabled = linkedSelectionModeEnabled;
+    }
+
+    /**
+     * @inheritdoc
+     */
+    clearProjects(cascade?: boolean) {
+        Object.keys(this.state.enabledProjects).forEach(key => {
+            delete this.state.enabledProjects[key];
+        });
+
+        if (cascade) {
+            this.clearStreams(cascade);
+        }
+    }
+
+    /**
+     * @inheritdoc
+     */
+    clearStreams(cascade?: boolean) {
+        Object.keys(this.state.enabledStreams).forEach(key => {
+            delete this.state.enabledStreams[key];
+
+            if (cascade) {
+                this.clearTemplates(cascade, key);
+            }
+        });
+    }
+
+    /**
+     * @inheritdoc
+     */
+    clearTemplates(cascade?: boolean, owningStreamId?: string) {
+        Object.keys(this.state.enabledTemplates).forEach(key => {
+            let { streamId, templateId } = decodeTemplateKey(key);
+            if (!owningStreamId || (streamId === owningStreamId)) {
+                delete this.state.enabledTemplates[key];
+
+                if (cascade) {
+                    this.clearSteps(owningStreamId ? { streamId: owningStreamId!, templateId: templateId } : undefined);
+                }
+            }
+        });
+    }
+
+    /**
+     * @inheritdoc
+     */
+    clearSteps(owningHierarchy?: { streamId: string; templateId: string }) {
+        Object.keys(this.state.enabledSteps).forEach(key => {
+            let { streamId, templateId, stepName } = decodeStepKey(key);
+            if ((!owningHierarchy) || (streamId === owningHierarchy.streamId && templateId === owningHierarchy.templateId)) {
+                delete this.state.enabledSteps[key];
+            }
+        });
+    }
+
+    // #region -- Private API --
+
+    private toggleTemplateInternal(id: string, name: string, enabled: boolean) {
+        if (enabled) {
+            this.state.enabledTemplates[id] = name;
+        } else {
+            delete this.state.enabledTemplates[id];
+            let { streamId, templateId } = decodeTemplateKey(id);
+            this.clearSteps({ streamId, templateId });
+        }
+    }
+
+    private toggleLinkedTemplates(sourceId: string, name: string, enabled: boolean) {
+        let { streamId, templateId } = decodeTemplateKey(sourceId);
+        Object.keys(this.state.enabledStreams).forEach(streamKey => {
+            if (streamKey !== streamId) {
+                let linkedTempalteId = encodeTemplateKeyFromStrings(streamKey, templateId);
+                if (this.hierarchyValidator?.isValidTemplate(streamKey, templateId)) { // we need to see if this is a valid pair?
+                    this.toggleTemplateInternal(linkedTempalteId, name, enabled);
+                }
+            }
+        });
+    }
+
+    private toggleStepInternal(id: string, name: string, enabled: boolean) {
+        if (enabled) {
+            this.state.enabledSteps[id] = name;
+        } else {
+            delete this.state.enabledSteps[id];
+        }
+    }
+
+    private toggleLinkedSteps(sourceId: string, name: string, enabled: boolean) {
+        let sourceDecodedStepKey = decodeStepKey(sourceId);
+        Object.keys(this.state.enabledTemplates).forEach(key => {
+            let targetDecodedTempalteKey = decodeTemplateKey(key);
+            // Set this step on every other enabled template, except the narrow case of the stream-template-step - we've just come from there.
+            if (targetDecodedTempalteKey.templateId !== sourceDecodedStepKey.templateId || targetDecodedTempalteKey.streamId !== sourceDecodedStepKey.streamId) {
+                let linkedEncodedStepId = encodeStepNameFromStringsPreserveCase(targetDecodedTempalteKey.streamId, targetDecodedTempalteKey.templateId, sourceDecodedStepKey.stepName);
+                if (this.hierarchyValidator?.isValidStep(targetDecodedTempalteKey.streamId, targetDecodedTempalteKey.templateId, sourceDecodedStepKey.stepName)) { // we need to see if this is a valid pair?
+                    this.toggleStepInternal(linkedEncodedStepId, name, enabled);
+                }
+            }
+        });
+    }
+
+    // #endregion -- Private API --
+}
+
+/**
+ * Controller responsible for modifying the underlying BuildHealthOptionsState.
+ */
+export class BuildHealthOptionsController {
+    // #region -- Private Members --
+
+    readonly state: BuildHealthOptionsState;
+    readonly uiState: BuildHealthUIOptionsState;
+    readonly handler: BuildHealthDataHandler;
+
+    private writer: IBuildHealthOptionWriter | null;
+    private hierarchyValidator?: IHierarchyValidator;
+    private lastSynchronize = -1;
+
+    // #endregion -- Private Members --
+
+    // #region -- Public Members --
+
+    optionsChangeVersion = 0;
+    querySchemaVersion: number = 2;
+
+    // #endregion -- Public Members --
+
+    // #region -- Constructor --
+
+    constructor(state: BuildHealthOptionsState, uiState: BuildHealthUIOptionsState) {
+        this.state = state;
+        this.uiState = uiState;
+        makeObservable(this);
+    }
+
+    private setOptionsChanged() {
+        this.optionsChangeVersion++;
+    }
+
+    // #endregion -- Constructor --
+
+    // #region -- Public API --
+
+    /**
+     * Sets the hierarchy validator for the controller.
+     * @param validator The hierarchy valildator to us.
+     */
+    setHierarchyValidator(validator: IHierarchyValidator) {
+        this.hierarchyValidator = validator;
+    }
+
+    /**
+     * Sets whether the date anchors should be used when visualizing the table.
+     * @param includeDateAnchors 
+     */
+    @action
+    setHideDateAnchors(includeDateAnchors: boolean) {
+        this.uiState.includeDateAnchors = includeDateAnchors;
+    }
+
+    /**
+     * Sets whether summary rows should consider warnings as a success, or failure.
+     * @param warningsAsSummaryFailure False to include warnings in the numerator of success ratio, true otherwise.
+     */
+    @action
+    setWarningsAsSummaryFailure(warningsAsSummaryFailure: boolean) {
+        this.uiState.warningsAsSummaryFailure = warningsAsSummaryFailure;
+    }
+
+    /**
+     * Gets a options transaction session to make changes to.
+     * @returns The @see IBuildHealthOptionWriter that will enact options changes.
+     */
+    @action
+    getTransactionSession(): IBuildHealthOptionWriter {
+        if (!this.writer) {
+            this.writer = new BuildHealthOptionsWriter(this.state, this.hierarchyValidator);
+        }
+
+        return this.writer;
+    }
+
+    /**
+     * Commits the current options transaction session if one is active.
+     * @returns The receipts for the current transaction.
+     */
+    @action
+    commitTransactionSession(): BuildHealthOptionStateDiff {
+        if (this.writer === null) {
+            return {
+                added: undefined,
+                removed: undefined,
+                enabledPreflight: undefined,
+                timespan: undefined
+            };
+        }
+
+        let finalReceiptes = this.writer.produceReceiptDiff();
+        this.writer = null;
+
+        // We can be smarter here, and only only request synchronization if we need to (by checking the option state diff).
+        this.setOptionsChanged();
+        this.synchronizeDerivedKeys();
+
+        return finalReceiptes;
+    }
+
+    // #endregion -- Public API --
+
+    // #region -- Private API -- 
+
+    /**
+     * Sycnrhonizes the options that are bound to UI elements, with the data observed by consumers of the option set.
+     * This is important to keep separate as it allows us to control when we signal/flush a finalized set of options to the consumers.
+     */
+    @action
+    private synchronizeDerivedKeys() {
+        if (this.lastSynchronize < this.optionsChangeVersion) {
+            this.state.stepOutcomeEnabledStreamKeys = Object.keys(this.state.enabledStreams);
+            this.state.stepOutcomeEnabledJobKeys = Object.keys(this.state.enabledTemplates);
+            this.state.stepOutcomeEnabledStepKeys = Object.keys(this.state.enabledSteps);
+            this.state.startDate = new Date(new Date().valueOf() - (this.state.jobHistoryTimeSpan.minutes * 60000));
+            this.state.endDate = new Date();
+
+            this.lastSynchronize = this.optionsChangeVersion;
+        }
+    }
+
+    // #endregion -- Private API --
+}
\ No newline at end of file
diff --git a/Engine/Source/Programs/Horde/HordeDashboard/src/components/buildhealth/BuildhealthView.tsx b/Engine/Source/Programs/Horde/HordeDashboard/src/components/buildhealth/BuildhealthView.tsx
--- a/Engine/Source/Programs/Horde/HordeDashboard/src/components/buildhealth/BuildhealthView.tsx
+++ b/Engine/Source/Programs/Horde/HordeDashboard/src/components/buildhealth/BuildhealthView.tsx
@@ -1,9 +1,9 @@
 // Copyright Epic Games, Inc. All Rights Reserved.
 
-import { Checkbox, ComboBox, DirectionalHint, Dropdown, DropdownMenuItemType, IconButton, IDropdownOption, IDropdownProps, ISelectableOption, Label, PrimaryButton, SearchBox, Stack, TextField, TooltipHost } from "@fluentui/react";
+import { Checkbox, ComboBox, DefaultButton, DirectionalHint, Dropdown, DropdownMenuItemType, IconButton, IDropdownOption, IDropdownProps, ISelectableOption, Label, mergeStyleSets, PrimaryButton, SearchBox, Stack, TooltipHost } from "@fluentui/react";
 import { GetProjectResponse, GetStreamResponse } from "horde/backend/Api";
 import { observer } from "mobx-react-lite";
-import { useCallback, useEffect, useMemo, useState } from "react";
+import { useCallback, useEffect, useMemo, useRef, useState } from "react";
 import { useLocation, useNavigate } from "react-router";
 import { useWindowSize } from "../../base/utilities/hooks";
 import { getHordeStyling } from "../../styles/Styles";
@@ -12,14 +12,34 @@ import { TopNav } from "../TopNav";
 import { StepOutcomeView } from "../buildhealth/stepoutcome/StepOutcome";
 import { BuildHealthDataHandler, DataHandlerRefreshRequest, expandRefreshRequest } from "./BuildHealthDataHandler";
 import { KEY_SEPARATOR, PARAMETER_KEY_PREFIX, StepRefData, TemplateRefData } from "./BuildHealthDataTypes";
-import { BuildHealthOptionsController, BuildHealthOptionsState, BuildHealthUIOptionsState, JobHistoryTimeSpans, loadBuildHealthOptionsFromParams, parseBuildHealthQueryParams, TimeSpan } from "./BuildHealthOptions";
-import { encodeStepKey, encodeTemplateKey } from "./BuildHealthUtilities";
+import { BuildHealthOptionsState, BuildHealthOptionStateDiff, BuildHealthUIOptionsState, JobHistoryTimeSpans, loadBuildHealthOptionsFromParams, parseBuildHealthQueryParams, TimeSpan, toNavigationQuery } from "./BuildHealthOptions";
+import { decodeStepKey, decodeTemplateKey, encodeStepKey, encodeTemplateKey } from "./BuildHealthUtilities";
 import { stepOutcomeTableClasses } from "./stepoutcome/StepOutcomeSharedUIComponents";
 import { runInAction } from "mobx";
 import { StepOutcomeLegend } from "./stepoutcome/StepOutcomeLegend";
+import { BuildHealthOptionsController } from "./BuildHealthOptionsController";
 
 // #region -- Dropdown Components --
 
+// #region -- Styles --
+
+const dropdownNavigationButtons = mergeStyleSets({
+    base: {
+        marginLeft: 4,
+        padding: "2px 6px",
+        fontSize: 12,
+        cursor: "pointer",
+    }
+});
+
+const dropdownOptions = mergeStyleSets({
+    base: {
+        padding: "4px 8px",
+    }
+});
+
+// #endregion -- Styles --
+
 // #region -- Dropdown Data Types --
 
 type BuildHealthOptionData = {
@@ -34,58 +54,171 @@ type BuildHealthOptionList = {
     toolTip: string
 }
 
+const SELECT_CLEAR_ALL_SPECIAL_GROUP: string = "select_clear_all";
+
 // #endregion -- Dropdown Data Types --
 
 // #region -- Base Dropdown Components --
 
+/**
+ * Props for the SearchableDropdown component.
+ */
 interface ISearchableDropdownProps extends IDropdownProps {
+
+    // Delegate that is run when search text field is changed.
     onSearchValueChanged?: (searchValue: string) => void;
+
+    // Refs to all groups in the dropdown.
+    groupRefs?: React.RefObject<(HTMLDivElement | null)[]>;
 }
 
+/**
+ * Searchable drop down component.
+ * @param props The props that control the component.
+ * @returns The React component.
+ */
 const SearchableDropdown: React.FC<ISearchableDropdownProps> = (props) => {
-    const { onSearchValueChanged, ...dropdownProps } = props;
-
-    // Add dummy search header
-    const getOptions = (): IDropdownOption[] => [
-        { key: 'search', text: '', itemType: 2 },
-        ...dropdownProps.options,
-    ];
-
-    const onRenderOption = (
-        option?: ISelectableOption,
-        defaultRender?: (props?: ISelectableOption) => JSX.Element | null,
-    ): JSX.Element | null => {
-        if (!option) return null;
-
-        // Render SearchBox for dummy header
-        if (option.key === 'search') {
-            return (
-                <SearchBox
-                    placeholder="Filter options"
-                    autoFocus
-                    onChange={(ev, newValue) => onSearchValueChanged?.(newValue || '')}
-                    onClear={() => onSearchValueChanged?.('')}
-                    styles={{ root: { margin: '8px', width: 'calc(100% - 16px)' } }}
-                    onKeyDown={(ev) => {
-                        if (ev.key === 'Enter' || ev.key === 'Escape') ev.stopPropagation();
-                    }}
-                />
-            );
+    const { modeColors } = getHordeStyling();
+    const { onSearchValueChanged, groupRefs, ...dropdownProps } = props;
+    const dropdownRef = useRef<HTMLDivElement | null>(null);
+    const currentIndex = useRef(0);
+    const PIXEL_AFFORDANCE: number = 5;
+
+    // #region -- Internal Dropdown Navigation Functions --
+
+    const recomputeCurrentIndex = (): { currentGroup: number, onHeader: boolean } => {
+        if (!dropdownRef.current || !groupRefs?.current) return { currentGroup: 0, onHeader: false };
+
+        const scrollTop = dropdownRef.current.scrollTop;
+        const keys = Object.keys(groupRefs.current || []);
+
+        let newIndex: number = 0;
+        let onHeader: boolean = false;
+
+        for (let i = 0; i < keys.length; i++) {
+            const el = groupRefs.current[keys[i]];
+            if (!el) continue;
+
+            const offset = el.offsetTop - 50;
+            if (scrollTop >= offset - PIXEL_AFFORDANCE) {
+                newIndex = i;
+
+                if (scrollTop >= offset - PIXEL_AFFORDANCE && scrollTop <= offset + PIXEL_AFFORDANCE) {
+                    onHeader = true;
+                }
+            }
         }
 
-        // Use custom or default rendering for options
-        return dropdownProps.onRenderOption ? dropdownProps.onRenderOption(option, defaultRender) : defaultRender?.(option) ?? null;
+        currentIndex.current = newIndex;
+        return { currentGroup: newIndex, onHeader: onHeader };
+    };
+
+    const scrollToTop = () => {
+        currentIndex.current = 0;
+        goToTargetGroup(currentIndex.current);
+    };
+
+    const scrollToBottom = () => {
+        let groupRefKeys = groupRefs?.current ? Object.keys(groupRefs.current) : [];
+        currentIndex.current = groupRefKeys.length - 1;
+        goToTargetGroup(currentIndex.current);
     };
 
+    const goToTargetGroup = (targetGroupIndex: number) => {
+        let groupRefKeys = groupRefs?.current ? Object.keys(groupRefs.current) : [];
+        if (!groupRefs?.current || !dropdownRef.current) {
+            return;
+        }
+        if (targetGroupIndex < 0 || targetGroupIndex >= groupRefKeys.length) {
+            return;
+        }
+
+        currentIndex.current = targetGroupIndex;
+        const el = groupRefs.current[groupRefKeys[targetGroupIndex]];
+        if (el) {
+            const offset = el.offsetTop - 50;
+            dropdownRef.current.scrollTop = offset >= 0 ? offset : 0;
+        }
+    }
+
+    const iterateToGroup = (index: number) => {
+        let currentLocation = recomputeCurrentIndex();
+        currentIndex.current = currentLocation.currentGroup;
+
+        // Determine how to move the current index.
+        // - If we are going "up" (-1 index), we need to see whether we are exactly on the header or not
+        // -- If we are exactly on the header, we intend to go to the previous group (so actually go to the previous group)
+        // -- If however we are within the group, we intend to snap to the current groups header
+        // - If we are going "down" (1 index), we simply just increment on top of the current group.
+        let incrementedIndex = index == -1 ? (currentLocation.onHeader ? currentIndex.current + index : currentIndex.current) : currentIndex.current + index;
+
+        goToTargetGroup(incrementedIndex);
+    };
+
+    // #endregion -- Internal Dropdown Navigation Functions --
+
     return (
         <Dropdown
+            ref={dropdownRef}
             {...dropdownProps}
-            options={getOptions()}
-            onRenderOption={onRenderOption}
             calloutProps={{
                 directionalHint: DirectionalHint.rightCenter,
                 styles: { calloutMain: { width: 500, minWidth: 500 } },
             }}
+            onRenderList={(props, defaultRender) => {
+                if (!props) {
+                    return null;
+                }
+
+                return (
+                    <div ref={dropdownRef} style={{ maxHeight: 1000, overflowY: 'auto', width: 500 }}>
+                        <Stack
+                            horizontal
+                            verticalAlign="center"
+                            styles={{ root: { position: "sticky", top: 0, zIndex: 1, padding: 8, background: modeColors.background } }}
+                        >
+                            <SearchBox
+                                placeholder="Filter options"
+                                autoFocus
+                                onChange={(_ev, newVal) => onSearchValueChanged?.(newVal || '')}
+                                onClear={() => onSearchValueChanged?.('')}
+                                onKeyDown={(ev) => {
+                                    if (ev.key === 'Enter' || ev.key === 'Escape') ev.stopPropagation();
+                                }}
+                                styles={{
+                                    root: {
+                                        flexGrow: 1,
+                                        marginRight: 4,
+                                    },
+                                    field: {
+                                    },
+                                }}
+                            />
+                            <IconButton
+                                iconProps={{ iconName: 'ChevronUp' }}
+                                title="Jump to previous group."
+                                className={dropdownNavigationButtons.base}
+                                onClick={() => { iterateToGroup(-1) }} />
+                            <IconButton
+                                iconProps={{ iconName: 'ChevronDown' }}
+                                title="Jump to next group."
+                                className={dropdownNavigationButtons.base}
+                                onClick={() => { iterateToGroup(1) }} />
+                            <IconButton
+                                iconProps={{ iconName: 'DoubleChevronUp' }}
+                                title="Jump to first group."
+                                className={dropdownNavigationButtons.base}
+                                onClick={scrollToTop} />
+                            <IconButton
+                                iconProps={{ iconName: 'DoubleChevronDown' }}
+                                title="Jump to last group."
+                                className={dropdownNavigationButtons.base}
+                                onClick={scrollToBottom} />
+                        </Stack>
+                        {defaultRender?.(props)}
+                    </div>
+                );
+            }}
         />
     );
 };
@@ -125,8 +258,9 @@ const BasicListParameter: React.FC<{ handler: BuildHealthDataHandler, buildHealt
             onChange(castedOption as IDropdownOption<BuildHealthOptionData>);
         }}
         onDismiss={() => {
-            handler.requestHierarchicalRefresh(request);
-            buildHealthOptions.synchronizeDerivedKeys();
+            buildHealthOptions.commitTransactionSession();
+            handler.requestHierarchicalRefresh(request, undefined, (_request: DataHandlerRefreshRequest) => {
+            });
         }}
     />);
 };
@@ -135,11 +269,15 @@ const BasicListParameter: React.FC<{ handler: BuildHealthDataHandler, buildHealt
  * React Component that is a multi select drop down.
  * @returns React Component.
  */
-const MultiListParameter: React.FC<{ handler: BuildHealthDataHandler, buildHealthOptions: BuildHealthOptionsController, param: BuildHealthOptionList, disabled?: boolean, selectedKeys: string[], request: DataHandlerRefreshRequest, onChange: (option: IDropdownOption<BuildHealthOptionData>) => void }> = function ConstructMultilistParamter({ param, handler, buildHealthOptions, disabled, selectedKeys, request, onChange }) {
+const MultiListParameter: React.FC<{ handler: BuildHealthDataHandler, buildHealthOptions: BuildHealthOptionsController, param: BuildHealthOptionList, disabled?: boolean, selectedKeys: string[], request: DataHandlerRefreshRequest, onChange: (option: IDropdownOption<BuildHealthOptionData>) => void, enabledSelectAll?: boolean, onSelectAll?: (groupOption: IDropdownOption<BuildHealthOptionData>, optionSelectAllContext?: IDropdownOption<BuildHealthOptionData>[]) => void, onHierarchicalRefreshCompleted?: (receipts?: BuildHealthOptionStateDiff) => void }> = function ConstructMultilistParamter({ param, handler, buildHealthOptions, disabled, selectedKeys, request, onChange, enabledSelectAll = false, onSelectAll, onHierarchicalRefreshCompleted }) {
     const [filterText, setFilterText] = useState("");
     const changeVersion = buildHealthOptions.optionsChangeVersion;
     const key = `${PARAMETER_KEY_PREFIX}${KEY_SEPARATOR}${param.label}`;
     const gset: Set<string> = new Set();
+    const selectedSet = useMemo(() => new Set(selectedKeys), [selectedKeys]);
+    const groupRefs = useRef<(HTMLDivElement | null)[]>([]);
+
+    useEffect(() => { groupRefs.current = []; }, [param.items]);
 
     param.items.forEach(item => {
         if (!item.group) {
@@ -153,75 +291,117 @@ const MultiListParameter: React.FC<{ handler: BuildHealthDataHandler, buildHealt
     // #region -- Data Memos --
 
     const doptions = useMemo(() => {
-        const gset: Set<string> = new Set();
-
-        param.items.forEach(item => gset.add(item.group ?? "__nogroup"));
+        const groupMap = new Map<string, BuildHealthOptionData[]>();
 
-        const groups = Array.from(gset).sort((a, b) => a.localeCompare(b));
+        for (const item of param.items) {
+            const group = item.group ?? "__nogroup";
+            if (!groupMap.has(group)) groupMap.set(group, []);
+            groupMap.get(group)!.push(item);
+        }
 
+        const groups = Array.from(groupMap.keys()).sort();
         const options: IDropdownOption<BuildHealthOptionData>[] = [];
 
-        groups.forEach(group => {
+        for (const group of groups) {
             if (group !== "__nogroup") {
-                options.push({
-                    key: `group_${group}`,
-                    text: group,
-                    itemType: DropdownMenuItemType.Header,
-                });
-            }
-            param.items.forEach(item => {
-                if ((item.group ?? "__nogroup") === group) {
+                options.push({ key: `group_${group}`, text: group, itemType: DropdownMenuItemType.Header });
+                if (enabledSelectAll) {
                     options.push({
-                        key: item.id,
-                        text: item.text,
-                        data: item,
-                        selected: selectedKeys.includes(item.id),
+                        key: `group_select_all${KEY_SEPARATOR}${group}`,
+                        itemType: DropdownMenuItemType.Header,
+                        text: SELECT_CLEAR_ALL_SPECIAL_GROUP,
                     });
                 }
-            });
-        });
+            }
+
+            for (const item of groupMap.get(group)!) {
+                options.push({
+                    key: item.id,
+                    text: item.text,
+                    data: item,
+                    selected: selectedSet.has(item.id),
+                });
+            }
+        }
 
         return options;
-    }, [param.items, selectedKeys]);
+    }, [param.items, selectedSet, enabledSelectAll]);
 
     // #endregion -- Data Memos --
 
-    const filteredOptions = doptions.filter((option) => {
-        if (option.itemType === DropdownMenuItemType.Header) return true;
-        if (!filterText) return true;
-        return option.text.toLowerCase().includes(filterText.toLowerCase());
-    });
+    const filteredOptions = useMemo(() => {
+        const lowerFilter = filterText.toLowerCase();
+        return doptions.filter(o =>
+            o.itemType === DropdownMenuItemType.Header || !filterText
+                ? true
+                : o.text.toLowerCase().includes(lowerFilter)
+        );
+    }, [doptions, filterText]);
+
+    const filteredOptionsRef = useRef<IDropdownOption<BuildHealthOptionData>[]>([]);
+
+    useEffect(() => {
+        filteredOptionsRef.current = filteredOptions;
+    }, [filteredOptions]);
 
     // #region -- Dropdown Callbacks --
 
     const onRenderOption = useCallback(
         (option?: ISelectableOption) => {
-            if (!option) return null;
+            if (!option) {
+                return null;
+            }
 
+            // Select & Clear all for the given filtered list.
+            if (option.itemType === DropdownMenuItemType.Header && option.text === SELECT_CLEAR_ALL_SPECIAL_GROUP) {
+                return (
+                    <Stack horizontal tokens={{ childrenGap: 8 }} styles={{
+                        root:
+                        {
+                            width: "100%",
+                            height: "100%",
+                            alignItems: "center",
+                            justifyContent: "center",
+                        }
+                    }}>
+                        {/* We have opted to use buttons here as we: 1. don't want to persist the state of a select all option; 2. the semantics of storing a "Select/Clear All" value gets complicated when the action is applied to a filtered set; */}
+                        <DefaultButton styles={{ root: { flex: 1, height: 20 } }} text="Select All"
+                            onClick={() => {
+                                if (!option) return;
+
+                                const synthesizedOption: IDropdownOption<BuildHealthOptionData> = {
+                                    ...option,
+                                    selected: true,
+                                };
+
+                                // Only send a suggested list if we actually have filtered something.
+                                let fitleredOptionList = filteredOptionsRef.current.length == doptions.length ? undefined : filteredOptionsRef.current;
+                                option && onSelectAll?.(synthesizedOption, fitleredOptionList)
+                            }}
+                        />
+                        <DefaultButton styles={{ root: { flex: 1, height: 20 } }} text="Clear All"
+                            onClick={() => {
+                                if (!option) return;
+
+                                const synthesizedOption: IDropdownOption<BuildHealthOptionData> = {
+                                    ...option,
+                                    selected: false,
+                                };
+
+                                // Only send a suggested list if we actually have filtered something.
+                                let fitleredOptionList = filteredOptionsRef.current.length == doptions.length ? undefined : filteredOptionsRef.current;
+                                option && onSelectAll?.(synthesizedOption, fitleredOptionList)
+                            }} />
+                    </Stack>
+                );
+            }
             if (option.itemType === DropdownMenuItemType.Header) {
-                return <div style={{ fontWeight: "bold", padding: "4px 8px" }}>{option.text}</div>;
+                return <div ref={el => groupRefs.current[option.key as string] = el} style={{ fontWeight: "bold", padding: "4px 8px" }}>{option.text}</div>;
             }
 
             return (
-                <TooltipHost key={`tooltip_${option.text}`} content={
-                    <div>
-                        <Stack>
-                            <Label>{(option as IDropdownOption<BuildHealthOptionData>).data?.group}</Label>
-                        </Stack>
-                        <Stack>
-                            {option.text}
-                        </Stack>
-                    </div>
-                }
-                    directionalHint={DirectionalHint.rightCenter}
-                    calloutProps={{
-                        styles: {
-                            root: { maxWidth: 800 }
-                        }
-                    }}
-                >
-                    <div style={{ padding: "4px 8px" }}>{option.text}</div>
-                </TooltipHost>
+                // Due to some load testing, when we get 1000s of entires we don't want to bloat the load time with ToolTipHost construction.
+                <div title={`Group: ${option.data?.group || ""}\n${option.text}`} className={dropdownOptions.base}>{option.text}</div>
             );
         }, []);
 
@@ -234,8 +414,16 @@ const MultiListParameter: React.FC<{ handler: BuildHealthDataHandler, buildHealt
 
     const handleMenuDismiss = useCallback(() => {
         setFilterText("");
-        handler.requestHierarchicalRefresh(request);
-        buildHealthOptions.synchronizeDerivedKeys();
+        let buildHealthOptionStateReceipts = buildHealthOptions.commitTransactionSession();
+        handler.requestHierarchicalRefresh(request, undefined,
+            // Upon completion of the hierarchical refresh, we want to commit the transaction session for the options. 
+            // We do this because the data handler can populate the template & step cache which may be useful for any downstream receipt processing, and actions that depend on the side-effects of updated hierarchy data.
+            (_request: DataHandlerRefreshRequest) => {
+                if (onHierarchicalRefreshCompleted) {
+                    onHierarchicalRefreshCompleted(buildHealthOptionStateReceipts);
+                }
+            });
+
     }, [request, buildHealthOptions]);
 
     // #endregion -- Dropdown Callbacks --
@@ -247,6 +435,7 @@ const MultiListParameter: React.FC<{ handler: BuildHealthDataHandler, buildHealt
             placeholder={param.label}
             options={filteredOptions}
             multiSelect
+            groupRefs={groupRefs}
             selectedKeys={selectedKeys}
             onChange={handleDropdownChange}
             onDismiss={handleMenuDismiss}
@@ -290,7 +479,7 @@ const ProjectDropdownSingle: React.FC<{ handler: BuildHealthDataHandler, buildHe
             <Label>Project</Label>
             <BasicListParameter handler={handler} buildHealthOptions={buildHealthOptions} param={projectDropdownParams} selectedKey={selectedProjectKey} request={request} onChange={(option: IDropdownOption<BuildHealthOptionData>) => {
                 const key = option.key as string;
-                buildHealthOptions.toggleSingleProject(key, option.text);
+                buildHealthOptions.getTransactionSession().toggleSingleProject(key, option.text);
             }} />
         </Stack>
     );
@@ -323,7 +512,7 @@ const StreamDropdown: React.FC<{ handler: BuildHealthDataHandler, buildHealthOpt
                 const key = option.key as string;
                 const isSelected = !!option.selected;
 
-                buildHealthOptions.toggleStream(key, option.text, isSelected);
+                buildHealthOptions.getTransactionSession().toggleStream(key, option.text, isSelected);
             }} />
         </Stack>
     );
@@ -353,12 +542,24 @@ const TemplateDropdown: React.FC<{ handler: BuildHealthDataHandler, buildHealthO
         <Stack id={`job_dropdown`}>
             <Label>Template(s)</Label>
             <Stack horizontal tokens={{ childrenGap: 8 }} verticalAlign="end">
-                <MultiListParameter handler={handler} buildHealthOptions={buildHealthOptions} param={templateDropdownParams} selectedKeys={selectedTemplateKeys} request={request} onChange={(option: IDropdownOption<BuildHealthOptionData>) => {
-                    const key = option.key as string;
-                    const isSelected = !!option.selected;
-
-                    buildHealthOptions.toggleTemplate(key, option.text, isSelected);
-                }} />
+                <MultiListParameter handler={handler} buildHealthOptions={buildHealthOptions} param={templateDropdownParams} selectedKeys={selectedTemplateKeys} request={request}
+                    onChange={(option: IDropdownOption<BuildHealthOptionData>) => {
+                        const key = option.key as string;
+                        const isSelected = !!option.selected;
+                        buildHealthOptions.getTransactionSession().toggleTemplate(key, option.text, isSelected);
+                    }}
+                    onHierarchicalRefreshCompleted={(receipts?: BuildHealthOptionStateDiff) => {
+                        // If we have added templates, issue another request to toggle all steps & close the transaction.
+                        // We do not want to propagate linked actions in this special case, as it can be quite destructive.
+                        if (receipts !== null && receipts!.added && receipts!.added.enabledTemplates.size > 0) {
+                            for (const key of receipts!.added.enabledTemplates!) {
+                                const templateData = decodeTemplateKey(key);
+                                let stepNames = handler.getCachedTemplateStepNames(templateData.streamId, templateData.templateId);
+                                buildHealthOptions.getTransactionSession().toggleAllSteps(templateData.streamId, templateData.templateId, stepNames, true, false);
+                                buildHealthOptions.commitTransactionSession();
+                            }
+                        }
+                    }} />
                 {
                     buildHealthOptions.state.linkedSelectionModeEnabled &&
                     (
@@ -368,9 +569,10 @@ const TemplateDropdown: React.FC<{ handler: BuildHealthDataHandler, buildHealthO
                             <IconButton
                                 iconProps={{ iconName: 'Link' }}
                                 styles={{ root: { width: 15, minWidth: 15, padding: 0 } }} onClick={() => {
-                                    buildHealthOptions.toggleAllLinkableTemplates(true);
+                                    // We do not want to re-issue select all steps for all newly linked templates, as it can be quite destructive.
+                                    buildHealthOptions.getTransactionSession().toggleAllLinkableTemplates(true);
+                                    buildHealthOptions.commitTransactionSession();
                                     handler.requestHierarchicalRefresh(request);
-                                    buildHealthOptions.synchronizeDerivedKeys();
                                 }
                                 } />
                         </TooltipHost>
@@ -382,6 +584,83 @@ const TemplateDropdown: React.FC<{ handler: BuildHealthDataHandler, buildHealthO
 });
 
 const StepDropdown: React.FC<{ handler: BuildHealthDataHandler, buildHealthOptions: BuildHealthOptionsController }> = observer(function ConstructProjectDropdown({ handler, buildHealthOptions }) {
+    let selectedStepKeys: string[] = Object.keys(buildHealthOptions.state.enabledSteps);
+    let request: DataHandlerRefreshRequest = { steps: false };
+
+    // #region -- Data Memos --
+
+    const items = useMemo<BuildHealthOptionData[]>(
+        () =>
+            handler.stepData
+                .map((step: StepRefData) => ({
+                    id: encodeStepKey(step),
+                    text: step.name,
+                    group: `${step.streamId} - ${step.templateId}`,
+                }))
+                .sort((a, b) => a.text.localeCompare(b.text)),
+        [handler.stepData]
+    );
+
+    // #endregion -- Data Memos --
+
+    // #region -- Callbacks --
+
+    const handleChange = useCallback(
+        (option: IDropdownOption<BuildHealthOptionData>) => {
+            const key = option.key as string;
+            const isSelected = !!option.selected;
+            buildHealthOptions
+                .getTransactionSession()
+                .toggleStep(key, option.text, isSelected);
+        },
+        [buildHealthOptions]
+    );
+
+    const handleSelectAll = useCallback(
+        (
+            option: IDropdownOption<BuildHealthOptionData>,
+            optionSelectAllContext?: IDropdownOption<BuildHealthOptionData>[]
+        ) => {
+            const key = option.key as string;
+            const [, data] = key.split(KEY_SEPARATOR);
+            const [streamId, templateId] = data.split(" - ");
+            const isSelected = !!option.selected;
+
+            // If there is no provided context, simply obtain all the known steps of the (stream, template).
+            const stepNames: string[] =
+                optionSelectAllContext === undefined
+                    ? handler.getCachedTemplateStepNames(streamId, templateId)
+                    : optionSelectAllContext
+                        .filter((opt) => {
+                            if (!opt.data) return false;
+                            const decoded = decodeStepKey(opt.data.id!);
+                            return (
+                                decoded.streamId === streamId &&
+                                decoded.templateId === templateId
+                            );
+                        })
+                        .map((x) => x.data!.text!);
+
+            // Do not close the transaction session; the user still has the modal open
+            buildHealthOptions
+                .getTransactionSession()
+                .toggleAllSteps(streamId, templateId, stepNames, isSelected);
+        },
+        [buildHealthOptions, handler]
+    );
+
+    const handleLinkSteps = useCallback(
+        () => {
+            // Close the transaction session; the modal is not open so there will be no such transaction completion otherwise.
+            buildHealthOptions.getTransactionSession().toggleAllLinkableSteps(true);
+            buildHealthOptions.commitTransactionSession();
+            handler.requestHierarchicalRefresh(request);
+        },
+        [buildHealthOptions, handler]
+    );
+
+    // #endregion -- Callbacks --
+
     const stepDropdownParams: BuildHealthOptionList = {
         label: "Select Step(s)",
         items: [
@@ -389,28 +668,16 @@ const StepDropdown: React.FC<{ handler: BuildHealthDataHandler, buildHealthOptio
         toolTip: "Select an option from the dropdown."
     };
 
-    stepDropdownParams.items = handler.stepData.map((step: StepRefData) => {
-        let listParam: BuildHealthOptionData = {
-            id: encodeStepKey(step),
-            text: step.name,
-            group: `${step.streamId} - ${step.templateId}`
-        };
-        return listParam;
-    }).sort((a, b) => a.text.localeCompare(b.text));
-
-    let selectedStepKeys: string[] = Object.keys(buildHealthOptions.state.enabledSteps);
-    let request: DataHandlerRefreshRequest = { steps: true };
+    stepDropdownParams.items = items;
 
     return stepDropdownParams.items.length > 0 && (
         <Stack id={`step_dropdown`}>
             <Label>Step(s)</Label>
             <Stack horizontal tokens={{ childrenGap: 8 }} verticalAlign="end">
-                <MultiListParameter handler={handler} buildHealthOptions={buildHealthOptions} param={stepDropdownParams} selectedKeys={selectedStepKeys} request={request} onChange={(option: IDropdownOption<BuildHealthOptionData>) => {
-                    const key = option.key as string;
-                    const isSelected = !!option.selected;
-
-                    buildHealthOptions.toggleStep(key, option.text, isSelected);
-                }} />
+                <MultiListParameter handler={handler} buildHealthOptions={buildHealthOptions} param={stepDropdownParams} selectedKeys={selectedStepKeys} request={request}
+                    onChange={handleChange}
+                    enabledSelectAll={true}
+                    onSelectAll={handleSelectAll} />
                 {
                     buildHealthOptions.state.linkedSelectionModeEnabled &&
                     (
@@ -419,12 +686,8 @@ const StepDropdown: React.FC<{ handler: BuildHealthDataHandler, buildHealthOptio
                             calloutProps={{ directionalHint: DirectionalHint.rightCenter }}>
                             <IconButton
                                 iconProps={{ iconName: 'Link' }}
-                                styles={{ root: { width: 15, minWidth: 15, padding: 0 } }} onClick={() => {
-                                    buildHealthOptions.toggleAllLinkableSteps(true);
-                                    handler.requestHierarchicalRefresh(request);
-                                    buildHealthOptions.synchronizeDerivedKeys();
-                                }
-                                } />
+                                styles={{ root: { width: 15, minWidth: 15, padding: 0 } }}
+                                onClick={handleLinkSteps} />
                         </TooltipHost>
                     )
                 }
@@ -457,7 +720,7 @@ const BuildHealthUrlSync: React.FC<{ buildHealthController: BuildHealthOptionsCo
     }, []);
 
     useEffect(() => {
-        const desiredSearch = buildHealthController.toNavigationQuery().toString();
+        const desiredSearch = toNavigationQuery(buildHealthController.state, buildHealthController.querySchemaVersion, handler).toString();
         if (desiredSearch && location.search !== `?${desiredSearch}`) {
             navigate(`${location.pathname}?${desiredSearch}`, { replace: true });
         }
@@ -488,17 +751,17 @@ const BuildHealthSidebar: React.FC<{ handler: BuildHealthDataHandler, buildHealt
     const updateJobHistoryTimeSpan = (_ev, option) => {
         const select = option as TimeSpan;
         setTimeSpanKey(select.key);
-        buildHealthOptions.setJobHistoryTimeSpan(select);
+        buildHealthOptions.getTransactionSession().setJobHistoryTimeSpan(select);
     };
 
     const updateIncludePreflight = (isChecked) => {
         setIncludePreflight(isChecked);
-        buildHealthOptions.setIncludePreflight(isChecked);
+        buildHealthOptions.getTransactionSession().setIncludePreflight(isChecked);
     }
 
     const updateLinkedSelectionMode = (isChecked) => {
         setLinkedSelectionMode(isChecked);
-        buildHealthOptions.setLinkedSelectionMode(isChecked);
+        buildHealthOptions.getTransactionSession().setLinkedSelectionMode(isChecked);
     }
 
     const updateIncludeDateAnchors = (isChecked) => {
@@ -575,9 +838,9 @@ const BuildHealthSidebar: React.FC<{ handler: BuildHealthDataHandler, buildHealt
                             <PrimaryButton text="Clear" onClick={() => {
 
                                 runInAction(() => {
-                                    buildHealthController.clearAll();
+                                    buildHealthController.getTransactionSession().clearAll();
                                     handler.reset();
-                                    buildHealthOptions.synchronizeDerivedKeys();
+                                    buildHealthOptions.commitTransactionSession();
                                 });
                             }} />
                         </TooltipHost>
@@ -592,9 +855,9 @@ const BuildHealthSidebar: React.FC<{ handler: BuildHealthDataHandler, buildHealt
                             styles={{ root: { width: timeComboWidth } }}
                             options={JobHistoryTimeSpans}
                             selectedKey={timeSpanKey}
-                            onChange={(ev, option, index, value) => {
+                            onChange={(ev, option, _index, _value) => {
                                 updateJobHistoryTimeSpan(ev, option);
-                                buildHealthOptions.synchronizeDerivedKeys();
+                                buildHealthOptions.commitTransactionSession();
                             }}
                         />
                     </Stack>
diff --git a/Engine/Source/Programs/Horde/HordeDashboard/src/components/buildhealth/stepoutcome/StepOutcomeD3Component.tsx b/Engine/Source/Programs/Horde/HordeDashboard/src/components/buildhealth/stepoutcome/StepOutcomeD3Component.tsx
--- a/Engine/Source/Programs/Horde/HordeDashboard/src/components/buildhealth/stepoutcome/StepOutcomeD3Component.tsx
+++ b/Engine/Source/Programs/Horde/HordeDashboard/src/components/buildhealth/stepoutcome/StepOutcomeD3Component.tsx
@@ -565,9 +565,9 @@ function GenerateD3Table(stepOutcomeTable: StepOutcomeTable, handler: StepOutcom
                                 })
                         getCellClasses(entry, handler.getJobResponseData(entry.jobId)).forEach(cls => div.classed(cls, true));
                     } else if (isSummary(entry)) { // SummaryTableEntry
-                        let ratio: number = 0;
+                        let ratio: number = -1;
 
-                        if (entry.steps > 0) {
+                        if (entry.steps > 0 && entry.stepsCompleted > 0) {
                             let numerator: number = entry.stepsPass
 
                             if (!stepOutcomeTable.viewModelOptions.warningsAsSummaryFailure) {
@@ -577,7 +577,7 @@ function GenerateD3Table(stepOutcomeTable: StepOutcomeTable, handler: StepOutcom
                             ratio = numerator / entry.stepsCompleted;
                         }
 
-                        let ratioStr: string = `${(ratio * 100).toFixed(1)}%`;
+                        let ratioStr: string = ratio === -1 ? "N/A" : `${(ratio * 100).toFixed(1)}%`;
 
                         div.style("cursor", "default")
                             .style("background-color", summaryColorScale(ratio))
diff --git a/Engine/Source/Programs/Horde/HordeDashboard/src/components/buildhealth/stepoutcome/StepOutcomeDataHandler.ts b/Engine/Source/Programs/Horde/HordeDashboard/src/components/buildhealth/stepoutcome/StepOutcomeDataHandler.ts
--- a/Engine/Source/Programs/Horde/HordeDashboard/src/components/buildhealth/stepoutcome/StepOutcomeDataHandler.ts
+++ b/Engine/Source/Programs/Horde/HordeDashboard/src/components/buildhealth/stepoutcome/StepOutcomeDataHandler.ts
@@ -147,7 +147,8 @@ export class StepOutcomeDataHandler extends PollBase {
         if (this.activeRefresh ||
             this.activeFilter === undefined ||
             !this.activeFilter.streams?.length ||
-            !this.activeFilter.jobs?.length
+            !this.activeFilter.jobs?.length ||
+            !this.activeFilter.steps?.length
         ) {
             this.stepOutcomeTableData.reset();
             this.setUpdated();
@@ -277,16 +278,14 @@ export class StepOutcomeDataHandler extends PollBase {
                         let streamStepSet = stepMap.get(jobResponse.streamId);
                         let templateStepSet = streamStepSet ? streamStepSet.get(jobResponse.templateId!) : undefined;
                         if (
-                            !templateStepSet ||
-                            templateStepSet.size === 0 ||
+                            templateStepSet &&
+                            templateStepSet.size >= 0 &&
                             templateStepSet.has(fullyQualifiedStepName)
                         ) {
                             this.stepOutcomeTableData.addEntry(new StepOutcomeTableEntry(step, jobResponse.id, jobResponse.name, jobResponse.streamId, jobResponse.createTime.toString(), jobResponse.change, this.changeData.get(jobResponse.change!)?.dateUtc));
                         }
                     }
                 }
-
-                // Summarize each row
             }
         }
 
