[Horde][Dashboard] Lint support & unit test support

#DisabledValidations Json

diff --git a/Engine/Source/Programs/Horde/HordeDashboard/.eslintrc.js b/Engine/Source/Programs/Horde/HordeDashboard/.eslintrc.js
deleted file mode 100644
diff --git a/Engine/Source/Programs/Horde/HordeDashboard/package.json b/Engine/Source/Programs/Horde/HordeDashboard/package.json
--- a/Engine/Source/Programs/Horde/HordeDashboard/package.json
+++ b/Engine/Source/Programs/Horde/HordeDashboard/package.json
@@ -8,7 +8,8 @@
     "build": "tsc && vite build",
     "lint": "eslint . --ext ts,tsx --report-unused-disable-directives --max-warnings 0",
     "preview": "vite preview",
-    "license-check": "license-checker-rseidelsohn --production --excludePackages horde-dashboard --depth 0 --limitAttributes licenses --csv --customPath ./license-format.json --onlyAllow \"Apache-2.0;MIT;ISC\""
+    "license-check": "license-checker-rseidelsohn --production --excludePackages horde-dashboard --depth 0 --limitAttributes licenses --csv --customPath ./license-format.json --onlyAllow \"Apache-2.0;MIT;ISC\"",
+    "test": "vitest run"
   },
   "dependencies": {
     "@blueprintjs/icons": "^3.13.0",
@@ -43,12 +44,14 @@
     "@typescript-eslint/parser": "^7.2.0",
     "@vitejs/plugin-react": "^4.2.1",
     "eslint": "^8.57.0",
+    "eslint-config-react-app": "^7.0.1",
     "eslint-plugin-react-hooks": "^4.6.0",
     "eslint-plugin-react-refresh": "^0.4.6",
     "license-checker-rseidelsohn": "4.3.0",
     "rimraf": "^3.0.0",
     "typescript": "^5.4.5",
     "vite": "^5.2.0",
-    "vite-tsconfig-paths": "^4.3.2"
+    "vite-tsconfig-paths": "^4.3.2",
+    "vitest": "^1.0.0"
   }
 }
diff --git a/Engine/Source/Programs/Horde/HordeDashboard/tests/StepOutcomeDataTypes.test.ts b/Engine/Source/Programs/Horde/HordeDashboard/tests/StepOutcomeDataTypes.test.ts
new file mode 100644
--- /dev/null
+++ b/Engine/Source/Programs/Horde/HordeDashboard/tests/StepOutcomeDataTypes.test.ts
@@ -0,0 +1,26 @@
+import { describe, it, expect } from 'vitest';
+import { encodeStepName } from '../src/components/buildhealth/stepoutcome/StepOutcomeDataTypes';
+
+
+// Mock object type
+type StepOutcomeTableEntry = {
+    streamId: string;
+    stepResponse: {
+        name: string;
+    };
+};
+
+const KEY_SEPARATOR = "_";
+
+describe('encodeStepName', () => {
+    it('Should encode the step name correctly', () => {
+        const mockEntry: StepOutcomeTableEntry = {
+            streamId: 'STREAM1',
+            stepResponse: { name: 'StepA' }
+        };
+        
+        const result = encodeStepName(mockEntry as any);
+
+        expect(result).toBe('stream1::stepa'); // streamId + KEY_SEPARATOR + stepResponse.name, lowercased
+    });
+});
\ No newline at end of file
diff --git a/Engine/Source/Programs/Horde/HordeDashboard/tsconfig.json b/Engine/Source/Programs/Horde/HordeDashboard/tsconfig.json
--- a/Engine/Source/Programs/Horde/HordeDashboard/tsconfig.json
+++ b/Engine/Source/Programs/Horde/HordeDashboard/tsconfig.json
@@ -36,6 +36,6 @@
     "noUnusedParameters": true,    
     */
   },
-  "include": ["src", "plugins"],
+  "include": ["src", "plugins", "tests"],
   "references": [{ "path": "./tsconfig.node.json" }]
 }
