Horde: Adding Root Cause Analysis Slack support

diff --git a/Engine/Source/Programs/Horde/Plugins/Build/HordeServer.Build/Notifications/Sinks/SlackNotificationSink.cs b/Engine/Source/Programs/Horde/Plugins/Build/HordeServer.Build/Notifications/Sinks/SlackNotificationSink.cs
--- a/Engine/Source/Programs/Horde/Plugins/Build/HordeServer.Build/Notifications/Sinks/SlackNotificationSink.cs
+++ b/Engine/Source/Programs/Horde/Plugins/Build/HordeServer.Build/Notifications/Sinks/SlackNotificationSink.cs
@@ -1170,6 +1170,7 @@ namespace HordeServer.Notifications.Sinks
 						actions.AddButton("Assign to Me", value: $"issue_{issue.Id}_ack", style: ButtonStyle.Primary);
 						actions.AddButton("Not Me", value: $"issue_{issue.Id}_decline", style: ButtonStyle.Danger);
 						actions.AddButton("Mark Fixed", value: $"issue_{issue.Id}_markfixed");
+						actions.AddButton("Root Cause", value: $"issue_{issue.Id}_rootcause");
 
 						string? context = null;
 						if (issue.OwnerId != null)
@@ -3237,6 +3238,109 @@ namespace HordeServer.Notifications.Sinks
 								_logger.LogInformation("Marked issue {IssueId} fixed by user {UserId} in {Change}", issueId, resolvedById, commitId);
 							}
 						}
+						else if (TryMatch(payload.View.CallbackId, @"^issue_(\d+)_rootcause_([a-fA-F0-9]{24})$", out match))
+						{
+							int issueId = Int32.Parse(match.Groups[1].Value, NumberStyles.Integer, CultureInfo.InvariantCulture);
+							UserId userId = UserId.Parse(match.Groups[2].Value);
+
+							string? rootChange;
+							string? rootCategory;
+							string? rootSummary;
+							int? dupeIssueId = null;
+							CommitId rootCommit = new CommitId("");
+
+							if (payload.View.State.TryGetValue("rootcause_cl", "rootcause_cl_action", out rootChange))
+							{
+								if (rootChange != null && rootChange.Length > 0)
+								{
+									if (Int32.TryParse(rootChange.Trim(), out int result))
+									{
+										rootCommit = CommitId.FromPerforceChange(result);
+										_logger.LogInformation("Setting issue {IssueId} root cause change by user {UserId} to {Change}", issueId, userId, rootCommit);
+									}
+									else
+									{
+										_logger.LogInformation("Setting issue {IssueId} root cause change by user {UserId} not set to invalid {Change}", issueId, userId, rootChange);
+									}
+								}
+							}
+
+							if (payload.View.State.TryGetValue("rootcause_category", "rootcause_category_action", out rootCategory))
+							{
+								if (rootCategory != null && rootCategory.Length > 0)
+								{
+									_logger.LogInformation("Setting issue {IssueId} root cause category by user {UserId} to {Category}", issueId, userId, rootCategory);
+								}
+								else
+								{
+									rootCategory = "";
+								}
+							}
+							else
+							{
+								rootCategory = "";
+							}
+
+							if (payload.View.State.TryGetValue("rootcause_summary", "rootcause_summary_action", out rootSummary))
+							{
+								if (rootSummary != null && rootSummary.Length > 0)
+								{
+									_logger.LogInformation("Setting issue {IssueId} root cause summary by user {UserId} to {Summary}", issueId, userId, rootSummary);
+								}
+								else
+								{
+									rootSummary = "";
+								}
+							}
+							else
+							{
+								rootSummary = "";
+							}
+
+							if (payload.View.State.TryGetValue("rootcause_dupeid", "rootcause_dupeid_action", out string? dupeIdStr))
+							{
+								if (dupeIdStr != null && dupeIdStr.Length > 0)
+								{
+									if (Int32.TryParse(dupeIdStr.Trim(), out int result))
+									{
+										dupeIssueId = result;
+										_logger.LogInformation("Setting issue {IssueId} dupliicate issue id by user {UserId} to {DupeId}", issueId, userId, dupeIssueId);
+									}
+									else
+									{
+										_logger.LogInformation("Setting issue {IssueId} duplicate issue id by user {UserId} not set to invalid {DupeId}, clearing", issueId, userId, dupeIdStr);
+										dupeIssueId = -1;
+									}
+								}
+							}
+							else
+							{
+								dupeIssueId = -1;
+							}
+
+							IIssue? issue = await _issueService.Collection.GetIssueAsync(issueId, cancellationToken);
+
+							string? summary = null;
+
+							if (issue != null)
+							{
+								string curSummary = issue.Summary;
+								bool markedDupe = curSummary.StartsWith("(Dupe) ", StringComparison.OrdinalIgnoreCase);
+								if (!markedDupe && dupeIssueId.HasValue && dupeIssueId != -1)
+								{
+									curSummary = "(Dupe) " + curSummary;
+								}
+								else if (markedDupe && (!dupeIssueId.HasValue || dupeIssueId == -1))
+								{
+									curSummary = curSummary.Substring(7);
+								}
+
+								summary = curSummary;
+							}
+
+							await _issueService.UpdateIssueAsync(issueId, summary: summary, duplicateIssueId: dupeIssueId, rootCommitId: rootCommit, rootCauseCategory: rootCategory, rootCauseSummary: rootSummary, cancellationToken: cancellationToken);
+
+						}
 						else if (TryMatch(payload.View.CallbackId, @"^issue_(\d+)_ack_([a-fA-F0-9]{24})$", out match))
 						{
 							int issueId = Int32.Parse(match.Groups[1].Value, NumberStyles.Integer, CultureInfo.InvariantCulture);
@@ -3393,6 +3497,80 @@ namespace HordeServer.Notifications.Sinks
 				view.Close = "Cancel";
 				view.Submit = "Mark Fixed";
 
+				await _slackClient.OpenViewAsync(triggerId, view, cancellationToken);
+				return;
+			}
+			else if (String.Equals(verb, "rootcause", StringComparison.Ordinal))
+			{
+				IIssue? issue = await _issueService.Collection.GetIssueAsync(issueId, cancellationToken);
+
+				if (issue == null)
+				{
+					return;
+				}
+
+				List<SlackOption> options = new List<SlackOption>
+				{
+				   new SlackOption( "Asset/Content", "Asset/Content"),
+				   new SlackOption( "Code", "Code"),
+				   new SlackOption( "UI/Localization", "UI/Localization"),
+				   new SlackOption( "Build/Deployment/Version", "Build/Deployment/Version"),
+				   new SlackOption( "Configuration/Data", "Configuration/Data"),
+				   new SlackOption( "Gameplay/Engine", "Gameplay/Engine"),
+				   new SlackOption( "Backend/Server/Network", "Backend/Server/Network"),
+				   new SlackOption( "Monitoring/No Code Change", "Monitoring/No Code Change"),
+				   new SlackOption( "Unintended Side Effect", "Unintended Side Effect"),
+				   new SlackOption( "Build Agent", "Build Agent"),
+				   new SlackOption( "Other", "Other"),				   
+				};						
+
+				SlackView view = new SlackView($"Issue {issueId}");
+				view.CallbackId = $"issue_{issueId}_rootcause_{user.Id}";
+
+				StaticSelectMenuElement categorySelect = new StaticSelectMenuElement("rootcause_category_action", "Select an Issue Category", options);
+				if (!String.IsNullOrEmpty(issue.RootCauseCategory))
+				{
+					SlackOption? existing = options.FirstOrDefault(i => i.Value == issue.RootCauseCategory);
+					if (existing != null)
+					{
+						// Only can get, issue with Epic Slack
+						//categorySelect.InitialOption = existing;
+					}
+				}
+
+				if (issue.DuplicateIssueId == null)
+				{
+					InputBlock selectBlock = view.AddInput("Root Cause Category:", categorySelect);
+					selectBlock.BlockId = "rootcause_category";
+					selectBlock.Optional = true;
+
+					InputBlock summaryBlock = view.AddInput("Root Cause Summary:", new PlainTextInputElement("rootcause_summary_action", multiLine: true, placeholder: "Please enter a brief summary describing the root cause of this issue", initialValue: issue?.RootCauseSummary));
+					summaryBlock.BlockId = "rootcause_summary";
+					summaryBlock.Optional = true;
+
+					InputBlock changeBlock = view.AddInput("Root Cause Changelist:", new PlainTextInputElement("rootcause_cl_action", placeholder: "Root Cause Changelist", initialValue: issue?.RootCommitId?.Name));
+					changeBlock.BlockId = "rootcause_cl";
+					changeBlock.Hint = "Change which caused this issue, please only enter if certain";
+
+					changeBlock.Optional = true;
+				}
+
+				InputBlock dupeBlock = view.AddInput("Duplicate Issue Id:", new PlainTextInputElement("rootcause_dupeid_action", placeholder: "Duplicate Issue Id", initialValue: issue?.DuplicateIssueId?.ToString()));
+				dupeBlock.BlockId = "rootcause_dupeid";
+				dupeBlock.Optional = true;
+
+				if (issue.DuplicateIssueId != null)
+				{
+					dupeBlock.Hint = "Issue ID of the source issue this issue duplicates.  Root cause inputs hidden for duplicate issue, leave blank and update to clear";
+				}
+				else
+				{
+					dupeBlock.Hint = "Issue ID of the source issue this issue duplicate. Enter 0 or an issue id to set this as a duplicate issue";
+				}
+
+				view.Close = "Cancel";
+				view.Submit = "Update";
+
 				await _slackClient.OpenViewAsync(triggerId, view, cancellationToken);
 				return;
 			}
