Horde: Fix server warning in Slack notifications
#submitool bypass

diff --git a/Engine/Source/Programs/Horde/Plugins/Build/HordeServer.Build/Notifications/Sinks/SlackNotificationSink.cs b/Engine/Source/Programs/Horde/Plugins/Build/HordeServer.Build/Notifications/Sinks/SlackNotificationSink.cs
--- a/Engine/Source/Programs/Horde/Plugins/Build/HordeServer.Build/Notifications/Sinks/SlackNotificationSink.cs
+++ b/Engine/Source/Programs/Horde/Plugins/Build/HordeServer.Build/Notifications/Sinks/SlackNotificationSink.cs
@@ -3559,7 +3559,7 @@ namespace HordeServer.Notifications.Sinks
 				dupeBlock.BlockId = "rootcause_dupeid";
 				dupeBlock.Optional = true;
 
-				if (issue.DuplicateIssueId != null)
+				if (issue?.DuplicateIssueId != null)
 				{
 					dupeBlock.Hint = "Issue ID of the source issue this issue duplicates.  Root cause inputs hidden for duplicate issue, leave blank and update to clear";
 				}
