Horde: Surface Gauntlet ensure warnings and make the issue summary reflects having a callstack

diff --git a/Engine/Source/Programs/Horde/Plugins/Build/HordeServer.Build.Tests/Issues/IssueServiceTests.cs b/Engine/Source/Programs/Horde/Plugins/Build/HordeServer.Build.Tests/Issues/IssueServiceTests.cs
--- a/Engine/Source/Programs/Horde/Plugins/Build/HordeServer.Build.Tests/Issues/IssueServiceTests.cs
+++ b/Engine/Source/Programs/Horde/Plugins/Build/HordeServer.Build.Tests/Issues/IssueServiceTests.cs
@@ -2225,6 +2225,38 @@ namespace HordeServer.Tests.Issues
 		[TestMethod]
 		public async Task GauntletTestAsync()
 		{
+			string logAssertionMessage = "Assertion failed: State.bGfxPSOSet [File:D:\\build\\U5M+Inc\\Sync\\Engine\\Source\\Runtime\\RHI\\Public\\RHIValidationContext.h] [Line: 809]\n"
+			+ "A Graphics PSO has to be set to set resources into a shader!";
+
+			string logEnsureMessage = "Ensure: Ensure condition failed: State.bGfxPSOSet [File:D:\\build\\U5M+Inc\\Sync\\Engine\\Source\\Runtime\\RHI\\Public\\RHIValidationContext.h] [Line: 809]\n"
+			+ "A Graphics PSO has to be set to set resources into a shader!";
+
+			string logCallstack = "	0x00007fff4b43a1f4 UnrealEditor-RHI.dll!FValidationContext::RHISetShaderParameters() [Unknown File]\n"
+			+ "	0x00007fff4b3e5a01 UnrealEditor-RHI.dll!FRHICommandSetShaderParameters<FRHIGraphicsShader>::Execute() [Unknown File]\n"
+			+ "	0x00007fff4b3e9a5a UnrealEditor-RHI.dll!FRHICommand<FRHICommandSetShaderParameters<FRHIGraphicsShader>,FRHICommandSetShaderParametersString1159>::ExecuteAndDestruct() [Unknown File]\n"
+			+ "	0x00007fff4b3e74f7 UnrealEditor - RHI.dll!FRHICommandListBase::Execute()[Unknown File]\n"
+			+ "	0x00007fff4b3f3228 UnrealEditor-RHI.dll!FRHICommandListImmediate::ExecuteAndReset() [Unknown File]\n"
+			+ "	0x00007fff4b464ebf UnrealEditor-RHI.dll!FRHIComputeCommandList::SubmitCommandsHint() [Unknown File]\n"
+			+ "	0x00007fff3d3066fc UnrealEditor-Renderer.dll!TBaseStaticDelegateInstance<void [Unknown File]\n"
+			+ "	0x00007fff48bd7813 UnrealEditor-RenderCore.dll!FRDGBuilder::ExecutePass() [Unknown File]\n"
+			+ "	0x00007fff48bd2a8a UnrealEditor-RenderCore.dll!FRDGBuilder::Execute() [Unknown File]\n"
+			+ "	0x00007fff3d31fef4 UnrealEditor-Renderer.dll!FSceneRenderer::RenderThreadEnd() [Unknown File]\n"
+			+ "	0x00007fff3d2ed57a UnrealEditor-Renderer.dll!`FPixelShaderUtils::AddFullscreenPass<FHZBTestPS>'::`2'::<lambda_1>::operator()() [Unknown File]\n"
+			+ "	0x00007fff3d3043f9 UnrealEditor - Renderer.dll!FSceneRenderer::DoOcclusionQueries()[Unknown File]\n"
+			+ "	0x00007fff3d30c888 UnrealEditor-Renderer.dll!TBaseStaticDelegateInstance<void [Unknown File]\n"
+			+ "	0x00007fff501fad42 UnrealEditor-Core.dll!FNamedTaskThread::ProcessTasksNamedThread() [Unknown File]\n"
+			+ "	0x00007fff501fb25e UnrealEditor-Core.dll!FNamedTaskThread::ProcessTasksUntilQuit() [Unknown File]\n"
+			+ "	0x00007fff48ccbb54 UnrealEditor-RenderCore.dll!RenderingThreadMain() [Unknown File]\n"
+			+ "	0x00007fff48ccff44 UnrealEditor-RenderCore.dll!FRenderingThread::Run() [Unknown File]\n"
+			+ "	0x00007fff5089c862 UnrealEditor-Core.dll!FRunnableThreadWin::Run() [Unknown File]\n"
+			+ "	0x00007fff5089a7bf UnrealEditor-Core.dll!FRunnableThreadWin::GuardedRun() [Unknown File]\n"
+			+ "	0x00007fff8e304ed0 KERNEL32.DLL!UnknownFunction [Unknown File]\n"
+			+ "	0x00007fff8f26e39b ntdll.dll!UnknownFunction [Unknown File]\n";
+			string logAlternateCallstack = "	0x00007fff4b43a1b4 UnrealEditor-RHI.dll!FValidationContext::RHISetShaderParameters() [Unknown File]\n"
+			+ "	0x00007fff4b3e5a01 UnrealEditor-RHI.dll!FRHICommandSetShaderParameters<FRHIGraphicsShader>::Execute() [Unknown File]\n"
+			+ "	0x00007fff4b3e9a5a UnrealEditor-RHI.dll!FRHICommand<FRHICommandSetShaderParameters<FRHIGraphicsShader>,FRHICommandSetShaderParametersString1159>::ExecuteAndDestruct() [Unknown File]\n"
+			+ "	0x00007fff4b3e74f7 UnrealEditor - RHI.dll!FRHICommandListBase::Execute()[Unknown File]\n";
+
 			// #1
 			// Scenario: Gauntlet test event with a property
 			// Expected: Gauntlet fingerprint using step identifiers
@@ -2246,45 +2278,17 @@ namespace HordeServer.Tests.Issues
 			// Scenario: Gauntlet Fatal event
 			// Expected: Gauntlet fingerprint using hash prefix and generate only one issue as they share the same summary
 			{
-				string logMessage = "Assertion failed: State.bGfxPSOSet [File:D:\\build\\U5M+Inc\\Sync\\Engine\\Source\\Runtime\\RHI\\Public\\RHIValidationContext.h] [Line: 809]\n"
-				+ "A Graphics PSO has to be set to set resources into a shader!";
-				string logCallstack = "	0x00007fff4b43a1f4 UnrealEditor-RHI.dll!FValidationContext::RHISetShaderParameters() [Unknown File]\n"
-				+ "	0x00007fff4b3e5a01 UnrealEditor-RHI.dll!FRHICommandSetShaderParameters<FRHIGraphicsShader>::Execute() [Unknown File]\n"
-				+ "	0x00007fff4b3e9a5a UnrealEditor-RHI.dll!FRHICommand<FRHICommandSetShaderParameters<FRHIGraphicsShader>,FRHICommandSetShaderParametersString1159>::ExecuteAndDestruct() [Unknown File]\n"
-				+ "	0x00007fff4b3e74f7 UnrealEditor - RHI.dll!FRHICommandListBase::Execute()[Unknown File]\n"
-				+ "	0x00007fff4b3f3228 UnrealEditor-RHI.dll!FRHICommandListImmediate::ExecuteAndReset() [Unknown File]\n"
-				+ "	0x00007fff4b464ebf UnrealEditor-RHI.dll!FRHIComputeCommandList::SubmitCommandsHint() [Unknown File]\n"
-				+ "	0x00007fff3d3066fc UnrealEditor-Renderer.dll!TBaseStaticDelegateInstance<void [Unknown File]\n"
-				+ "	0x00007fff48bd7813 UnrealEditor-RenderCore.dll!FRDGBuilder::ExecutePass() [Unknown File]\n"
-				+ "	0x00007fff48bd2a8a UnrealEditor-RenderCore.dll!FRDGBuilder::Execute() [Unknown File]\n"
-				+ "	0x00007fff3d31fef4 UnrealEditor-Renderer.dll!FSceneRenderer::RenderThreadEnd() [Unknown File]\n"
-				+ "	0x00007fff3d2ed57a UnrealEditor-Renderer.dll!`FPixelShaderUtils::AddFullscreenPass<FHZBTestPS>'::`2'::<lambda_1>::operator()() [Unknown File]\n"
-				+ "	0x00007fff3d3043f9 UnrealEditor - Renderer.dll!FSceneRenderer::DoOcclusionQueries()[Unknown File]\n"
-				+ "	0x00007fff3d30c888 UnrealEditor-Renderer.dll!TBaseStaticDelegateInstance<void [Unknown File]\n"
-				+ "	0x00007fff501fad42 UnrealEditor-Core.dll!FNamedTaskThread::ProcessTasksNamedThread() [Unknown File]\n"
-				+ "	0x00007fff501fb25e UnrealEditor-Core.dll!FNamedTaskThread::ProcessTasksUntilQuit() [Unknown File]\n"
-				+ "	0x00007fff48ccbb54 UnrealEditor-RenderCore.dll!RenderingThreadMain() [Unknown File]\n"
-				+ "	0x00007fff48ccff44 UnrealEditor-RenderCore.dll!FRenderingThread::Run() [Unknown File]\n"
-				+ "	0x00007fff5089c862 UnrealEditor-Core.dll!FRunnableThreadWin::Run() [Unknown File]\n"
-				+ "	0x00007fff5089a7bf UnrealEditor-Core.dll!FRunnableThreadWin::GuardedRun() [Unknown File]\n"
-				+ "	0x00007fff8e304ed0 KERNEL32.DLL!UnknownFunction [Unknown File]\n"
-				+ "	0x00007fff8f26e39b ntdll.dll!UnknownFunction [Unknown File]\n";
-				string logAlternateCallstack = "	0x00007fff4b43a1b4 UnrealEditor-RHI.dll!FValidationContext::RHISetShaderParameters() [Unknown File]\n"
-				+ "	0x00007fff4b3e5a01 UnrealEditor-RHI.dll!FRHICommandSetShaderParameters<FRHIGraphicsShader>::Execute() [Unknown File]\n"
-				+ "	0x00007fff4b3e9a5a UnrealEditor-RHI.dll!FRHICommand<FRHICommandSetShaderParameters<FRHIGraphicsShader>,FRHICommandSetShaderParametersString1159>::ExecuteAndDestruct() [Unknown File]\n"
-				+ "	0x00007fff4b3e74f7 UnrealEditor - RHI.dll!FRHICommandListBase::Execute()[Unknown File]\n";
-
 				IJob job = CreateJob(MainStreamId, 120, "Test Build", Graph);
 				await using (TestJsonLogger logger = await CreateLoggerAsync(job, 0, 0))
 				{
-					logger.LogError(KnownLogEvents.Gauntlet_FatalEvent, "{Summary}\n{Callstack}", logMessage, logCallstack);
+					logger.LogError(KnownLogEvents.Gauntlet_FatalEvent, "{Summary}\n{Callstack}", logAssertionMessage, logCallstack);
 				}
 				await UpdateCompleteStepAsync(job, 0, 0, JobStepOutcome.Failure);
 				await using (TestJsonLogger logger = await CreateLoggerAsync(job, 0, 1))
 				{
 					Dictionary<string, object?> properties = new Dictionary<string, object?>
 					{
-						["Summary"] = logMessage,
+						["Summary"] = logAssertionMessage,
 						["Callstack"] = logAlternateCallstack
 					};
 					foreach (JsonLogEvent outJsonLogEvent in MultilineLogEvent(LogLevel.Error, KnownLogEvents.Gauntlet_FatalEvent, "Engine encounter a critical failure\n{Summary}\n{Callstack}", properties))
@@ -2295,7 +2299,7 @@ namespace HordeServer.Tests.Issues
 				await UpdateCompleteStepAsync(job, 0, 1, JobStepOutcome.Failure);
 				await using (TestJsonLogger logger = await CreateLoggerAsync(job, 0, 2))
 				{
-					LogValue logSummary = new LogValue(new Utf8String("Summary"), logMessage);
+					LogValue logSummary = new LogValue(new Utf8String("Summary"), logAssertionMessage);
 					logger.LogError(KnownLogEvents.Gauntlet_FatalEvent, "{Summary}\n{Callstack}", logSummary, logAlternateCallstack);
 				}
 				await UpdateCompleteStepAsync(job, 0, 2, JobStepOutcome.Failure);
@@ -2306,8 +2310,48 @@ namespace HordeServer.Tests.Issues
 				Assert.AreEqual("Gauntlet:fatal:with-callstack:", issues[0].Fingerprints[0].Type.Substring(0, 30));
 				Assert.AreEqual("hash:", issues[0].Fingerprints[0].Keys.First().Name.Substring(0, 5));
 				Assert.AreEqual(37, issues[0].Fingerprints[0].Keys.First().Name.Length);
-				Assert.AreEqual("Automation fatal errors in Update Version Files, Compile UnrealHeaderTool Win64 and Compile ShooterGameEditor Win64", issues[0].Summary);
+				Assert.AreEqual("Automation fatal callstack errors in Update Version Files, Compile UnrealHeaderTool Win64 and Compile ShooterGameEditor Win64", issues[0].Summary);
 			}
+			// #2.5
+			// Scenario: Gauntlet ensure event
+			// Expected: Gauntlet fingerprint using hash prefix and generate only one issue as they share the same summary
+			{
+				IJob job = CreateJob(MainStreamId, 120, "Test Build", Graph);
+				await using (TestJsonLogger logger = await CreateLoggerAsync(job, 0, 0))
+				{
+					logger.LogWarning(KnownLogEvents.Gauntlet, "{Summary}\n{Callstack}", logEnsureMessage, logCallstack);
+				}
+				await UpdateCompleteStepAsync(job, 0, 0, JobStepOutcome.Warnings);
+				await using (TestJsonLogger logger = await CreateLoggerAsync(job, 0, 1))
+				{
+					Dictionary<string, object?> properties = new Dictionary<string, object?>
+					{
+						["Summary"] = logEnsureMessage,
+						["Callstack"] = logAlternateCallstack
+					};
+					foreach (JsonLogEvent outJsonLogEvent in MultilineLogEvent(LogLevel.Warning, KnownLogEvents.Gauntlet, "Engine encounter a critical failure\n{Summary}\n{Callstack}", properties))
+					{
+						logger.LogJsonLogEvent(outJsonLogEvent);
+					}
+				}
+				await UpdateCompleteStepAsync(job, 0, 1, JobStepOutcome.Warnings);
+				await using (TestJsonLogger logger = await CreateLoggerAsync(job, 0, 2))
+				{
+					LogValue logSummary = new LogValue(new Utf8String("Summary"), logEnsureMessage);
+					logger.LogWarning(KnownLogEvents.Gauntlet, "{Summary}\n{Callstack}", logSummary, logAlternateCallstack);
+				}
+				await UpdateCompleteStepAsync(job, 0, 2, JobStepOutcome.Warnings);
+
+				IReadOnlyList<IIssue> issues = await IssueCollection.FindIssuesAsync();
+				Assert.AreEqual(2, issues.Count);
+				Assert.AreEqual(1, issues[0].Fingerprints.Count);
+				Assert.AreEqual(true, issues[0].Fingerprints[0].Type.StartsWith("Gauntlet:test framework:with-callstack:", StringComparison.OrdinalIgnoreCase));
+				Assert.AreEqual("hash:", issues[0].Fingerprints[0].Keys.First().Name.Substring(0, 5));
+				Assert.AreEqual(37, issues[0].Fingerprints[0].Keys.First().Name.Length);
+				Assert.AreEqual("Automation test framework ensure callstack warnings in Update Version Files, Compile UnrealHeaderTool Win64 and Compile ShooterGameEditor Win64", issues[0].Summary);
+				Assert.AreEqual("Automation fatal callstack errors in Update Version Files, Compile UnrealHeaderTool Win64 and Compile ShooterGameEditor Win64", issues[1].Summary);
+			}
+
 			// #3
 			// Scenario: Gauntlet Test
 			// Expected: Gauntlet Test fingerprint using hash prefix and job step salt
diff --git a/Engine/Source/Programs/Shared/EpicGames.Horde/Issues/Handlers/GauntletIssueHandler.cs b/Engine/Source/Programs/Shared/EpicGames.Horde/Issues/Handlers/GauntletIssueHandler.cs
--- a/Engine/Source/Programs/Shared/EpicGames.Horde/Issues/Handlers/GauntletIssueHandler.cs
+++ b/Engine/Source/Programs/Shared/EpicGames.Horde/Issues/Handlers/GauntletIssueHandler.cs
@@ -161,13 +161,11 @@ namespace EpicGames.Horde.Issues.Handlers
 		/// <param name="keys">Receives a set of the keys</param>
 		/// <param name="metadata">Receives a set of metadata</param>
 		/// <param name="hasCallstack">Set true if a callstack property was found</param>
-		private void GetHash(IssueEvent issueEvent, HashSet<IssueKey> keys, HashSet<IssueMetadata> metadata, out bool hasCallstack)
+		private void GetHash(IssueEvent issueEvent, HashSet<IssueKey> keys, HashSet<IssueMetadata> metadata, bool hasCallstack)
 		{
-			hasCallstack = false;
 			if (TryGetHash(issueEvent, out Md5Hash hash))
 			{
-				string key = $"hash:{hash}";
-				hasCallstack = EventHasCallstackProperty(issueEvent);
+				string key = $"hash:{hash}";				
 				if (!hasCallstack)
 				{
 					// add job step salt if no Callstack property was found
@@ -285,9 +283,11 @@ namespace EpicGames.Horde.Issues.Handlers
 					return true;
 				}
 
-				if (issueEvent.Severity == LogLevel.Warning)
+				bool hasCallstack = EventHasCallstackProperty(issueEvent);
+
+				if (issueEvent.Severity == LogLevel.Warning && !hasCallstack)
 				{
-					// Do not generate a flood of Build Health issues for smallset of  known Gauntlet log event warnings
+					// Do not generate a flood of Build Health issues for smallset of known Gauntlet log event warnings
 					// If we identify cases where we want these warnings, they should be defined as new log events, and force their warning in the issue handler, here
 					// Though, we have been unable to come up with a single case in exploration					
 					return true;
@@ -303,11 +303,24 @@ namespace EpicGames.Horde.Issues.Handlers
 					return true;
 				}
 
+				bool isEnsure = false;
+				if (hasCallstack)
+				{
+					string rendered = issueEvent.Render();
+					if (issueEvent.Render().Contains("Ensure condition failed", StringComparison.OrdinalIgnoreCase))
+					{
+						isEnsure = true;
+					}
+				}
+
+				string callstack = isEnsure ? "ensure callstack" : "callstack";
+
+				string summaryTemplate = hasCallstack ? "Automation {Meta:GauntletType} " + callstack + " {Severity} in {Meta:Node}" : "Automation {Meta:GauntletType} {Severity} in {Meta:Node}";
+
 				string gauntletType = GetEventPrefix(issueEvent.EventId!.Value);
-				IssueEventGroup issue = new IssueEventGroup($"Gauntlet:{gauntletType}", "Automation {Meta:GauntletType} {Severity} in {Meta:Node}", IssueChangeFilter.All);
+				IssueEventGroup issue = new IssueEventGroup($"Gauntlet:{gauntletType}",summaryTemplate, IssueChangeFilter.All);
 				issue.Events.Add(issueEvent);
-				bool hasCallstack;
-				GetHash(issueEvent, issue.Keys, issue.Metadata, out hasCallstack);
+				GetHash(issueEvent, issue.Keys, issue.Metadata, hasCallstack);
 				issue.Metadata.Add("GauntletType", gauntletType);
 				if (hasCallstack)
 				{
