Horde: Add extended bearer token ACL
#submitool bypass

diff --git a/Engine/Source/Programs/Horde/HordeServer.Shared/Server/ServerAclAction.cs b/Engine/Source/Programs/Horde/HordeServer.Shared/Server/ServerAclAction.cs
--- a/Engine/Source/Programs/Horde/HordeServer.Shared/Server/ServerAclAction.cs
+++ b/Engine/Source/Programs/Horde/HordeServer.Shared/Server/ServerAclAction.cs
@@ -28,5 +28,11 @@ namespace HordeServer.Server
 		/// Issue bearer token for the current user
 		/// </summary>
 		public static AclAction IssueBearerToken { get; } = new AclAction("IssueBearerToken");
+
+		/// <summary>
+		/// Issue bearer token for the current user
+		/// </summary>
+		public static AclAction IssueExtendedBearerToken { get; } = new AclAction("IssueExtendedBearerToken");
+
 	}
 }
diff --git a/Engine/Source/Programs/Horde/HordeServer/Server/AdminController.cs b/Engine/Source/Programs/Horde/HordeServer/Server/AdminController.cs
--- a/Engine/Source/Programs/Horde/HordeServer/Server/AdminController.cs
+++ b/Engine/Source/Programs/Horde/HordeServer/Server/AdminController.cs
@@ -77,12 +77,17 @@ namespace HordeServer.Server
 		[Route("/api/v1/admin/token")]
 		public async Task<ActionResult<string>> GetTokenAsync(CancellationToken cancellationToken)
 		{
-			if (!_globalConfig.Value.Authorize(ServerAclAction.IssueBearerToken, User))
+			bool issueBearerToken = _globalConfig.Value.Authorize(ServerAclAction.IssueBearerToken, User);
+			bool issueExtendedBearer = _globalConfig.Value.Authorize(ServerAclAction.IssueExtendedBearerToken, User); 
+
+			if (!issueBearerToken && !issueExtendedBearer)
 			{
 				return Forbid(ServerAclAction.IssueBearerToken);
 			}
 
-			return await _aclService.IssueBearerTokenAsync(User.Claims, GetDefaultExpiryTime(), cancellationToken);
+			TimeSpan? expires = GetDefaultExpiryTime(issueExtendedBearer);
+
+			return await _aclService.IssueBearerTokenAsync(User.Claims, expires, cancellationToken);
 		}
 
 		/// <summary>
@@ -213,16 +218,19 @@ namespace HordeServer.Server
 
 		/// <summary>
 		/// Gets the default expiry time for a token
-		/// </summary>
+		/// </summary>		
+		/// <param name="extended">Whether to use the extended bearer expiryTime</param>
 		/// <returns></returns>
-		private TimeSpan? GetDefaultExpiryTime()
+		private TimeSpan? GetDefaultExpiryTime(bool extended = false)
 		{
 			ServerSettings serverSettings = _serverSettings.Value;
 
+			int expireyHours = extended ? serverSettings.ExtendedJwtExpiryTimeHours : serverSettings.JwtExpiryTimeHours;
+
 			TimeSpan? expiryTime = null;
-			if (serverSettings.JwtExpiryTimeHours != -1)
+			if (expireyHours != -1)
 			{
-				expiryTime = TimeSpan.FromHours(serverSettings.JwtExpiryTimeHours);
+				expiryTime = TimeSpan.FromHours(expireyHours);
 			}
 			return expiryTime;
 		}
diff --git a/Engine/Source/Programs/Horde/HordeServer/ServerSettings.cs b/Engine/Source/Programs/Horde/HordeServer/ServerSettings.cs
--- a/Engine/Source/Programs/Horde/HordeServer/ServerSettings.cs
+++ b/Engine/Source/Programs/Horde/HordeServer/ServerSettings.cs
@@ -322,6 +322,11 @@ namespace HordeServer
 		/// </summary>
 		public int JwtExpiryTimeHours { get; set; } = 8;
 
+		/// <summary>
+		/// Length of time before extended JWT tokens expire, in hours, defaults to 1 week
+		/// </summary>
+		public int ExtendedJwtExpiryTimeHours { get; set; } = 168;
+
 		/// <summary>
 		/// The claim type for administrators
 		/// </summary>
