[Horde][BuildHealth][Dashboard]  Apply tree filtering to current APIs

diff --git a/Engine/Source/Programs/Horde/HordeDashboard/src/components/buildhealth/stepoutcome/StepOutcomeDataHandler.ts b/Engine/Source/Programs/Horde/HordeDashboard/src/components/buildhealth/stepoutcome/StepOutcomeDataHandler.ts
--- a/Engine/Source/Programs/Horde/HordeDashboard/src/components/buildhealth/stepoutcome/StepOutcomeDataHandler.ts
+++ b/Engine/Source/Programs/Horde/HordeDashboard/src/components/buildhealth/stepoutcome/StepOutcomeDataHandler.ts
@@ -106,7 +106,8 @@ export interface StepOutcomeFilters {
  * @todo Data sharing & caching would ideally occur between the BuildHealth View & Step Outcome.
  */
 export class StepOutcomeDataHandler extends PollBase {
-    private readonly STREAM_JOBS_FILTER = "id,streamId,name,change,templateId,state,createTime,updateTime,batches,preflightCommitId,preflightDescription,cancellationReason";
+    private readonly STREAM_JOBS_FILTER = "id,streamId,name,change,templateId,state,createTime,updateTime,preflightCommitId,preflightDescription,cancellationReason,labels,defaultLabel";
+    private readonly STREAM_JOBS_BATCH_STEP_FILTER = "batches.steps.id,batches.steps.name,batches.steps.state,batches.steps.outcome,batches.steps.startTime,batches.steps.finishTime,batches.steps.error"
     private readonly STREAM_CHANGES_FILTER = "id,dateUtc";
     private readonly MAX_JOB_COUNT_DEFAULT = 150;
 
@@ -204,7 +205,7 @@ export class StepOutcomeDataHandler extends PollBase {
                 const query: JobStreamQuery = {
                     template: [templateId],
                     count: jobCount,
-                    filter: this.STREAM_JOBS_FILTER,
+                    filter: `${this.STREAM_JOBS_FILTER},${this.STREAM_JOBS_BATCH_STEP_FILTER}`,
                     modifiedAfter: modifiedAfter
                 };
                 queries.push({ streamId, query });
