Updated RCA information to be included in the issue telemetry.

diff --git a/Engine/Source/Programs/Horde/Plugins/Build/HordeServer.Build/Issues/IssueTelemetry.cs b/Engine/Source/Programs/Horde/Plugins/Build/HordeServer.Build/Issues/IssueTelemetry.cs
--- a/Engine/Source/Programs/Horde/Plugins/Build/HordeServer.Build/Issues/IssueTelemetry.cs
+++ b/Engine/Source/Programs/Horde/Plugins/Build/HordeServer.Build/Issues/IssueTelemetry.cs
@@ -14,7 +14,7 @@ namespace HordeServer.Issues
 	/// Record used for Issue telemetry. Models a <see cref="IIssue"/>.
 	/// </summary>
 	/// <remarks>Modifications to this record's structure should be versioned appropriately through <see cref="SchemaVersion"/> via <see cref="CurrentSchemaVersion"/>, and call sites should continue to publish in a backwards compatible manner.</remarks>
-	public record IssueTelemetry(string EventName, int SchemaVersion, int Id, DateTime? AcknowledgedAt, DateTime CreatedAt, UserId? OwnerId, string? ExternalIssueKey, CommitId? FixChange, bool FixedSystemic, CommitId? RootChange, DateTime LastSeenAt, DateTime? NominatedAt, UserId? NominatedById, DateTime? ResolvedAt, UserId? ResolvedById, IssueSeverity Severity, string Summary, DateTime? VerifiedAt)
+	public record IssueTelemetry(string EventName, int SchemaVersion, int Id, DateTime? AcknowledgedAt, DateTime CreatedAt, UserId? OwnerId, string? ExternalIssueKey, CommitId? FixChange, bool FixedSystemic, CommitId? RootChange, DateTime LastSeenAt, DateTime? NominatedAt, UserId? NominatedById, DateTime? ResolvedAt, UserId? ResolvedById, IssueSeverity Severity, string Summary, DateTime? VerifiedAt, int? DuplicateIssueId, string? RootCauseCategory, string? RootCauseSummary)
 	{
 		/// <summary>
 		/// Default event name for the IssueTelemetry.
@@ -48,7 +48,10 @@ namespace HordeServer.Issues
 					issue.ResolvedById,
 					issue.Severity,
 					issue.Summary,
-					issue.VerifiedAt
+					issue.VerifiedAt,
+					issue.DuplicateIssueId,
+					issue.RootCauseCategory,
+					issue.RootCauseSummary
 				  )
 		{
 		}
