Bumping SchemaVersion now we have added RCA data.

diff --git a/Engine/Source/Programs/Horde/Plugins/Build/HordeServer.Build/Issues/IssueTelemetry.cs b/Engine/Source/Programs/Horde/Plugins/Build/HordeServer.Build/Issues/IssueTelemetry.cs
--- a/Engine/Source/Programs/Horde/Plugins/Build/HordeServer.Build/Issues/IssueTelemetry.cs
+++ b/Engine/Source/Programs/Horde/Plugins/Build/HordeServer.Build/Issues/IssueTelemetry.cs
@@ -24,7 +24,7 @@ namespace HordeServer.Issues
 		/// <summary>
 		/// Current schema version.
 		/// </summary>
-		public static readonly int CurrentSchemaVersion = 1;
+		public static readonly int CurrentSchemaVersion = 2;
 
 		/// <summary>
 		/// Telemetry payload for <see cref="IIssue"/>.
