Horde: Add RCA Change Backout category

diff --git a/Engine/Source/Programs/Horde/HordeDashboard/src/components/IssueViewV2.tsx b/Engine/Source/Programs/Horde/HordeDashboard/src/components/IssueViewV2.tsx
--- a/Engine/Source/Programs/Horde/HordeDashboard/src/components/IssueViewV2.tsx
+++ b/Engine/Source/Programs/Horde/HordeDashboard/src/components/IssueViewV2.tsx
@@ -2320,6 +2320,7 @@ const rcaCategoryOptions: IDropdownOption[] = [
    { key: 'rca_monitoring', text: 'Monitoring/No Code Change' },
    { key: 'rca_use', text: 'Unintended Side Effect' },
    { key: 'rca_buildagent', text: 'Build Agent' },
+   { key: 'rca_changebackout', text: 'Change Backout' },
    { key: 'rca_other', text: 'Other' },
    { key: 'rca_clear', text: 'Clear Selection' }
 ];
diff --git a/Engine/Source/Programs/Horde/Plugins/Build/HordeServer.Build/Notifications/Sinks/SlackNotificationSink.cs b/Engine/Source/Programs/Horde/Plugins/Build/HordeServer.Build/Notifications/Sinks/SlackNotificationSink.cs
--- a/Engine/Source/Programs/Horde/Plugins/Build/HordeServer.Build/Notifications/Sinks/SlackNotificationSink.cs
+++ b/Engine/Source/Programs/Horde/Plugins/Build/HordeServer.Build/Notifications/Sinks/SlackNotificationSink.cs
@@ -3521,6 +3521,7 @@ namespace HordeServer.Notifications.Sinks
 				   new SlackOption( "Monitoring/No Code Change", "Monitoring/No Code Change"),
 				   new SlackOption( "Unintended Side Effect", "Unintended Side Effect"),
 				   new SlackOption( "Build Agent", "Build Agent"),
+				   new SlackOption( "Change Backout", "Change Backout"),
 				   new SlackOption( "Other", "Other"),				   
 				};						
 
