[Horde][BuildHealth][Dashboard] StepOutcome table supports N values for the same cell

diff --git a/Engine/Source/Programs/Horde/HordeDashboard/src/components/buildhealth/stepoutcome/StepOutcome.tsx b/Engine/Source/Programs/Horde/HordeDashboard/src/components/buildhealth/stepoutcome/StepOutcome.tsx
--- a/Engine/Source/Programs/Horde/HordeDashboard/src/components/buildhealth/stepoutcome/StepOutcome.tsx
+++ b/Engine/Source/Programs/Horde/HordeDashboard/src/components/buildhealth/stepoutcome/StepOutcome.tsx
@@ -1,6 +1,6 @@
 // Copyright Epic Games, Inc. All Rights Reserved.
 
-import { Dialog, DialogType, Label, mergeStyleSets, ScrollablePane, Spinner, SpinnerSize, Stack } from "@fluentui/react";
+import { Dialog, DialogType, Label, mergeStyleSets, Spinner, SpinnerSize, Stack } from "@fluentui/react";
 import { Text } from '@fluentui/react/lib/Text';
 import { observer } from "mobx-react-lite";
 import React, { useEffect, useState } from 'react';
@@ -54,7 +54,8 @@ const StepOutcomeFocusModal: React.FC<StepOutcomeFocusModalProps> = ({ item, onD
     }, []);
 
     const details: (string | JSX.Element)[][] = [
-        ['Date', startDate ? startDate.toDateString() : "N/A"],
+        ['Date (Job)', item.jobCreateTime ? new Date(item.jobCreateTime).toLocaleString() : "N/A"],
+        ['Date (Step)', startDate ? startDate.toLocaleString() : "N/A"],
         ['Stream', item.streamId],
         ['Code Change', item.change],
         ['Job', item.jobName],
@@ -67,6 +68,8 @@ const StepOutcomeFocusModal: React.FC<StepOutcomeFocusModalProps> = ({ item, onD
         ['Step Link', <a href={`/job/${item.jobId}?step=${item.stepResponse.id}`} target="_blank" rel="noopener noreferrer">View Step</a>],
     ].filter((d): d is (string | JSX.Element)[] => Boolean(d));
 
+    let shouldRenderDuplicateColumn = item.duplicateEntries && item.duplicateEntries.length > 0;
+
     return (
         <Dialog
             modalProps={{
@@ -85,7 +88,7 @@ const StepOutcomeFocusModal: React.FC<StepOutcomeFocusModalProps> = ({ item, onD
             }}
             onDismiss={onDismiss}
             hidden={false}
-            minWidth={1200}
+            minWidth={1400}
             dialogContentProps={{
                 type: DialogType.close,
                 onDismiss: onDismiss,
@@ -94,24 +97,27 @@ const StepOutcomeFocusModal: React.FC<StepOutcomeFocusModalProps> = ({ item, onD
             <Stack>
                 <Stack horizontal tokens={{ childrenGap: 16 }} styles={{ root: { width: '100%', flexGrow: 0 } }}>
                     {/* Left Column - Step Details */}
-                    <Stack styles={{ root: { width: '35%' } }} tokens={{ childrenGap: 4 }}>
+                    <Stack styles={{ root: { width: '25%' } }} tokens={{ childrenGap: 4 }}>
+                        <Stack horizontal tokens={{ childrenGap: 12 }} className={stepOutcomeModalClasses.bottomBorder}>
+                            <Label styles={{ root: { padding: 0 } }}>Step Details:</Label>
+                        </Stack>
                         {details.map(([label, value], index) => (
-                            <Stack key={`modal_step_item_${index}`} horizontal tokens={{ childrenGap: 12 }}>
+                            <Stack key={`modal_step_item_${index}`} styles={{ root: { paddingLeft: 16 } }} horizontal tokens={{ childrenGap: 12 }}>
                                 <Label styles={{ root: { minWidth: 100, padding: 0 } }}>{label}:</Label>
                                 <Text>{value}</Text>
                             </Stack>
                         ))}
                     </Stack>
 
-                    {/* Right Column - Issues*/}
-                    <div style={{ width: '65%', maxHeight: 400, overflow: 'auto', position: 'relative' }}>
-                        <ScrollablePane styles={{ root: { height: '100%' } }}>
-                            <Stack tokens={{ childrenGap: 8 }}>
-                                <Stack horizontal tokens={{ childrenGap: 4 }} className={stepOutcomeModalClasses.bottomBorder}>
-                                    <Label styles={{ root: { padding: 0 } }}>Issue Count:</Label>
-                                    <Text>{item.issuesData.length}</Text>
-                                    <Text>[{item.issuesData.filter((issue) => !issue.resolvedAt).length} unresolved, {item.issuesData.filter((issue) => issue.resolvedAt).length} resolved]</Text>
-                                </Stack>
+                    {/* Center Column - Issues*/}
+                    <Stack style={{ width: shouldRenderDuplicateColumn ? '50%' : '75%', maxHeight: 400, display: 'flex', flexDirection: 'column' }}>
+                        <Stack horizontal tokens={{ childrenGap: 12 }} className={stepOutcomeModalClasses.bottomBorder}>
+                            <Label styles={{ root: { padding: 0 } }}>Issue Count:</Label>
+                            <Text>{item.issuesData.length}</Text>
+                            <Text>[{item.issuesData.filter((issue) => !issue.resolvedAt).length} unresolved, {item.issuesData.filter((issue) => issue.resolvedAt).length} resolved]</Text>
+                        </Stack>
+                        <Stack style={{ flex: 1, overflowY: 'auto' }}>
+                            <Stack>
                                 {item.issuesData.length > 0 ?
                                     <Stack tokens={{ childrenGap: 10 }} styles={{ root: { paddingLeft: 16 } }}>
                                         {item.issuesData.map((issue, index) => (
@@ -132,13 +138,48 @@ const StepOutcomeFocusModal: React.FC<StepOutcomeFocusModalProps> = ({ item, onD
                                             </Stack>
                                         ))}
                                     </Stack> : null
+
                                 }
                             </Stack>
-                        </ScrollablePane>
-                    </div>
+                        </Stack>
+                    </Stack>
+
+                    {/* Right Column - Related Steps*/}
+                    {shouldRenderDuplicateColumn && (
+                        <Stack style={{ width: '25%', maxHeight: 400, display: 'flex', flexDirection: 'column' }}>
+                            <Stack horizontal tokens={{ childrenGap: 12 }} className={stepOutcomeModalClasses.bottomBorder}>
+                                <Label styles={{ root: { padding: 0 } }}>Duplicate Steps:</Label>
+                            </Stack>
+                            <Stack style={{ flex: 1, overflowY: 'auto' }}>
+                                <Stack tokens={{ childrenGap: 10 }} styles={{ root: { paddingLeft: 16 } }}>
+                                    {item.duplicateEntries.map((relatedEntry, index) => (
+                                        <Stack
+                                            key={`related_job_item${index}`}
+                                            tokens={{ childrenGap: 4 }}
+                                            className={stepOutcomeModalClasses.bottomBorder}
+                                        >
+                                            {[
+                                                ['Date (Job)', relatedEntry.jobCreateTime ? new Date(relatedEntry.jobCreateTime).toLocaleString() : 'N/A'],
+                                                ['Date (Step)', relatedEntry.stepResponse.startTime ? new Date(relatedEntry.stepResponse.startTime).toLocaleString() : 'N/A'],
+                                                ['Job', relatedEntry.jobName],
+                                                ['Job Link', <a href={`/job/${relatedEntry.jobId}`} target="_blank" rel="noopener noreferrer">View Related Job</a>],
+                                                ['Step Link', <a href={`/job/${relatedEntry.jobId}?step=${relatedEntry.stepResponse.id}`} target="_blank" rel="noopener noreferrer">View Related Step</a>],
+                                                ['State', relatedEntry.stepResponse.state],
+                                            ].map(([label, value], j) => (
+                                                <Stack key={j} horizontal tokens={{ childrenGap: 4 }}>
+                                                    <Label styles={{ root: { minWidth: 80, padding: 0 } }}>{label}:</Label>
+                                                    <Text>{value}</Text>
+                                                </Stack>
+                                            ))}
+                                        </Stack>
+                                    ))}
+                                </Stack>
+                            </Stack>
+                        </Stack>
+                    )}
                 </Stack>
-            </Stack>
-        </Dialog>
+            </Stack >
+        </Dialog >
     );
 };
 
diff --git a/Engine/Source/Programs/Horde/HordeDashboard/src/components/buildhealth/stepoutcome/StepOutcomeDataTypes.ts b/Engine/Source/Programs/Horde/HordeDashboard/src/components/buildhealth/stepoutcome/StepOutcomeDataTypes.ts
--- a/Engine/Source/Programs/Horde/HordeDashboard/src/components/buildhealth/stepoutcome/StepOutcomeDataTypes.ts
+++ b/Engine/Source/Programs/Horde/HordeDashboard/src/components/buildhealth/stepoutcome/StepOutcomeDataTypes.ts
@@ -122,6 +122,7 @@ export class StepOutcomeTableEntry {
     streamId: string;
     jobCreateTime: string;
     issuesData: IssueData[] = [];
+    duplicateEntries: StepOutcomeTableEntry[] = [];
 
     // #endregion -- Public Members --
 
@@ -180,7 +181,7 @@ export class StepOutcomeTableEntry {
      * @returns CSV representation of the table entry.
      */
     toCSVRepresntation(relatedJobContext: JobResponseData | undefined): string {
-        let baseUrl =  window.location.origin;
+        let baseUrl = window.location.origin;
         let jobUrl = `${baseUrl}/job/${this.jobId}`;
         let stepUrl = `${baseUrl}/job/${this.jobId}?step=${this.stepResponse.id}`;
         return `${this.change},${this.streamId},${this.jobName},${this.jobId},${this.stepResponse.name},${this.stepResponse.id},${this.stepResponse.state},${this.stepResponse.outcome},${this.stepResponse.stepDurationMs / 1000},${this.jobCreateTime},${this.stepResponse.startTime},${relatedJobContext?.summarizedOutcome},${relatedJobContext?.state},${jobUrl},${stepUrl}\n`;
@@ -432,11 +433,33 @@ export class StepOutcomeTable {
         }
 
         const existingEntry = this.tableEntries[insertionRow][insertionCol];
+
+        // If we have an existing entry, replace the existing entry with the newly added one if it is more recent.
         if (existingEntry !== null) {
-            console.warn(`Duplicate entry for step new (${tableEntry.streamId}::${tableEntry.jobName}::${tableEntry.jobId}::${tableEntry.stepResponse.name} (pre-existing: ${existingEntry.streamId}::${existingEntry.jobName}::${existingEntry.jobId}::${existingEntry.stepResponse.name})) and change ("${tableEntry.change}" (${existingEntry.change})). Ambiguous selector needed.`, {
-                existingEntry,
-                newEntry: tableEntry
-            });
+            // Attempt to sort by step sort time first. This fails under two scenarios:
+            // 1. The step has Not Started (so the startTime will be null and thus, a delay in this taking top position)
+            // 2. The step has been skipped (and we must fall back to job creation time as a sorting mechanism - startTime will never be !null).
+            let startA: number = tableEntry.stepResponse.startTime ? new Date(tableEntry.stepResponse.startTime).getTime() : 0;
+            let startB: number = existingEntry.stepResponse.startTime ? new Date(existingEntry.stepResponse.startTime).getTime() : 0;
+
+            // If either is invalid (most likely due to the Not Started case, but secondarily if it's skipped/aborted), fall back to considering the job created time.
+            // Note: this will not be captured on the edge case of step-retry only. The best recommendation for this edge case is for users is to review the job details which they will have access
+            // to from the modal.
+            if (startA === 0 || startB === 0) {
+                startA = new Date(tableEntry.jobCreateTime).getTime();
+                startB = new Date(existingEntry.jobCreateTime).getTime();
+            }
+
+            const isNewer = startA > startB;
+
+            if (isNewer) {
+                tableEntry.duplicateEntries.push(existingEntry, ...existingEntry.duplicateEntries);
+                existingEntry.duplicateEntries = []; // we can zero this out; we've swapped to the table entry being the primary one.
+                this.tableEntries[insertionRow][insertionCol] = tableEntry;
+            } else {
+                existingEntry.duplicateEntries.push(tableEntry);
+            }
+
             return;
         }
 
diff --git a/Engine/Source/Programs/Horde/HordeDashboard/src/components/buildhealth/stepoutcome/StepOutcomeSharedUIComponents.tsx b/Engine/Source/Programs/Horde/HordeDashboard/src/components/buildhealth/stepoutcome/StepOutcomeSharedUIComponents.tsx
--- a/Engine/Source/Programs/Horde/HordeDashboard/src/components/buildhealth/stepoutcome/StepOutcomeSharedUIComponents.tsx
+++ b/Engine/Source/Programs/Horde/HordeDashboard/src/components/buildhealth/stepoutcome/StepOutcomeSharedUIComponents.tsx
@@ -3,8 +3,8 @@
 import { mergeStyleSets, Spinner, SpinnerSize } from "@fluentui/react";
 import dashboard, { StatusColor } from "horde/backend/Dashboard";
 import { observer } from "mobx-react-lite";
-import { StepOutcomeDataHandler, StepOutcomeViewOptions } from "./StepOutcomeDataHandler";
-import { AbortedStr, FailureStr, JobResponseData, RunningStr, SkippedStr, StepOutcomeTableEntry, SuccessStr, WarningStr } from "./StepOutcomeDataTypes";
+import { StepOutcomeDataHandler } from "./StepOutcomeDataHandler";
+import { AbortedStr, FailureStr, JobResponseData, SkippedStr, StepOutcomeTableEntry, SuccessStr, WarningStr } from "./StepOutcomeDataTypes";
 import { isTerminalState } from "./StepOutcomeUtilities";
 
 // #region -- Constants --
@@ -47,6 +47,23 @@ export const stepOutcomeTableClasses = mergeStyleSets({
     root: {
         maxWidth: "min(85vw, 2400px)",
     },
+    multiValue: {
+        position: "relative",
+        selectors: {
+            "&::after": {
+                content: "'*'",
+                position: "absolute",
+                top: 4,
+                right: -4,
+                color: dashboard.darktheme ? "#ffffff" : "#000000",
+                fontWeight: "bold",
+                fontSize: 25,
+                lineHeight: 1,
+                pointerEvents: "none",
+                userSelect: "none",
+            },
+        },
+    },
     ...Object.fromEntries(
         Object.entries(colors).map(([className, color]) => [
             className, // whatever key your getColorRecords() provides
@@ -101,7 +118,7 @@ export function resetCellClasses(): string[] {
  * @param jobState The job response data which will be used to style applicable table entries.
  * @returns The classes to apply to the cell.
  */
-export function getCellClasses(entry: StepOutcomeTableEntry | null,  jobState: JobResponseData | undefined): string[] {
+export function getCellClasses(entry: StepOutcomeTableEntry | null, jobState: JobResponseData | undefined): string[] {
     const classes: string[] = [];
 
     if (!entry || jobState?.summarizedOutcome === undefined) {
@@ -112,6 +129,10 @@ export function getCellClasses(entry: StepOutcomeTableEntry | null,  jobState: J
         classes.push(stepOutcomeTableClasses[jobState?.summarizedOutcome]);
     }
 
+    if (entry.duplicateEntries !== undefined && entry.duplicateEntries.length > 0) {
+        classes.push(stepOutcomeTableClasses.multiValue);
+    }
+
     return classes;
 }
 
