[Build Health Dashboard][Step Outcome] Introduce POST API for changes in order to support single request of batch data

diff --git a/Engine/Source/Programs/Horde/Plugins/Build/HordeServer.Build.Tests/Streams/StreamCollectionTests.cs b/Engine/Source/Programs/Horde/Plugins/Build/HordeServer.Build.Tests/Streams/StreamCollectionTests.cs
--- a/Engine/Source/Programs/Horde/Plugins/Build/HordeServer.Build.Tests/Streams/StreamCollectionTests.cs
+++ b/Engine/Source/Programs/Horde/Plugins/Build/HordeServer.Build.Tests/Streams/StreamCollectionTests.cs
@@ -114,5 +114,78 @@ namespace HordeServer.Tests.Streams
 				Assert.AreEqual(StreamsController.BadRequest().StatusCode, result.StatusCode);
 			}
 		}
+
+
+		[TestMethod]
+		public async Task ValidateGetBatchedChangesAsync()
+		{
+			Fixture fixture = await CreateFixtureAsync();
+			IUser jerry = await UserCollection.FindOrAddUserByLoginAsync("Jerry");
+
+			int firstChangelist = 112233;
+			int secondChangelist = 112234;
+			int thirdChangelist = 112235;
+			int invalidChangelist = 112236;
+
+			PerforceService.AddChange(fixture.StreamId, firstChangelist, jerry, "Edited file1.cpp.", new[] { "file.cpp" });
+			PerforceService.AddChange(fixture.StreamId, secondChangelist, jerry, "Added file2.cpp.", new[] { "file2.cpp" });
+			PerforceService.AddChange(fixture.StreamId, thirdChangelist, jerry, "Edited file1.cpp.", new[] { "file1.cpp" });
+
+			CommitId firstCommitId = new CommitId(firstChangelist.ToString());
+			CommitId secondCommitId = new CommitId(secondChangelist.ToString());
+			CommitId thirdCommitId = new CommitId(thirdChangelist.ToString());
+			CommitId invalidCommitId = new CommitId(invalidChangelist.ToString());
+
+			// Single element edge case
+			{
+				ActionResult<List<object>> results = await StreamsController.GetBatchedChangesAsync(fixture.StreamId, [firstCommitId]);
+
+				Assert.IsNotNull(results.Value);
+
+				List<GetCommitResponse> castedValue = [.. results.Value.Cast<GetCommitResponse>()];
+				Assert.IsTrue(castedValue.Count == 1);
+				Assert.IsNotNull(castedValue.SingleOrDefault(c => c.Id.Name == firstChangelist.ToString()));
+			}
+
+			// Multi element nominal case
+			{
+				ActionResult<List<object>> results = await StreamsController.GetBatchedChangesAsync(fixture.StreamId, [firstCommitId, secondCommitId]);
+
+				Assert.IsNotNull(results.Value);
+
+				List<GetCommitResponse> castedValue = [.. results.Value.Cast<GetCommitResponse>()];
+				Assert.IsTrue(castedValue.Count == 2);
+				Assert.IsNotNull(castedValue.SingleOrDefault(c => c.Id.Name == firstChangelist.ToString()));
+				Assert.IsNotNull(castedValue.SingleOrDefault(c => c.Id.Name == secondChangelist.ToString()));
+			}
+
+			// Partial success case
+			{
+				StreamsController httpContextPreservedStreamsController = StreamsController;
+				ActionResult<List<object>> results = await httpContextPreservedStreamsController.GetBatchedChangesAsync(fixture.StreamId,[firstCommitId, invalidCommitId]);
+
+				Assert.IsNotNull(results.Value);
+
+				List<GetCommitResponse> castedValue = [.. results.Value.Cast<GetCommitResponse>()];
+				Assert.IsTrue(castedValue.Count == 1);
+				Assert.IsNotNull(castedValue.SingleOrDefault(c => c.Id.Name == firstChangelist.ToString()));
+				Assert.IsNull(castedValue.SingleOrDefault(c => c.Id.Name == invalidChangelist.ToString()));
+				Assert.IsTrue(httpContextPreservedStreamsController.Response.Headers.ContainsKey(StreamsController.MissingCommitsHeader));
+
+				// Assert contents
+				StringValues stringValues = httpContextPreservedStreamsController.Response.Headers[StreamsController.MissingCommitsHeader];
+				Assert.IsTrue(stringValues.Contains(invalidChangelist.ToString()));
+			}
+
+			// Empty edge case
+			{
+				ActionResult<List<object>> results = await StreamsController.GetBatchedChangesAsync(fixture.StreamId, []);
+
+				Assert.IsNotNull(results.Value);
+
+				List<GetCommitResponse> castedValue = [.. results.Value.Cast<GetCommitResponse>()];
+				Assert.IsTrue(castedValue.Count == 0);
+			}
+		}
 	}
 }
diff --git a/Engine/Source/Programs/Horde/Plugins/Build/HordeServer.Build/Streams/StreamsController.cs b/Engine/Source/Programs/Horde/Plugins/Build/HordeServer.Build/Streams/StreamsController.cs
--- a/Engine/Source/Programs/Horde/Plugins/Build/HordeServer.Build/Streams/StreamsController.cs
+++ b/Engine/Source/Programs/Horde/Plugins/Build/HordeServer.Build/Streams/StreamsController.cs
@@ -235,15 +235,92 @@ namespace HordeServer.Streams
 			}
 
 			List<ICommit>? commits = null;
+			IReadOnlyCollection<string> failedCommits = [];
 
 			if (changeNumbers != null) // Batch commit request path.
 			{
-				commits = new List<ICommit>(changeNumbers.Length);
-				ConcurrentBag<string> failedCommits = [];
+				(commits, failedCommits) = await GetBatchedCommitsAsync(streamConfig, changeNumbers, cancellationToken);
+			}
+			else // Change range & tag request path.
+			{
+				commits = await _commitService.GetCollection(streamConfig).FindAsync(min, max, results, commitTags, cancellationToken).ToListAsync(cancellationToken);
+			}
+
+			// Store any failed commits (only populated through batch commit request path) into the response header.
+			if (failedCommits.Count > 0)
+			{
+				Response.Headers.Append(MissingCommitsHeader, String.Join(",", failedCommits));
+			}
+
+			List<GetCommitResponse> responses = new List<GetCommitResponse>(commits.Count);
+			foreach (ICommit commit in commits)
+			{
+				IUser? author = await _userCollection.GetCachedUserAsync(commit.AuthorId, cancellationToken);
+				responses.Add(CreateGetCommitResponse(commit, author!, null, null));
+			}
+			return responses.ConvertAll(x => PropertyFilter.Apply(x, filter));
+		}
+
+		/// <summary>
+		/// Gets a list of changes for a stream
+		/// </summary>
+		/// <param name="streamId">The stream id</param>
+		/// <param name="changeNumbers">Array of specific changelist IDs to retrieve for the stream.</param>
+		/// <param name="filter">The filter to apply to the results</param>
+		/// <param name="cancellationToken">Cancellation token for the operation</param>
+		/// Returns 200 OK with list of filtered commits. Partial failures for batch retrieval are indicated via the <see cref="MissingCommitsHeader"/> header.
+		/// <remarks>This endpoint is a workaround to retrieve large sets of changes which would result in 414 reponses from <see cref=GetChangesAsync"/>.
+		/// <remarks>Any commits that were not found or that threw exceptions are stored in the <see cref="MissingCommitsHeader"/> response header.</remarks>
+		[HttpPost]
+		[Route("/api/v1/streams/{streamId}/changes")]
+		[ProducesResponseType(typeof(List<GetCommitResponse>), 200)]
+		public async Task<ActionResult<List<object>>> GetBatchedChangesAsync(StreamId streamId, [FromBody] CommitId[]? changeNumbers = null,  PropertyFilter? filter = null, CancellationToken cancellationToken = default)
+		{
+			StreamConfig? streamConfig;
+			if (!_buildConfig.Value.TryGetStream(streamId, out streamConfig))
+			{
+				return NotFound(streamId);
+			}
+			if (!streamConfig.Authorize(StreamAclAction.ViewChanges, User))
+			{
+				return Forbid(StreamAclAction.ViewChanges, streamId);
+			}
+
+			List<ICommit>? commits;
+			IReadOnlyCollection<string> failedCommits = [];
+
+			if (changeNumbers != null)
+			{
+				(commits, failedCommits) = await GetBatchedCommitsAsync(streamConfig, changeNumbers, cancellationToken);
+			}
+			else
+			{
+				commits = new List<ICommit>();
+			}
+
+			// Store any failed commits into the response header.
+			if (failedCommits.Count > 0)
+			{
+				Response.Headers.Append(MissingCommitsHeader, String.Join(",", failedCommits));
+			}
+
+			List<GetCommitResponse> responses = new List<GetCommitResponse>(commits.Count);
+			foreach (ICommit commit in commits)
+			{
+				IUser? author = await _userCollection.GetCachedUserAsync(commit.AuthorId, cancellationToken);
+				responses.Add(CreateGetCommitResponse(commit, author!, null, null));
+			}
+			return responses.ConvertAll(x => PropertyFilter.Apply(x, filter));
+		}
 
-				// This is a pre-emptive measure as to not inundate our SCM with too many concurrent requests.
-				using SemaphoreSlim maxSCMActiveRequests = new SemaphoreSlim(MaxConcurrentSCMChangeRequests);
-				Task<ICommit?>[] commitServiceTasks = [.. changeNumbers.Select(async id =>
+		async Task<(List<ICommit>, IReadOnlyCollection<string>)> GetBatchedCommitsAsync(StreamConfig streamConfig, CommitId[] changeNumbers, CancellationToken cancellationToken)
+		{
+			List<ICommit> commits = new List<ICommit>(changeNumbers.Length);
+			ConcurrentBag<string> failedCommits = [];
+
+			// This is a pre-emptive measure as to not inundate our SCM with too many concurrent requests.
+			using SemaphoreSlim maxSCMActiveRequests = new SemaphoreSlim(MaxConcurrentSCMChangeRequests);
+			Task<ICommit?>[] commitServiceTasks = [.. changeNumbers.Select(async id =>
 				{
 					await maxSCMActiveRequests.WaitAsync(cancellationToken);
 					try
@@ -262,29 +339,10 @@ namespace HordeServer.Streams
 					}
 				})];
 
-				ICommit?[] commitServiceResults = await Task.WhenAll(commitServiceTasks);
-				commits = [.. commitServiceResults.Where(c => c != null).Select(c => c!)];
-
-				// Store any failed commits into the response header.
-				if (!failedCommits.IsEmpty)
-				{
-					Response.Headers.Append(MissingCommitsHeader, String.Join(",", failedCommits));
-				}
+			ICommit?[] commitServiceResults = await Task.WhenAll(commitServiceTasks);
+			commits = [.. commitServiceResults.Where(c => c != null).Select(c => c!)];
 
-				commits = [.. commits.OrderByDescending(x => x.Id)];
-			}
-			else // Change range & tag request path.
-			{
-				commits = await _commitService.GetCollection(streamConfig).FindAsync(min, max, results, commitTags, cancellationToken).ToListAsync(cancellationToken);
-			}
-
-			List<GetCommitResponse> responses = new List<GetCommitResponse>();
-			foreach (ICommit commit in commits)
-			{
-				IUser? author = await _userCollection.GetCachedUserAsync(commit.AuthorId, cancellationToken);
-				responses.Add(CreateGetCommitResponse(commit, author!, null, null));
-			}
-			return responses.ConvertAll(x => PropertyFilter.Apply(x, filter));
+			return ([.. commits.OrderByDescending(x => x.Id)], failedCommits.ToList());
 		}
 
 		static GetCommitResponse CreateGetCommitResponse(ICommit commit, IUser author, IReadOnlyList<CommitTag>? tags, IReadOnlyList<string>? files)
