[Horde][Analytics][Server] Fully qualified label names + categories.

diff --git a/Engine/Source/Programs/Horde/Plugins/Analytics/HordeServer.Analytics.Integration.Tests/Telemetry/TelemetryIntegrationTests.cs b/Engine/Source/Programs/Horde/Plugins/Analytics/HordeServer.Analytics.Integration.Tests/Telemetry/TelemetryIntegrationTests.cs
--- a/Engine/Source/Programs/Horde/Plugins/Analytics/HordeServer.Analytics.Integration.Tests/Telemetry/TelemetryIntegrationTests.cs
+++ b/Engine/Source/Programs/Horde/Plugins/Analytics/HordeServer.Analytics.Integration.Tests/Telemetry/TelemetryIntegrationTests.cs
@@ -241,8 +241,12 @@ namespace HordeServer.Analytics.Integration.Tests.Telemetry
 		private const string MacGameNodeName = "Compile ShooterGameEditor Mac";
 
 		private const string AllWin64Label = "Label_AllWin64Nodes";
+		private const string AllGameLabel = "Label_AllGameNodes";
 		private const string ToolWin64Label = "Label_Win64UHTOnly";
-		private const string ToolAllLabel = "Label_AllTools";
+		private const string AllToolLabel = "Label_AllTools";
+
+		private const string Game1Category = "Game_1";
+		private const string Win64Category = "Win64";
 
 		private const string Aggregate = "AllNodes";
 
@@ -285,16 +289,43 @@ namespace HordeServer.Analytics.Integration.Tests.Telemetry
 			return (templateRefId1, streamConfig, concreteTelemetryWriter);
 		}
 
-		private static NewLabel CreateLabel(string labelName, List<string> nodes)
+		private static List<NewLabel> BasicUncategorized()
+		{
+			return new List<NewLabel>
+			{
+				CreateLabel(AllWin64Label, [Win64ToolNodeName, Win64GameNodeName]),
+				CreateLabel(ToolWin64Label, [Win64ToolNodeName]),
+				CreateLabel(AllToolLabel,  [Win64ToolNodeName, MacToolNodeName])
+			};
+		}
+
+		private static List<NewLabel> BasicCategorized()
+		{
+			return new List<NewLabel>
+			{
+				CreateLabel(AllWin64Label, [Win64ToolNodeName, Win64GameNodeName], Win64Category),
+				CreateLabel(ToolWin64Label, [Win64ToolNodeName], Win64Category),
+				CreateLabel(AllToolLabel,  [Win64ToolNodeName, MacToolNodeName]),
+				CreateLabel(AllGameLabel, [Win64GameNodeName, MacGameNodeName], Game1Category)
+			};
+		}
+
+		private static NewLabel CreateLabel(string labelName, List<string> nodes, string? category = null)
 		{
 			NewLabel label = new NewLabel();
 			label.DashboardName = labelName;
 			label.IncludedNodes = nodes;
 			label.RequiredNodes = nodes;
+
+			if (category != null)
+			{
+				label.DashboardCategory = category;
+			}
+
 			return label;
 		}
 
-		private static Task<IGraph> UpdateGraph(IGraphCollection graphCollection, IGraph rootGraph, ITemplate _)
+		private static Task<IGraph> UpdateGraph(IGraphCollection graphCollection, IGraph rootGraph, Func<List<NewLabel>> labelGenerator, ITemplate _)
 		{
 			// Nodes
 
@@ -311,12 +342,7 @@ namespace HordeServer.Analytics.Integration.Tests.Telemetry
 			};
 
 			// Labels
-			List<NewLabel> labels = new List<NewLabel>
-			{
-				CreateLabel(AllWin64Label, [Win64ToolNodeName, Win64GameNodeName]),
-				CreateLabel(ToolWin64Label, [Win64ToolNodeName]),
-				CreateLabel(ToolAllLabel,  [Win64ToolNodeName, MacToolNodeName])
-			};
+			List<NewLabel> labels = labelGenerator();
 
 			// Groups
 			NewGroup group1 = new NewGroup("DefaultAgent1", nodes1);
@@ -327,14 +353,6 @@ namespace HordeServer.Analytics.Integration.Tests.Telemetry
 			return graphCollection.AppendAsync(rootGraph, [group1, group2], [aggregate], labels);
 		}
 
-		#endregion -- Test Helpers --
-
-		#region -- Test Constants --
-
-		const string EventNameKeyStr = "EventName";
-
-		#endregion -- Test Constants -- 
-
 		protected override void ConfigureServices(IServiceCollection services)
 		{
 			base.ConfigureServices(services);
@@ -343,6 +361,14 @@ namespace HordeServer.Analytics.Integration.Tests.Telemetry
 			services.AddSingleton<ITelemetryWriter>(sp => sp.GetRequiredService<TestTelemetryWriter>());
 		}
 
+		#endregion -- Test Helpers --
+
+		#region -- Test Constants --
+
+		const string EventNameKeyStr = "EventName";
+
+		#endregion -- Test Constants -- 
+
 		[TestMethod]
 		public async Task TestJobSummaryEmissionAsync()
 		{
@@ -404,7 +430,7 @@ namespace HordeServer.Analytics.Integration.Tests.Telemetry
 
 				ITemplate template = await TemplateCollection.GetOrAddAsync(streamConfig.Templates[0]);
 				IGraph graph = await GraphCollection.AddAsync(template);
-				graph = await UpdateGraph(GraphCollection, graph, template);
+				graph = await UpdateGraph(GraphCollection, graph, new Func<List<NewLabel>>(() => { return BasicUncategorized(); }), template);
 
 				options.Arguments.Add(IJob.TargetArgumentPrefix + Aggregate);
 
@@ -454,7 +480,7 @@ namespace HordeServer.Analytics.Integration.Tests.Telemetry
 
 				Assert.IsTrue(actualJobSummaryTelemetry.SuccessfulLabels.Contains(ToolWin64Label));
 				Assert.IsTrue(actualJobSummaryTelemetry.SuccessfulLabels.Contains(AllWin64Label));
-				Assert.IsTrue(actualJobSummaryTelemetry.SuccessfulLabels.Contains(ToolAllLabel));
+				Assert.IsTrue(actualJobSummaryTelemetry.SuccessfulLabels.Contains(AllToolLabel));
 				Assert.IsFalse(actualJobSummaryTelemetry.FailedLabels.Any());
 			}
 
@@ -468,7 +494,7 @@ namespace HordeServer.Analytics.Integration.Tests.Telemetry
 
 				ITemplate template = await TemplateCollection.GetOrAddAsync(streamConfig.Templates[0]);
 				IGraph graph = await GraphCollection.AddAsync(template);
-				graph = await UpdateGraph(GraphCollection, graph, template);
+				graph = await UpdateGraph(GraphCollection, graph, new Func<List<NewLabel>>(() => { return BasicUncategorized(); }), template);
 
 				options.Arguments.Add(IJob.TargetArgumentPrefix + Aggregate);
 
@@ -518,7 +544,70 @@ namespace HordeServer.Analytics.Integration.Tests.Telemetry
 
 				Assert.IsTrue(actualJobSummaryTelemetry.SuccessfulLabels.Contains(ToolWin64Label));
 				Assert.IsTrue(actualJobSummaryTelemetry.SuccessfulLabels.Contains(AllWin64Label));
-				Assert.IsTrue(actualJobSummaryTelemetry.FailedLabels.Contains(ToolAllLabel));
+				Assert.IsTrue(actualJobSummaryTelemetry.FailedLabels.Contains(AllToolLabel));
+			}
+
+			// Label warnings - within batch
+			{
+				// Reset the writer
+				concreteTelemetryWriter.ClearStoredTelemetry();
+
+				CreateJobOptions options = new CreateJobOptions();
+				options.PreflightCommitId = CommitId.FromPerforceChange(999);
+
+				ITemplate template = await TemplateCollection.GetOrAddAsync(streamConfig.Templates[0]);
+				IGraph graph = await GraphCollection.AddAsync(template);
+				graph = await UpdateGraph(GraphCollection, graph, new Func<List<NewLabel>>(() => { return BasicUncategorized(); }), template);
+
+				options.Arguments.Add(IJob.TargetArgumentPrefix + Aggregate);
+
+				IJob job = await JobService.CreateJobAsync(null, streamConfig, templateRefId1, template.Hash, graph, "TestJobLabelTelemetryFailure", CommitIdWithOrder.FromPerforceChange(1234), CommitIdWithOrder.FromPerforceChange(1233), options);
+				job = Deref(await JobService.UpdateBatchAsync(job, job.Batches[0].Id, streamConfig, LogIdUtils.GenerateNewId(), JobStepBatchState.Running));
+				job = Deref(await JobService.UpdateBatchAsync(job, job.Batches[1].Id, streamConfig, LogIdUtils.GenerateNewId(), JobStepBatchState.Running));
+				job = Deref(await JobService.UpdateBatchAsync(job, job.Batches[2].Id, streamConfig, LogIdUtils.GenerateNewId(), JobStepBatchState.Running));
+
+				// Batch 0
+				job = Deref(await JobService.UpdateStepAsync(job, job.Batches[0].Id, job.Batches[0].Steps[0].Id, streamConfig, JobStepState.Running));
+
+				// Batch 1
+				job = Deref(await JobService.UpdateStepAsync(job, job.Batches[1].Id, job.Batches[1].Steps[0].Id, streamConfig, JobStepState.Running)); // Tool
+				job = Deref(await JobService.UpdateStepAsync(job, job.Batches[1].Id, job.Batches[1].Steps[1].Id, streamConfig, JobStepState.Running)); // Game
+
+				// Batch 2
+				job = Deref(await JobService.UpdateStepAsync(job, job.Batches[2].Id, job.Batches[2].Steps[0].Id, streamConfig, JobStepState.Running)); // Tool
+				job = Deref(await JobService.UpdateStepAsync(job, job.Batches[2].Id, job.Batches[2].Steps[1].Id, streamConfig, JobStepState.Running)); // Game
+
+				FakeClock clock = ServiceProvider.GetRequiredService<FakeClock>();
+
+				await clock.AdvanceAsync(TimeSpan.FromSeconds(60.0));
+
+				// Batch 0 - Setup Build
+				job = Deref(await JobService.UpdateStepAsync(job, job.Batches[0].Id, job.Batches[0].Steps[0].Id, streamConfig, JobStepState.Completed, JobStepOutcome.Success));
+
+				// Batch 1 - Win64
+				job = Deref(await JobService.UpdateStepAsync(job, job.Batches[1].Id, job.Batches[1].Steps[0].Id, streamConfig, JobStepState.Completed, JobStepOutcome.Warnings)); // Tool
+				job = Deref(await JobService.UpdateStepAsync(job, job.Batches[1].Id, job.Batches[1].Steps[1].Id, streamConfig, JobStepState.Completed, JobStepOutcome.Success)); // Game
+
+				// Batch 2 - Mac
+				job = Deref(await JobService.UpdateStepAsync(job, job.Batches[2].Id, job.Batches[2].Steps[0].Id, streamConfig, JobStepState.Completed, JobStepOutcome.Success)); // Tool
+				job = Deref(await JobService.UpdateStepAsync(job, job.Batches[2].Id, job.Batches[2].Steps[1].Id, streamConfig, JobStepState.Completed, JobStepOutcome.Success)); // Game
+
+				// Validate we have received a telemetry event.
+				Assert.IsTrue(concreteTelemetryWriter.RecordedTelemetry.ContainsKey(TelemetryStoreId.Default));
+
+				// Verify presence of job summary telemetry
+				IList<Tuple<TelemetryRecordMeta?, Hashtable>> actualTelemetryPayloads = [.. concreteTelemetryWriter.RecordedTelemetry[TelemetryStoreId.Default]];
+				Assert.IsTrue(actualTelemetryPayloads.Any(x => (x.Item2[EventNameKeyStr] as string) == JobLabelSummaryTelemetry.DefaultEventName));
+
+				// Assert we actually have the expected telemetry type.
+				Hashtable actualStateJobSummaryPayload = actualTelemetryPayloads.Where(x => x.Item2[EventNameKeyStr] as string == JobLabelSummaryTelemetry.DefaultEventName).First().Item2;
+				JobLabelSummaryTelemetry? actualJobSummaryTelemetry = actualStateJobSummaryPayload.MapToObject<JobLabelSummaryTelemetry>();
+
+				Assert.IsNotNull(actualJobSummaryTelemetry);
+
+				Assert.IsTrue(actualJobSummaryTelemetry.WarningLabels.Contains(ToolWin64Label));
+				Assert.IsTrue(actualJobSummaryTelemetry.WarningLabels.Contains(AllWin64Label));
+				Assert.IsTrue(actualJobSummaryTelemetry.WarningLabels.Contains(AllToolLabel));
 			}
 
 			// Label failures - within batch
@@ -531,7 +620,7 @@ namespace HordeServer.Analytics.Integration.Tests.Telemetry
 
 				ITemplate template = await TemplateCollection.GetOrAddAsync(streamConfig.Templates[0]);
 				IGraph graph = await GraphCollection.AddAsync(template);
-				graph = await UpdateGraph(GraphCollection, graph, template);
+				graph = await UpdateGraph(GraphCollection, graph, new Func<List<NewLabel>>(() => { return BasicUncategorized(); }), template);
 
 				options.Arguments.Add(IJob.TargetArgumentPrefix + Aggregate);
 
@@ -581,7 +670,258 @@ namespace HordeServer.Analytics.Integration.Tests.Telemetry
 
 				Assert.IsTrue(actualJobSummaryTelemetry.FailedLabels.Contains(ToolWin64Label));
 				Assert.IsTrue(actualJobSummaryTelemetry.FailedLabels.Contains(AllWin64Label));
-				Assert.IsTrue(actualJobSummaryTelemetry.FailedLabels.Contains(ToolAllLabel));
+				Assert.IsTrue(actualJobSummaryTelemetry.FailedLabels.Contains(AllToolLabel));
+			}
+
+			// Label Category success - cross batch
+			{
+				// Reset the writer
+				concreteTelemetryWriter.ClearStoredTelemetry();
+
+				CreateJobOptions options = new CreateJobOptions();
+				options.PreflightCommitId = CommitId.FromPerforceChange(999);
+
+				ITemplate template = await TemplateCollection.GetOrAddAsync(streamConfig.Templates[0]);
+				IGraph graph = await GraphCollection.AddAsync(template);
+				graph = await UpdateGraph(GraphCollection, graph, new Func<List<NewLabel>>(() => { return BasicCategorized(); }), template);
+
+				options.Arguments.Add(IJob.TargetArgumentPrefix + Aggregate);
+
+				IJob job = await JobService.CreateJobAsync(null, streamConfig, templateRefId1, template.Hash, graph, "TestJobLabelTelemetryFailure", CommitIdWithOrder.FromPerforceChange(1234), CommitIdWithOrder.FromPerforceChange(1233), options);
+				job = Deref(await JobService.UpdateBatchAsync(job, job.Batches[0].Id, streamConfig, LogIdUtils.GenerateNewId(), JobStepBatchState.Running));
+				job = Deref(await JobService.UpdateBatchAsync(job, job.Batches[1].Id, streamConfig, LogIdUtils.GenerateNewId(), JobStepBatchState.Running));
+				job = Deref(await JobService.UpdateBatchAsync(job, job.Batches[2].Id, streamConfig, LogIdUtils.GenerateNewId(), JobStepBatchState.Running));
+
+				// Batch 0
+				job = Deref(await JobService.UpdateStepAsync(job, job.Batches[0].Id, job.Batches[0].Steps[0].Id, streamConfig, JobStepState.Running));
+
+				// Batch 1
+				job = Deref(await JobService.UpdateStepAsync(job, job.Batches[1].Id, job.Batches[1].Steps[0].Id, streamConfig, JobStepState.Running)); // Tool
+				job = Deref(await JobService.UpdateStepAsync(job, job.Batches[1].Id, job.Batches[1].Steps[1].Id, streamConfig, JobStepState.Running)); // Game
+
+				// Batch 2
+				job = Deref(await JobService.UpdateStepAsync(job, job.Batches[2].Id, job.Batches[2].Steps[0].Id, streamConfig, JobStepState.Running)); // Tool
+				job = Deref(await JobService.UpdateStepAsync(job, job.Batches[2].Id, job.Batches[2].Steps[1].Id, streamConfig, JobStepState.Running)); // Game
+
+				FakeClock clock = ServiceProvider.GetRequiredService<FakeClock>();
+
+				await clock.AdvanceAsync(TimeSpan.FromSeconds(60.0));
+
+				// Batch 0 - Setup Build
+				job = Deref(await JobService.UpdateStepAsync(job, job.Batches[0].Id, job.Batches[0].Steps[0].Id, streamConfig, JobStepState.Completed, JobStepOutcome.Success));
+
+				// Batch 1 - Win64
+				job = Deref(await JobService.UpdateStepAsync(job, job.Batches[1].Id, job.Batches[1].Steps[0].Id, streamConfig, JobStepState.Completed, JobStepOutcome.Success)); // Tool
+				job = Deref(await JobService.UpdateStepAsync(job, job.Batches[1].Id, job.Batches[1].Steps[1].Id, streamConfig, JobStepState.Completed, JobStepOutcome.Success)); // Game
+
+				// Batch 2 - Mac
+				job = Deref(await JobService.UpdateStepAsync(job, job.Batches[2].Id, job.Batches[2].Steps[0].Id, streamConfig, JobStepState.Completed, JobStepOutcome.Failure)); // Tool
+				job = Deref(await JobService.UpdateStepAsync(job, job.Batches[2].Id, job.Batches[2].Steps[1].Id, streamConfig, JobStepState.Completed, JobStepOutcome.Success)); // Game
+
+				// Validate we have received a telemetry event.
+				Assert.IsTrue(concreteTelemetryWriter.RecordedTelemetry.ContainsKey(TelemetryStoreId.Default));
+
+				// Verify presence of job summary telemetry
+				IList<Tuple<TelemetryRecordMeta?, Hashtable>> actualTelemetryPayloads = [.. concreteTelemetryWriter.RecordedTelemetry[TelemetryStoreId.Default]];
+				Assert.IsTrue(actualTelemetryPayloads.Any(x => (x.Item2[EventNameKeyStr] as string) == JobLabelSummaryTelemetry.DefaultEventName));
+
+				// Assert we actually have the expected telemetry type.
+				Hashtable actualStateJobSummaryPayload = actualTelemetryPayloads.Where(x => x.Item2[EventNameKeyStr] as string == JobLabelSummaryTelemetry.DefaultEventName).First().Item2;
+				JobLabelSummaryTelemetry? actualJobSummaryTelemetry = actualStateJobSummaryPayload.MapToObject<JobLabelSummaryTelemetry>();
+
+				Assert.IsNotNull(actualJobSummaryTelemetry);
+
+				// Assert that the category prefixed label is present.
+				Assert.IsTrue(actualJobSummaryTelemetry.SuccessfulLabels.Contains(JobLabelSummaryTelemetry.CreateFullyQualifiedLabel(Win64Category, ToolWin64Label)));
+				Assert.IsTrue(actualJobSummaryTelemetry.SuccessfulLabels.Contains(JobLabelSummaryTelemetry.CreateFullyQualifiedLabel(Win64Category, AllWin64Label)));
+				Assert.IsTrue(actualJobSummaryTelemetry.SuccessfulCategories.Contains(Win64Category));
+			}
+
+			// Label Category failure - cross batch
+			{
+				// Reset the writer
+				concreteTelemetryWriter.ClearStoredTelemetry();
+
+				CreateJobOptions options = new CreateJobOptions();
+				options.PreflightCommitId = CommitId.FromPerforceChange(999);
+
+				ITemplate template = await TemplateCollection.GetOrAddAsync(streamConfig.Templates[0]);
+				IGraph graph = await GraphCollection.AddAsync(template);
+				graph = await UpdateGraph(GraphCollection, graph, new Func<List<NewLabel>>(() => { return BasicCategorized(); }), template);
+
+				options.Arguments.Add(IJob.TargetArgumentPrefix + Aggregate);
+
+				IJob job = await JobService.CreateJobAsync(null, streamConfig, templateRefId1, template.Hash, graph, "TestJobLabelTelemetryFailure", CommitIdWithOrder.FromPerforceChange(1234), CommitIdWithOrder.FromPerforceChange(1233), options);
+				job = Deref(await JobService.UpdateBatchAsync(job, job.Batches[0].Id, streamConfig, LogIdUtils.GenerateNewId(), JobStepBatchState.Running));
+				job = Deref(await JobService.UpdateBatchAsync(job, job.Batches[1].Id, streamConfig, LogIdUtils.GenerateNewId(), JobStepBatchState.Running));
+				job = Deref(await JobService.UpdateBatchAsync(job, job.Batches[2].Id, streamConfig, LogIdUtils.GenerateNewId(), JobStepBatchState.Running));
+
+				// Batch 0
+				job = Deref(await JobService.UpdateStepAsync(job, job.Batches[0].Id, job.Batches[0].Steps[0].Id, streamConfig, JobStepState.Running));
+
+				// Batch 1
+				job = Deref(await JobService.UpdateStepAsync(job, job.Batches[1].Id, job.Batches[1].Steps[0].Id, streamConfig, JobStepState.Running)); // Tool
+				job = Deref(await JobService.UpdateStepAsync(job, job.Batches[1].Id, job.Batches[1].Steps[1].Id, streamConfig, JobStepState.Running)); // Game
+
+				// Batch 2
+				job = Deref(await JobService.UpdateStepAsync(job, job.Batches[2].Id, job.Batches[2].Steps[0].Id, streamConfig, JobStepState.Running)); // Tool
+				job = Deref(await JobService.UpdateStepAsync(job, job.Batches[2].Id, job.Batches[2].Steps[1].Id, streamConfig, JobStepState.Running)); // Game
+
+				FakeClock clock = ServiceProvider.GetRequiredService<FakeClock>();
+
+				await clock.AdvanceAsync(TimeSpan.FromSeconds(60.0));
+
+				// Batch 0 - Setup Build
+				job = Deref(await JobService.UpdateStepAsync(job, job.Batches[0].Id, job.Batches[0].Steps[0].Id, streamConfig, JobStepState.Completed, JobStepOutcome.Success));
+
+				// Batch 1 - Win64
+				job = Deref(await JobService.UpdateStepAsync(job, job.Batches[1].Id, job.Batches[1].Steps[0].Id, streamConfig, JobStepState.Completed, JobStepOutcome.Success)); // Tool
+				job = Deref(await JobService.UpdateStepAsync(job, job.Batches[1].Id, job.Batches[1].Steps[1].Id, streamConfig, JobStepState.Completed, JobStepOutcome.Success)); // Game
+
+				// Batch 2 - Mac
+				job = Deref(await JobService.UpdateStepAsync(job, job.Batches[2].Id, job.Batches[2].Steps[0].Id, streamConfig, JobStepState.Completed, JobStepOutcome.Failure)); // Tool
+				job = Deref(await JobService.UpdateStepAsync(job, job.Batches[2].Id, job.Batches[2].Steps[1].Id, streamConfig, JobStepState.Completed, JobStepOutcome.Failure)); // Game
+
+				// Validate we have received a telemetry event.
+				Assert.IsTrue(concreteTelemetryWriter.RecordedTelemetry.ContainsKey(TelemetryStoreId.Default));
+
+				// Verify presence of job summary telemetry
+				IList<Tuple<TelemetryRecordMeta?, Hashtable>> actualTelemetryPayloads = [.. concreteTelemetryWriter.RecordedTelemetry[TelemetryStoreId.Default]];
+				Assert.IsTrue(actualTelemetryPayloads.Any(x => (x.Item2[EventNameKeyStr] as string) == JobLabelSummaryTelemetry.DefaultEventName));
+
+				// Assert we actually have the expected telemetry type.
+				Hashtable actualStateJobSummaryPayload = actualTelemetryPayloads.Where(x => x.Item2[EventNameKeyStr] as string == JobLabelSummaryTelemetry.DefaultEventName).First().Item2;
+				JobLabelSummaryTelemetry? actualJobSummaryTelemetry = actualStateJobSummaryPayload.MapToObject<JobLabelSummaryTelemetry>();
+
+				Assert.IsNotNull(actualJobSummaryTelemetry);
+
+				// Assert that the category prefixed label is present.
+				Assert.IsTrue(actualJobSummaryTelemetry.FailedLabels.Contains(JobLabelSummaryTelemetry.CreateFullyQualifiedLabel(Game1Category, AllGameLabel)));
+				Assert.IsTrue(actualJobSummaryTelemetry.FailedCategories.Contains(Game1Category));
+			}
+
+			// Label Category warning - within categroy - warning label is higher precedence than success label
+			{
+				// Reset the writer
+				concreteTelemetryWriter.ClearStoredTelemetry();
+
+				CreateJobOptions options = new CreateJobOptions();
+				options.PreflightCommitId = CommitId.FromPerforceChange(999);
+
+				ITemplate template = await TemplateCollection.GetOrAddAsync(streamConfig.Templates[0]);
+				IGraph graph = await GraphCollection.AddAsync(template);
+				graph = await UpdateGraph(GraphCollection, graph, new Func<List<NewLabel>>(() => { return BasicCategorized(); }), template);
+
+				options.Arguments.Add(IJob.TargetArgumentPrefix + Aggregate);
+
+				IJob job = await JobService.CreateJobAsync(null, streamConfig, templateRefId1, template.Hash, graph, "TestJobLabelTelemetryFailure", CommitIdWithOrder.FromPerforceChange(1234), CommitIdWithOrder.FromPerforceChange(1233), options);
+				job = Deref(await JobService.UpdateBatchAsync(job, job.Batches[0].Id, streamConfig, LogIdUtils.GenerateNewId(), JobStepBatchState.Running));
+				job = Deref(await JobService.UpdateBatchAsync(job, job.Batches[1].Id, streamConfig, LogIdUtils.GenerateNewId(), JobStepBatchState.Running));
+				job = Deref(await JobService.UpdateBatchAsync(job, job.Batches[2].Id, streamConfig, LogIdUtils.GenerateNewId(), JobStepBatchState.Running));
+
+				// Batch 0
+				job = Deref(await JobService.UpdateStepAsync(job, job.Batches[0].Id, job.Batches[0].Steps[0].Id, streamConfig, JobStepState.Running));
+
+				// Batch 1
+				job = Deref(await JobService.UpdateStepAsync(job, job.Batches[1].Id, job.Batches[1].Steps[0].Id, streamConfig, JobStepState.Running)); // Tool
+				job = Deref(await JobService.UpdateStepAsync(job, job.Batches[1].Id, job.Batches[1].Steps[1].Id, streamConfig, JobStepState.Running)); // Game
+
+				// Batch 2
+				job = Deref(await JobService.UpdateStepAsync(job, job.Batches[2].Id, job.Batches[2].Steps[0].Id, streamConfig, JobStepState.Running)); // Tool
+				job = Deref(await JobService.UpdateStepAsync(job, job.Batches[2].Id, job.Batches[2].Steps[1].Id, streamConfig, JobStepState.Running)); // Game
+
+				FakeClock clock = ServiceProvider.GetRequiredService<FakeClock>();
+
+				await clock.AdvanceAsync(TimeSpan.FromSeconds(60.0));
+
+				// Batch 0 - Setup Build
+				job = Deref(await JobService.UpdateStepAsync(job, job.Batches[0].Id, job.Batches[0].Steps[0].Id, streamConfig, JobStepState.Completed, JobStepOutcome.Success));
+
+				// Batch 1 - Win64
+				job = Deref(await JobService.UpdateStepAsync(job, job.Batches[1].Id, job.Batches[1].Steps[0].Id, streamConfig, JobStepState.Completed, JobStepOutcome.Success)); // Tool
+				job = Deref(await JobService.UpdateStepAsync(job, job.Batches[1].Id, job.Batches[1].Steps[1].Id, streamConfig, JobStepState.Completed, JobStepOutcome.Warnings)); // Game
+
+				// Batch 2 - Mac
+				job = Deref(await JobService.UpdateStepAsync(job, job.Batches[2].Id, job.Batches[2].Steps[0].Id, streamConfig, JobStepState.Completed, JobStepOutcome.Success)); // Tool
+				job = Deref(await JobService.UpdateStepAsync(job, job.Batches[2].Id, job.Batches[2].Steps[1].Id, streamConfig, JobStepState.Completed, JobStepOutcome.Success)); // Game
+
+				// Validate we have received a telemetry event.
+				Assert.IsTrue(concreteTelemetryWriter.RecordedTelemetry.ContainsKey(TelemetryStoreId.Default));
+
+				// Verify presence of job summary telemetry
+				IList<Tuple<TelemetryRecordMeta?, Hashtable>> actualTelemetryPayloads = [.. concreteTelemetryWriter.RecordedTelemetry[TelemetryStoreId.Default]];
+				Assert.IsTrue(actualTelemetryPayloads.Any(x => (x.Item2[EventNameKeyStr] as string) == JobLabelSummaryTelemetry.DefaultEventName));
+
+				// Assert we actually have the expected telemetry type.
+				Hashtable actualStateJobSummaryPayload = actualTelemetryPayloads.Where(x => x.Item2[EventNameKeyStr] as string == JobLabelSummaryTelemetry.DefaultEventName).First().Item2;
+				JobLabelSummaryTelemetry? actualJobSummaryTelemetry = actualStateJobSummaryPayload.MapToObject<JobLabelSummaryTelemetry>();
+
+				Assert.IsNotNull(actualJobSummaryTelemetry);
+
+				// Assert that the category prefixed label is present.
+				Assert.IsTrue(actualJobSummaryTelemetry.WarningCategories.Contains(Win64Category));
+			}
+
+			// Label Category failurue - within category - failed label is higher precedence than warning label
+			{
+				// Reset the writer
+				concreteTelemetryWriter.ClearStoredTelemetry();
+
+				CreateJobOptions options = new CreateJobOptions();
+				options.PreflightCommitId = CommitId.FromPerforceChange(999);
+
+				ITemplate template = await TemplateCollection.GetOrAddAsync(streamConfig.Templates[0]);
+				IGraph graph = await GraphCollection.AddAsync(template);
+				graph = await UpdateGraph(GraphCollection, graph, new Func<List<NewLabel>>(() => { return BasicCategorized(); }), template);
+
+				options.Arguments.Add(IJob.TargetArgumentPrefix + Aggregate);
+
+				IJob job = await JobService.CreateJobAsync(null, streamConfig, templateRefId1, template.Hash, graph, "TestJobLabelTelemetryFailure", CommitIdWithOrder.FromPerforceChange(1234), CommitIdWithOrder.FromPerforceChange(1233), options);
+				job = Deref(await JobService.UpdateBatchAsync(job, job.Batches[0].Id, streamConfig, LogIdUtils.GenerateNewId(), JobStepBatchState.Running));
+				job = Deref(await JobService.UpdateBatchAsync(job, job.Batches[1].Id, streamConfig, LogIdUtils.GenerateNewId(), JobStepBatchState.Running));
+				job = Deref(await JobService.UpdateBatchAsync(job, job.Batches[2].Id, streamConfig, LogIdUtils.GenerateNewId(), JobStepBatchState.Running));
+
+				// Batch 0
+				job = Deref(await JobService.UpdateStepAsync(job, job.Batches[0].Id, job.Batches[0].Steps[0].Id, streamConfig, JobStepState.Running));
+
+				// Batch 1
+				job = Deref(await JobService.UpdateStepAsync(job, job.Batches[1].Id, job.Batches[1].Steps[0].Id, streamConfig, JobStepState.Running)); // Tool
+				job = Deref(await JobService.UpdateStepAsync(job, job.Batches[1].Id, job.Batches[1].Steps[1].Id, streamConfig, JobStepState.Running)); // Game
+
+				// Batch 2
+				job = Deref(await JobService.UpdateStepAsync(job, job.Batches[2].Id, job.Batches[2].Steps[0].Id, streamConfig, JobStepState.Running)); // Tool
+				job = Deref(await JobService.UpdateStepAsync(job, job.Batches[2].Id, job.Batches[2].Steps[1].Id, streamConfig, JobStepState.Running)); // Game
+
+				FakeClock clock = ServiceProvider.GetRequiredService<FakeClock>();
+
+				await clock.AdvanceAsync(TimeSpan.FromSeconds(60.0));
+
+				// Batch 0 - Setup Build
+				job = Deref(await JobService.UpdateStepAsync(job, job.Batches[0].Id, job.Batches[0].Steps[0].Id, streamConfig, JobStepState.Completed, JobStepOutcome.Success));
+
+				// Batch 1 - Win64
+				job = Deref(await JobService.UpdateStepAsync(job, job.Batches[1].Id, job.Batches[1].Steps[0].Id, streamConfig, JobStepState.Completed, JobStepOutcome.Warnings)); // Tool
+				job = Deref(await JobService.UpdateStepAsync(job, job.Batches[1].Id, job.Batches[1].Steps[1].Id, streamConfig, JobStepState.Completed, JobStepOutcome.Failure)); // Game
+
+				// Batch 2 - Mac
+				job = Deref(await JobService.UpdateStepAsync(job, job.Batches[2].Id, job.Batches[2].Steps[0].Id, streamConfig, JobStepState.Completed, JobStepOutcome.Success)); // Tool
+				job = Deref(await JobService.UpdateStepAsync(job, job.Batches[2].Id, job.Batches[2].Steps[1].Id, streamConfig, JobStepState.Completed, JobStepOutcome.Success)); // Game
+
+				// Validate we have received a telemetry event.
+				Assert.IsTrue(concreteTelemetryWriter.RecordedTelemetry.ContainsKey(TelemetryStoreId.Default));
+
+				// Verify presence of job summary telemetry
+				IList<Tuple<TelemetryRecordMeta?, Hashtable>> actualTelemetryPayloads = [.. concreteTelemetryWriter.RecordedTelemetry[TelemetryStoreId.Default]];
+				Assert.IsTrue(actualTelemetryPayloads.Any(x => (x.Item2[EventNameKeyStr] as string) == JobLabelSummaryTelemetry.DefaultEventName));
+
+				// Assert we actually have the expected telemetry type.
+				Hashtable actualStateJobSummaryPayload = actualTelemetryPayloads.Where(x => x.Item2[EventNameKeyStr] as string == JobLabelSummaryTelemetry.DefaultEventName).First().Item2;
+				JobLabelSummaryTelemetry? actualJobSummaryTelemetry = actualStateJobSummaryPayload.MapToObject<JobLabelSummaryTelemetry>();
+
+				Assert.IsNotNull(actualJobSummaryTelemetry);
+
+				// Assert that the category prefixed label is present.
+				Assert.IsTrue(actualJobSummaryTelemetry.FailedCategories.Contains(Win64Category));
 			}
 		}
 	}
diff --git a/Engine/Source/Programs/Horde/Plugins/Build/HordeServer.Build/Jobs/JobTelemetry.cs b/Engine/Source/Programs/Horde/Plugins/Build/HordeServer.Build/Jobs/JobTelemetry.cs
--- a/Engine/Source/Programs/Horde/Plugins/Build/HordeServer.Build/Jobs/JobTelemetry.cs
+++ b/Engine/Source/Programs/Horde/Plugins/Build/HordeServer.Build/Jobs/JobTelemetry.cs
@@ -183,16 +183,41 @@ namespace HordeServer.Jobs
 		/// </summary>
 		public string?[] SuccessfulLabels { get; init; }
 
+		/// <summary>
+		/// Array representing all labels that were warnings in the job.
+		/// </summary>
+		public string?[] WarningLabels { get; init; }
+
 		/// <summary>
 		/// Array representing all labels that were failures in the job.
 		/// </summary>
 		public string?[] FailedLabels { get; init; }
 
+		/// <summary>
+		/// Array representing all categories that were successful in the job.
+		/// </summary>
+		public string?[] SuccessfulCategories { get; init; }
+
+		/// <summary>
+		/// Array representing all categories that were warnings in the job.
+		/// </summary>
+		public string?[] WarningCategories { get; init; }
+
+		/// <summary>
+		/// Array representing all categories that were failures in the job.
+		/// </summary>
+		public string?[] FailedCategories { get; init; }
+
 		/// <summary>
 		/// Default event name for the JobLabelSummaryTelemetry.
 		/// </summary>
 		public const string DefaultEventName = "State.JobLabelSummary";
 
+		/// <summary>
+		/// Current schema version.
+		/// </summary>
+		public static readonly int CurrentSchemaVersion = 1;
+
 		/// <summary>
 		/// Constructor.
 		/// </summary>
@@ -202,6 +227,7 @@ namespace HordeServer.Jobs
 		{
 			EventName = DefaultEventName;
 			SuccessfulLabels = [];
+			WarningLabels = [];
 			FailedLabels = [];
 			EventName = DefaultEventName;
 			JobId = job.Id;
@@ -213,8 +239,66 @@ namespace HordeServer.Jobs
 			List<GetLabelStateResponse> responses = [];
 			job.GetLabelStateResponses(graph, responses);
 
-			SuccessfulLabels = responses.Where(r => r.State == LabelState.Complete && r.Outcome == LabelOutcome.Success).Select(r => r.DashboardName).ToArray();
-			FailedLabels = responses.Where(r => r.State == LabelState.Complete && r.Outcome == LabelOutcome.Failure).Select(r => r.DashboardName).ToArray();
+			SuccessfulLabels = [.. responses.Where(r => r.State == LabelState.Complete && r.Outcome == LabelOutcome.Success && r.DashboardName != null).Select(r => CreateFullyQualifiedLabel(r))];
+			WarningLabels = [.. responses.Where(r => r.State == LabelState.Complete && r.Outcome == LabelOutcome.Warnings && r.DashboardName != null).Select(r => CreateFullyQualifiedLabel(r))];
+			FailedLabels = [.. responses.Where(r => r.State == LabelState.Complete && r.Outcome == LabelOutcome.Failure && r.DashboardName != null).Select(r => CreateFullyQualifiedLabel(r))];
+
+			Dictionary<string, LabelOutcome> categoryToOutcome = [];
+
+			for (int i = 0; i < responses.Count; ++i)
+			{
+				GetLabelStateResponse response = responses[i];
+				if (response.DashboardCategory == null || response.Outcome == null || response.State != LabelState.Complete)
+				{
+					continue;
+				}
+
+				if (!categoryToOutcome.ContainsKey(response.DashboardCategory))
+				{
+					categoryToOutcome.Add(response.DashboardCategory, response.Outcome.Value);
+				}
+				else
+				{
+					if (categoryToOutcome[response.DashboardCategory] > response.Outcome.Value)
+					{
+						categoryToOutcome[response.DashboardCategory] = response.Outcome.Value;
+					}
+				}
+			}
+
+			List<KeyValuePair<string, LabelOutcome>> flattenedCategoryList = categoryToOutcome.ToList();
+
+			SuccessfulCategories = [.. flattenedCategoryList.Where(r => r.Value == LabelOutcome.Success).Select(r => r.Key)];
+			WarningCategories = [.. flattenedCategoryList.Where(r => r.Value == LabelOutcome.Warnings).Select(r => r.Key)];
+			FailedCategories = [.. flattenedCategoryList.Where(r => r.Value == LabelOutcome.Failure).Select(r => r.Key)];
+		}
+
+		/// <summary>
+		/// Delimiter used in separating the Label Category and the Label Name.
+		/// </summary>
+		public const string FullyQualifiedLabelDelimiter = "-";
+
+		/// <summary>
+		/// Creates a fully qualified label name.
+		/// </summary>
+		/// <param name="response">The label response to convert to a fully qualified label name.</param>
+		/// <returns>The fully qualified label name.</returns>
+		/// <remarks>If there is no valid category, the label name will simply be the <see cref="GetLabelStateResponse.DashboardName"/>. If no valid dashboard name, the function will return <see cref="String.Empty"/>.</remarks>
+		public static string CreateFullyQualifiedLabel(GetLabelStateResponse response)
+		{
+			return CreateFullyQualifiedLabel(response.DashboardCategory, response.DashboardName);
+		}
+
+		/// <summary>
+		/// Creates a fully qualified label name.
+		/// </summary>
+		/// <param name="dashboardCategory">The category of the label.</param>
+		/// <param name="dashboardName">The name of the label.</param>
+		/// <returns>The fully qualified label name.</returns>
+		/// <remarks>If there is no valid category, the label name will simply be the <see cref="GetLabelStateResponse.DashboardName"/>. If no valid dashboard name, the function will return <see cref="String.Empty"/>.</remarks>
+		public static string CreateFullyQualifiedLabel(string? dashboardCategory, string? dashboardName)
+		{
+			return (dashboardCategory != null && dashboardName != null) ? $"{dashboardCategory}{FullyQualifiedLabelDelimiter}{dashboardName}" : (dashboardName != null ? dashboardName : String.Empty);
 		}
 	}
 }
\ No newline at end of file
