[Horde][BuildHealth][Dashboard] Missed an edge case for automatic schema upgrades - select all was implied on schema 1 for a template; we are now explicit in schema version 2+.

diff --git a/Engine/Source/Programs/Horde/HordeDashboard/src/components/buildhealth/BuildHealthOptions.ts b/Engine/Source/Programs/Horde/HordeDashboard/src/components/buildhealth/BuildHealthOptions.ts
--- a/Engine/Source/Programs/Horde/HordeDashboard/src/components/buildhealth/BuildHealthOptions.ts
+++ b/Engine/Source/Programs/Horde/HordeDashboard/src/components/buildhealth/BuildHealthOptions.ts
@@ -10,6 +10,21 @@ import { BuildHealthOptionsController, IBuildHealthOptionWriter } from "./BuildH
 
 // #region -- Options Utilities & Data Types --
 
+/**
+ * Perfoms a deep copy of a provided build health query params.
+ * @param params The params to perform a deep clone on.
+ * @returns The deep clone.
+ */
+export function cloneBuildHealthQueryParams(params: BuildHealthQueryParams): BuildHealthQueryParams {
+    return {
+        ...params,
+        projectIds: [...params.projectIds],
+        streamIds: [...params.streamIds],
+        templateIds: [...params.templateIds],
+        stepNames: [...params.stepNames],
+    };
+}
+
 /**
  * Data structure used to pass state from location query, to Build Health Options.
  */
@@ -196,8 +211,7 @@ export function parseBuildHealthQueryParams(location: Location): BuildHealthQuer
     if (querySchemaVersion === 1) {
         templateIds = getList(TEMPLATES_URL_PARAM_NAME);
         stepNames = getList(STEPS_URL_PARAM_NAME);
-    }
-    else if (querySchemaVersion === 2) {
+    } else if (querySchemaVersion === 2) {
         const rawTemplates = searchParams.get(TEMPLATES_URL_PARAM_NAME);
         if (!rawTemplates) {
             templateIds = [];
@@ -238,6 +252,45 @@ export function parseBuildHealthQueryParams(location: Location): BuildHealthQuer
     return { projectIds, streamIds, templateIds: templateIds, stepNames, timeSpanKey, includePreflight, querySchemaVersion };
 }
 
+/**
+ * Performs an upgrade on a provided build healthy query params object. 
+ * @param params The params object to attempt to upgrade.
+ * @param currentSchemaVersion The target schema version to attempt an upgrade for.
+ * @returns If applicable, a cloned params object that has been upgraded. Otherwise returns the original object.
+ */
+export function upgradeBuildHealthQueryParams(params: BuildHealthQueryParams, currentSchemaVersion: number): BuildHealthQueryParams {
+    if (params.querySchemaVersion === currentSchemaVersion) {
+        return params;
+    }
+
+    let clonedParams = cloneBuildHealthQueryParams(params);
+    if (clonedParams.querySchemaVersion === 1 && currentSchemaVersion >= 2) {
+        // Upgrades from query structure to 1 -> 2+ we adopted explciit select-all instead of inferred step all
+        // - template ids in the query param used to be enough to signify "use all"; we have since moved away from this to a more explicit standard
+        // - if there is a template id in the query params, and it has no matching entry in the step filters, infer that it is the STEPS_URL_PARAM_ALL_STEPS case.
+        let upgradedSteps: string[] = [];
+
+        for (let idx: number = 0; idx < clonedParams.templateIds.length; idx++) {
+
+            const templateIdKey = clonedParams.templateIds[idx];
+            let { streamId, templateId } = decodeTemplateKey(templateIdKey);
+
+            // Create an encoded step key with no step name so we can prefix match against the query param's encoded step keys to see if the template is represented.
+            let encodedStepKeyWithNoName = encodeStepNameFromStringsPreserveCase(streamId, templateId, "");
+
+            const hasPrefix = clonedParams.stepNames.some((encodedStep: string) => encodedStep.startsWith(encodedStepKeyWithNoName));
+
+            if (!hasPrefix) {
+                upgradedSteps.push(encodeStepNameFromStringsPreserveCase(streamId, templateId, STEPS_URL_PARAM_ALL_STEPS));
+            }
+        }
+
+        clonedParams.stepNames.push(...upgradedSteps);
+    }
+
+    return clonedParams;
+}
+
 /**
  * Processes provided Build Health Query Params for the Build Health View.
  * @param params The params to process.
diff --git a/Engine/Source/Programs/Horde/HordeDashboard/src/components/buildhealth/BuildhealthView.tsx b/Engine/Source/Programs/Horde/HordeDashboard/src/components/buildhealth/BuildhealthView.tsx
--- a/Engine/Source/Programs/Horde/HordeDashboard/src/components/buildhealth/BuildhealthView.tsx
+++ b/Engine/Source/Programs/Horde/HordeDashboard/src/components/buildhealth/BuildhealthView.tsx
@@ -12,7 +12,7 @@ import { TopNav } from "../TopNav";
 import { StepOutcomeView } from "../buildhealth/stepoutcome/StepOutcome";
 import { BuildHealthDataHandler, DataHandlerRefreshRequest, expandRefreshRequest } from "./BuildHealthDataHandler";
 import { KEY_SEPARATOR, PARAMETER_KEY_PREFIX, StepRefData, TemplateRefData } from "./BuildHealthDataTypes";
-import { BuildHealthOptionsState, BuildHealthOptionStateDiff, BuildHealthUIOptionsState, JobHistoryTimeSpans, loadBuildHealthOptionsFromParams, parseBuildHealthQueryParams, TimeSpan, toNavigationQuery } from "./BuildHealthOptions";
+import { BuildHealthOptionsState, BuildHealthOptionStateDiff, BuildHealthUIOptionsState, JobHistoryTimeSpans, loadBuildHealthOptionsFromParams, parseBuildHealthQueryParams, TimeSpan, toNavigationQuery, upgradeBuildHealthQueryParams } from "./BuildHealthOptions";
 import { decodeStepKey, decodeTemplateKey, encodeStepKey, encodeTemplateKey } from "./BuildHealthUtilities";
 import { stepOutcomeTableClasses } from "./stepoutcome/StepOutcomeSharedUIComponents";
 import { runInAction } from "mobx";
@@ -713,7 +713,8 @@ const BuildHealthUrlSync: React.FC<{ buildHealthController: BuildHealthOptionsCo
     useEffect(() => {
         const initializeFromUrl = async () => {
             const params = parseBuildHealthQueryParams(location);
-            await loadBuildHealthOptionsFromParams(params, buildHealthController, handler);
+            const upgradedParams = upgradeBuildHealthQueryParams(params, buildHealthController.querySchemaVersion);
+            await loadBuildHealthOptionsFromParams(upgradedParams, buildHealthController, handler);
         };
 
         initializeFromUrl();
diff --git a/Engine/Source/Programs/Horde/HordeDashboard/tests/BuildHealthOptions.test.ts b/Engine/Source/Programs/Horde/HordeDashboard/tests/BuildHealthOptions.test.ts
new file mode 100644
--- /dev/null
+++ b/Engine/Source/Programs/Horde/HordeDashboard/tests/BuildHealthOptions.test.ts
@@ -0,0 +1,81 @@
+import { parseBuildHealthQueryParams, upgradeBuildHealthQueryParams } from 'horde/components/buildhealth/BuildHealthOptions';
+import { describe, it, expect } from 'vitest';
+import { Location } from 'react-router-dom';
+import LZString from "lz-string"
+
+describe('parseBuildHealthQueryParams', () => {
+	it('should correctly parse query parameters (schema v1)', () => {
+		// Arrange
+		const mockParams = new URLSearchParams();
+
+		// Use LZString compression for encoded URL params
+		const compressedProjects = LZString.compressToEncodedURIComponent('projA,projB');
+		const compressedStreams = LZString.compressToEncodedURIComponent('stream1,stream2');
+		const compressedTemplates = LZString.compressToEncodedURIComponent('stream1::template1,stream2::template2');
+		const compressedSteps = LZString.compressToEncodedURIComponent('stream1::template1::step1,stream2::template2::step2');
+		const compressedTimeSpan = LZString.compressToEncodedURIComponent('time_1_week');
+		const compressedPreflight = LZString.compressToEncodedURIComponent('true');
+
+		// Do *NOT* use the param const strings; we want to catch those contract failures as well.
+		mockParams.set('projects', compressedProjects);
+		mockParams.set('streams', compressedStreams);
+		mockParams.set('templates', compressedTemplates);
+		mockParams.set('stepNames', compressedSteps);
+		mockParams.set('lastTimeRange', compressedTimeSpan);
+		mockParams.set('includePreflight', compressedPreflight);
+		mockParams.set('schema', '1');
+
+		const location = { search: `?${mockParams.toString()}` } as Location;
+
+		// Act
+		const result = parseBuildHealthQueryParams(location);
+
+		// Assert
+		expect(result.projectIds).toEqual(['projA', 'projB']);
+		expect(result.streamIds).toEqual(['stream1', 'stream2']);
+		expect(result.templateIds).toEqual(['stream1::template1', 'stream2::template2']);
+		expect(result.stepNames).toEqual(['stream1::template1::step1', 'stream2::template2::step2']);
+		expect(result.timeSpanKey).toBe('time_1_week');
+		expect(result.includePreflight).toBe(true);
+		expect(result.querySchemaVersion).toBe(1);
+	});
+});
+
+describe('upgradeBuildHealthQueryParams', () => {
+	it('should correctly upgrade query parameters (schema v1 to v2)', () => {
+		// Arrange
+		const mockParams = new URLSearchParams();
+
+		// Use LZString compression for encoded URL params
+		const compressedProjects = LZString.compressToEncodedURIComponent('projA,projB');
+		const compressedStreams = LZString.compressToEncodedURIComponent('stream1,stream2');
+		const compressedTemplates = LZString.compressToEncodedURIComponent('stream1::template1,stream2::template2,stream2::template3');
+		const compressedSteps = LZString.compressToEncodedURIComponent('stream1::template1::step1,stream2::template2::step2');
+		const compressedTimeSpan = LZString.compressToEncodedURIComponent('time_1_week');
+		const compressedPreflight = LZString.compressToEncodedURIComponent('true');
+
+		// Do *NOT* use the param const strings; we want to catch those contract failures as well.
+		mockParams.set('projects', compressedProjects);
+		mockParams.set('streams', compressedStreams);
+		mockParams.set('templates', compressedTemplates);
+		mockParams.set('stepNames', compressedSteps);
+		mockParams.set('lastTimeRange', compressedTimeSpan);
+		mockParams.set('includePreflight', compressedPreflight);
+		mockParams.set('schema', '1');
+
+		const location = { search: `?${mockParams.toString()}` } as Location;
+
+		// Act
+		const result = parseBuildHealthQueryParams(location);
+		const upgradedParams = upgradeBuildHealthQueryParams(result, 2);
+
+		// Assert
+		expect(upgradedParams.projectIds).toEqual(['projA', 'projB']);
+		expect(upgradedParams.streamIds).toEqual(['stream1', 'stream2']);
+		expect(upgradedParams.templateIds).toEqual(['stream1::template1', 'stream2::template2', 'stream2::template3']);
+		expect(upgradedParams.stepNames).toEqual(['stream1::template1::step1', 'stream2::template2::step2', 'stream2::template3::stepNames_AllKnownValues']);
+		expect(upgradedParams.timeSpanKey).toBe('time_1_week');
+		expect(upgradedParams.includePreflight).toBe(true);
+		expect(upgradedParams.querySchemaVersion).toBe(1);
+	});
+});
\ No newline at end of file
