[Build Health Dashboard][Step Outcome] Minor fixes to comments and unit tests
#rnx

diff --git a/Engine/Source/Programs/Horde/Plugins/Build/HordeServer.Build.Tests/Streams/StreamCollectionTests.cs b/Engine/Source/Programs/Horde/Plugins/Build/HordeServer.Build.Tests/Streams/StreamCollectionTests.cs
--- a/Engine/Source/Programs/Horde/Plugins/Build/HordeServer.Build.Tests/Streams/StreamCollectionTests.cs
+++ b/Engine/Source/Programs/Horde/Plugins/Build/HordeServer.Build.Tests/Streams/StreamCollectionTests.cs
@@ -45,7 +45,7 @@ namespace HordeServer.Tests.Streams
 			int thirdChangelist = 112235;
 			int invalidChangelist = 112236;
 
-			PerforceService.AddChange(fixture.StreamId, firstChangelist, jerry, "Edited file1.cpp.", new[] { "file.cpp" });
+			PerforceService.AddChange(fixture.StreamId, firstChangelist, jerry, "Edited file1.cpp.", new[] { "file1.cpp" });
 			PerforceService.AddChange(fixture.StreamId, secondChangelist, jerry, "Added file2.cpp.", new[] { "file2.cpp" });
 			PerforceService.AddChange(fixture.StreamId, thirdChangelist, jerry, "Edited file1.cpp.", new[] { "file1.cpp" });
 
@@ -127,7 +127,7 @@ namespace HordeServer.Tests.Streams
 			int thirdChangelist = 112235;
 			int invalidChangelist = 112236;
 
-			PerforceService.AddChange(fixture.StreamId, firstChangelist, jerry, "Edited file1.cpp.", new[] { "file.cpp" });
+			PerforceService.AddChange(fixture.StreamId, firstChangelist, jerry, "Edited file1.cpp.", new[] { "file1.cpp" });
 			PerforceService.AddChange(fixture.StreamId, secondChangelist, jerry, "Added file2.cpp.", new[] { "file2.cpp" });
 			PerforceService.AddChange(fixture.StreamId, thirdChangelist, jerry, "Edited file1.cpp.", new[] { "file1.cpp" });
 
diff --git a/Engine/Source/Programs/Horde/Plugins/Build/HordeServer.Build/Streams/StreamsController.cs b/Engine/Source/Programs/Horde/Plugins/Build/HordeServer.Build/Streams/StreamsController.cs
--- a/Engine/Source/Programs/Horde/Plugins/Build/HordeServer.Build/Streams/StreamsController.cs
+++ b/Engine/Source/Programs/Horde/Plugins/Build/HordeServer.Build/Streams/StreamsController.cs
@@ -269,8 +269,10 @@ namespace HordeServer.Streams
 		/// <param name="filter">The filter to apply to the results</param>
 		/// <param name="cancellationToken">Cancellation token for the operation</param>
 		/// Returns 200 OK with list of filtered commits. Partial failures for batch retrieval are indicated via the <see cref="MissingCommitsHeader"/> header.
-		/// <remarks>This endpoint is a workaround to retrieve large sets of changes which would result in 414 reponses from <see cref=GetChangesAsync"/>.
-		/// <remarks>Any commits that were not found or that threw exceptions are stored in the <see cref="MissingCommitsHeader"/> response header.</remarks>
+		/// <remarks>
+		/// This endpoint is a workaround to retrieve large sets of changes which would result in 414 reponses from <see cref="GetChangesAsync"/>.
+		/// Any commits that were not found or that threw exceptions are stored in the <see cref="MissingCommitsHeader"/> response header.
+		/// </remarks>
 		[HttpPost]
 		[Route("/api/v1/streams/{streamId}/changes")]
 		[ProducesResponseType(typeof(List<GetCommitResponse>), 200)]
