[Build Health Dashboard][Step Outcome] Minor fixes to crefs in comments to increase clarity
#rnx

diff --git a/Engine/Source/Programs/Horde/Plugins/Build/HordeServer.Build/Streams/StreamsController.cs b/Engine/Source/Programs/Horde/Plugins/Build/HordeServer.Build/Streams/StreamsController.cs
--- a/Engine/Source/Programs/Horde/Plugins/Build/HordeServer.Build/Streams/StreamsController.cs
+++ b/Engine/Source/Programs/Horde/Plugins/Build/HordeServer.Build/Streams/StreamsController.cs
@@ -185,7 +185,7 @@ namespace HordeServer.Streams
 		}
 
 		/// <summary>
-		/// Headers key utilized for expressing missing commits in <see cref="StreamsController.GetChangesAsync"/>.
+		/// Headers key utilized for expressing missing commits in <see cref="GetChangesAsync"/> and <see cref="GetBatchedChangesAsync"/>.
 		/// </summary>
 		public const string MissingCommitsHeader = "X-Missing-Commits";
 
@@ -208,7 +208,7 @@ namespace HordeServer.Streams
 		/// <param name="filter">The filter to apply to the results</param>
 		/// <param name="cancellationToken">Cancellation token for the operation</param>
 		/// Returns 200 OK with list of filtered commits. Partial failures for batch retrieval are indicated via the <see cref="MissingCommitsHeader"/> header.
-		/// <remarks>Batch request mode executed with <paramref name="changeNumbers"/> will store any commits that were not found or that threw exceptions in the <see cref="MissingCommitsHeader"/> response header.</remarks>
+		/// <remarks>Batch request mode executed with <paramref name="changeNumbers"/> will store any commits that were not found or that threw a <see cref="PerforceException"/> in the <see cref="MissingCommitsHeader"/> response header.</remarks>
 		[HttpGet]
 		[Route("/api/v1/streams/{streamId}/changes")]
 		[ProducesResponseType(typeof(List<GetCommitResponse>), 200)]
@@ -271,7 +271,7 @@ namespace HordeServer.Streams
 		/// Returns 200 OK with list of filtered commits. Partial failures for batch retrieval are indicated via the <see cref="MissingCommitsHeader"/> header.
 		/// <remarks>
 		/// This endpoint is a workaround to retrieve large sets of changes which would result in 414 reponses from <see cref="GetChangesAsync"/>.
-		/// Any commits that were not found or that threw exceptions are stored in the <see cref="MissingCommitsHeader"/> response header.
+		/// Any commits that were not found or that threw a <see cref="PerforceException"/> are stored in the <see cref="MissingCommitsHeader"/> response header.
 		/// </remarks>
 		[HttpPost]
 		[Route("/api/v1/streams/{streamId}/changes")]
