[Horde][Analytics][Tests] Fixing test to use proper where clause on any.

diff --git a/Engine/Source/Programs/Horde/Plugins/Analytics/HordeServer.Analytics.Integration.Tests/Telemetry/TelemetryIntegrationTests.cs b/Engine/Source/Programs/Horde/Plugins/Analytics/HordeServer.Analytics.Integration.Tests/Telemetry/TelemetryIntegrationTests.cs
--- a/Engine/Source/Programs/Horde/Plugins/Analytics/HordeServer.Analytics.Integration.Tests/Telemetry/TelemetryIntegrationTests.cs
+++ b/Engine/Source/Programs/Horde/Plugins/Analytics/HordeServer.Analytics.Integration.Tests/Telemetry/TelemetryIntegrationTests.cs
@@ -11,6 +11,7 @@ using EpicGames.Horde.Jobs.Templates;
 using EpicGames.Horde.Projects;
 using EpicGames.Horde.Streams;
 using EpicGames.Horde.Telemetry;
+using EpicGames.Perforce;
 using HordeCommon;
 using HordeServer.Analytics.Tests;
 using HordeServer.Issues;
@@ -25,6 +26,7 @@ using HordeServer.Telemetry.Metrics;
 using HordeServer.Tests;
 using HordeServer.Tests.Issues;
 using Microsoft.Extensions.DependencyInjection;
+using StackExchange.Redis;
 
 namespace HordeServer.Analytics.Integration.Tests.Telemetry
 {
@@ -170,8 +172,8 @@ namespace HordeServer.Analytics.Integration.Tests.Telemetry
 
 				// Verify presence of Issue telemetry
 				IList<Tuple<TelemetryRecordMeta?, Hashtable>> actualTelemetryPayloads = [.. concreteTelemetryWriter.RecordedTelemetry[TelemetryStoreId.Default]];
-				Assert.IsTrue(actualTelemetryPayloads.Select(x => x.Item2[EventNameKeyStr] as string == IssueTelemetry.DefaultEventName).Any());
-				Assert.IsTrue(actualTelemetryPayloads.Select(x => x.Item2[EventNameKeyStr] as string == IssueSpanTelemetry.DefaultEventName).Any());
+				Assert.IsTrue(actualTelemetryPayloads.Select(x => x).Where(x => x.Item2[EventNameKeyStr] as string == IssueTelemetry.DefaultEventName).Any());
+				Assert.IsTrue(actualTelemetryPayloads.Select(x => x).Where(x => x.Item2[EventNameKeyStr] as string == IssueSpanTelemetry.DefaultEventName).Any());
 			}
 
 			// Update the issue and make sure we get back external issues
