Horde: Ignore Horde-based views in Perforce materializer temporarily

diff --git a/Engine/Source/Programs/Horde/Drivers/JobDriver.Tests/PerforceMaterializerTests.cs b/Engine/Source/Programs/Horde/Drivers/JobDriver.Tests/PerforceMaterializerTests.cs
--- a/Engine/Source/Programs/Horde/Drivers/JobDriver.Tests/PerforceMaterializerTests.cs
+++ b/Engine/Source/Programs/Horde/Drivers/JobDriver.Tests/PerforceMaterializerTests.cs
@@ -72,6 +72,7 @@ public class PerforceMaterializerTests : BasePerforceFixtureTest
 	[DataRow("/Engine/Build/...", "//Foo/Main/Engine/Build/... //MyComputer/Engine/Build/...")]
 	[DataRow("-/Engine/Docs/...", "-//Foo/Main/Engine/Docs/... //MyComputer/Engine/Docs/...")]
 	[DataRow("   -/Engine/Docs/...", "-//Foo/Main/Engine/Docs/... //MyComputer/Engine/Docs/...")]
+	[DataRow("/Engine/Build/...", "//Foo/Main/Engine/Build/... //MyComputer/Engine/Build/...")]
 	public void ConvertHordeViews(string hordeView, string expected)
 	{
 		string actual = PerforceMaterializer.ConvertHordeViewToClientView(hordeView, Stream, ClientName);
diff --git a/Engine/Source/Programs/Horde/Drivers/JobDriver/Execution/PerforceMaterializer.cs b/Engine/Source/Programs/Horde/Drivers/JobDriver/Execution/PerforceMaterializer.cs
--- a/Engine/Source/Programs/Horde/Drivers/JobDriver/Execution/PerforceMaterializer.cs
+++ b/Engine/Source/Programs/Horde/Drivers/JobDriver/Execution/PerforceMaterializer.cs
@@ -292,7 +292,7 @@ public sealed class PerforceMaterializer : IWorkspaceMaterializer
 			Root = _workspaceDir.FullName,
 			Options = ClientOptions.Clobber,
 		};
-		newClient.View.AddRange(view);
+		newClient.View.AddRange(streamView.View); // Ignore Horde-based view for now
 		Directory.CreateDirectory(newClient.Root);
 		
 		using ILoggerProgress status = _logger.BeginProgressScope(forceCreate ? "Force creating client ..." : "Updating client ...");
