Horde: Fix issue with ignoring Horde views for Perforce materializer

diff --git a/Engine/Source/Programs/Horde/Drivers/JobDriver/Execution/PerforceMaterializer.cs b/Engine/Source/Programs/Horde/Drivers/JobDriver/Execution/PerforceMaterializer.cs
--- a/Engine/Source/Programs/Horde/Drivers/JobDriver/Execution/PerforceMaterializer.cs
+++ b/Engine/Source/Programs/Horde/Drivers/JobDriver/Execution/PerforceMaterializer.cs
@@ -281,7 +281,7 @@ public sealed class PerforceMaterializer : IWorkspaceMaterializer
 		ClientRecord streamView = await _perforceWithoutClient.GetStreamViewAsync(PlaceholderName, streamName, -1, cancellationToken);
 		
 		List<string> hordeView = _options.AgentWorkspace.View.Select(x => ConvertHordeViewToClientView(x, streamName, clientName)).ToList();
-		List<string> view = [..streamView.View.Select(x => x.Replace(PlaceholderName, clientName, StringComparison.Ordinal)), ..hordeView];
+		List<string> view = [..streamView.View.Select(x => x.Replace(PlaceholderName, clientName, StringComparison.Ordinal))];
 		
 		// Create a classic client while still applying the view from stream together with extra Horde-defined view entries.
 		// A stream-based client won't allow changing the view
@@ -292,7 +292,7 @@ public sealed class PerforceMaterializer : IWorkspaceMaterializer
 			Root = _workspaceDir.FullName,
 			Options = ClientOptions.Clobber,
 		};
-		newClient.View.AddRange(streamView.View); // Ignore Horde-based view for now
+		newClient.View.AddRange(view);
 		Directory.CreateDirectory(newClient.Root);
 		
 		using ILoggerProgress status = _logger.BeginProgressScope(forceCreate ? "Force creating client ..." : "Updating client ...");
