[Horde][BuildHealth][Dashboard] Updating API used for obtaining changelist results.
Unshelved from pending changelist '46710882', to commit on Julian's behalf.

diff --git a/Engine/Source/Programs/Horde/HordeDashboard/src/backend/Backend.ts b/Engine/Source/Programs/Horde/HordeDashboard/src/backend/Backend.ts
--- a/Engine/Source/Programs/Horde/HordeDashboard/src/backend/Backend.ts
+++ b/Engine/Source/Programs/Horde/HordeDashboard/src/backend/Backend.ts
@@ -854,7 +854,7 @@ export class Backend {
         });
     }
 
-    getChangeSummaries(streamId: string, minChange?: number, maxChange?: number, changeNumbers? :number[], filter?: string, maxResults = 100): Promise<ChangeSummaryData[]> {
+    getChangeSummaries(streamId: string, minChange?: number, maxChange?: number, changeNumbers?: number[], filter?: string, maxResults = 100): Promise<ChangeSummaryData[]> {
 
         return new Promise<ChangeSummaryData[]>((resolve, reject) => {
 
@@ -874,8 +874,7 @@ export class Backend {
                 params.filter = filter;
             }
 
-            if(changeNumbers)
-            {
+            if (changeNumbers) {
                 params.changeNumbers = changeNumbers;
             }
 
@@ -891,6 +890,25 @@ export class Backend {
 
     }
 
+    getBatchChangeSummaries(streamId: string, changeNumbers?: number[], filter?: string): Promise<ChangeSummaryData[]> {
+        return new Promise<ChangeSummaryData[]>((resolve, reject) => {
+            const body = changeNumbers ?? [];
+            const params: any = {
+            };
+
+            if (filter) {
+                params.filter = filter;
+            }
+
+            this._fetch.post(`/api/v1/streams/${streamId}/changes`, body, { params }
+            ).then((value) => {
+                resolve(value.data as ChangeSummaryData[]);
+            }).catch(reason => {
+                reject(reason);
+            });
+        });
+    }
+
     getAdminToken(): Promise<string> {
 
         return new Promise<string>((resolve, reject) => {
diff --git a/Engine/Source/Programs/Horde/HordeDashboard/src/components/buildhealth/stepoutcome/StepOutcomeDataHandler.ts b/Engine/Source/Programs/Horde/HordeDashboard/src/components/buildhealth/stepoutcome/StepOutcomeDataHandler.ts
--- a/Engine/Source/Programs/Horde/HordeDashboard/src/components/buildhealth/stepoutcome/StepOutcomeDataHandler.ts
+++ b/Engine/Source/Programs/Horde/HordeDashboard/src/components/buildhealth/stepoutcome/StepOutcomeDataHandler.ts
@@ -251,7 +251,7 @@ export class StepOutcomeDataHandler extends PollBase {
         if (jobHistoryCount.size > 0) {
             // Parallelize backend calls for change summaries
             const changeResults = await Promise.all(
-                changeQueries.map(({ streamId, query }) => backend.getChangeSummaries(streamId, undefined, undefined, [...new Set(jobHistoryCount.get(streamId))], query.filter))
+                changeQueries.map(({ streamId, query }) => backend.getBatchChangeSummaries(streamId, [...new Set(jobHistoryCount.get(streamId))], query.filter))
             );
 
             let flattenedResults = changeResults.flat();
