Horde: Update release notes

#rnx

diff --git a/Engine/Source/Programs/Horde/Docs/ReleaseNotes.md b/Engine/Source/Programs/Horde/Docs/ReleaseNotes.md
--- a/Engine/Source/Programs/Horde/Docs/ReleaseNotes.md
+++ b/Engine/Source/Programs/Horde/Docs/ReleaseNotes.md
@@ -1,5 +1,52 @@
 # Release Notes
 
+## 2025-10-09
+
+* [Build Health Dashboard] Updating API used for obtaining changelist result (46721751)
+* Fix issue with ignoring Horde views for Perforce materializer (46717212)
+* Ignore Horde-based views in Perforce materializer temporarily (46714793)
+* Fixing test to use proper where clause on any. (46705063)
+* Missed an edge case for automatic schema upgrades - select all was implied on schema 1 for a template; we are now explicit in schema version 2+. (46701236)
+* Fully qualified label names + categories. (46699620)
+* [Build Health Dashboard][Step Outcome] Introduce POST API for changes in order to support single request of batch data (46699481)
+* StepOutcome table supports N values for the same cell (46682375)
+* Avoid infinite trace span for background task in Perforce cache (46667057)
+* EpicGames.Perforce: Add OpenTelemetry tracing (46665044)
+* Add RCA Change Backout category (46662844)
+* Bumping SchemaVersion now we have added RCA data. (46607473)
+* Updated RCA information to be included in the issue telemetry. (46605642)
+* Apply tree filtering to current APIs (46605166)
+* Add extended bearer token ACL (46601390)
+* Surface Gauntlet ensure warnings and make the issue summary reflects having a callstack (46583249)
+* Fix server warning in Slack notifications (46580566)
+* Adding Root Cause Analysis Slack support (46580346)
+* Lint support & unit test support (46580285)
+* Collection of QoL UI features for interacting with dropdowns:
+  - clear all + select all (for each group)
+  -- filter context is applied
+  - sticky search field
+  - Jump to top/bottom of scrollable
+  - Jump up group and down group
+  - optimizations to dropdown to reduce load on heavy case (removal of ToolTipHost)
+  - Compression of step names in query param (remove duplicates of data across all stepnames)
+  -- query param schema version & auto-upgrades to v1
+  - NaN fix for summary where no completed steps exist (46580282)
+* Fix issue with issue generation delta comparison (46577182)
+* Address issue update, issues (46576547)
+* Do not generate a flood of Build Health issues for a small set of known Gauntlet log event warnings
+  If we identify cases where we want these warnings, they should be defined as new log events, and force their warning in the issue handler
+  Though, we have been unable to come up with a single case in exploration (46575909)
+* Instrument CachedCommitSource.FindAsync (46518362)
+* Downgrade OpenTelemetry.Instrumentation.StackExchangeRedis to match OTel version
+  * 1.0.0-rc9.15 requires OpenTelemetry 1.9.0+
+  * Everything compiles fine but fails silently at runtime unless you enable diagnostics logs for OTel (46502090)
+* More tweaks to the Issue Modal for new functionality (46465830)
+* Fix clearing root cause change in modal (46462618)
+* Horde - Removed running of deleted test project HordeServer.Ddc.Tests.csproj (46462588)
+* Fixing some issues with the issue view (46462221)
+* Deleted HordeServer.DDC - This was never a supported way to use Horde as DDC cache. We recommend using Zen for most of these needs, or Cloud DDC if you need a cloud based geo replicated DDC. (46461947)
+* Tweak deployment issue messaging (46460552)
+
 ## 2025-10-02
 
 * Root Cause Analysis, dashboard side (46444423)
