Horde: Replace PerformanceCounter with WINAPI GetSystemTimes

Added WINAPI functionality to WindowsCapabilities to capture CPU metrics and call the new function from within TelemetryService.
This allows a more direct approach rather than relying on PerformanceCounter and is supported by both Windows and Wine.

Merged on behalf of Daniel Holden.

diff --git a/Engine/Source/Programs/Horde/HordeAgent/Services/TelemetryService.cs b/Engine/Source/Programs/Horde/HordeAgent/Services/TelemetryService.cs
--- a/Engine/Source/Programs/Horde/HordeAgent/Services/TelemetryService.cs
+++ b/Engine/Source/Programs/Horde/HordeAgent/Services/TelemetryService.cs
@@ -138,15 +138,8 @@ public sealed class DefaultSystemMetrics : ISystemMetrics
 /// Windows specific implementation for gathering system metrics
 /// </summary>
 [SupportedOSPlatform("windows")]
-public sealed class WindowsSystemMetrics : ISystemMetrics, IDisposable
+public sealed class WindowsSystemMetrics : ISystemMetrics
 {
-	private const string ProcessorInfo = "Processor Information"; // Prefer this over "Processor" as it's more modern
-	private const string Total = "_Total";
-
-	private readonly PerformanceCounter _procIdleTime = new(ProcessorInfo, "% Idle Time", Total);
-	private readonly PerformanceCounter _procUserTime = new(ProcessorInfo, "% User Time", Total);
-	private readonly PerformanceCounter _procPrivilegedTime = new(ProcessorInfo, "% Privileged Time", Total);
-
 	private readonly DirectoryReference _workingDir;
 	private readonly Stopwatch _driveInfoUpdateTimer = Stopwatch.StartNew();
 	private DriveInfo? _driveInfo;
@@ -162,27 +155,20 @@ public sealed class WindowsSystemMetrics : ISystemMetrics, IDisposable
 		GetCpu(); // Trigger this to ensure performance counter has a fetched value. Avoids an initial zero result when called later.
 	}
 
-	/// <inheritdoc/>
-	public void Dispose()
-	{
-		_procIdleTime.Dispose();
-		_procUserTime.Dispose();
-		_procPrivilegedTime.Dispose();
-	}
-
-	/// <inheritdoc />
-	public CpuMetrics GetCpu()
-	{
-		return new CpuMetrics
-		{
-			User = _procUserTime.NextValue(),
-			System = _procPrivilegedTime.NextValue(),
-			Idle = _procIdleTime.NextValue()
-		};
-	}
-
-	/// <inheritdoc />
-	public MemoryMetrics GetMemory()
+    /// <inheritdoc />
+    public CpuMetrics GetCpu()
+    {
+        (float idle, float system, float user) = WindowsCapabilities.GetCpuMetrics();
+        return new CpuMetrics
+        {
+            User = user,
+            System = system,
+            Idle = idle
+        };
+    }
+
+    /// <inheritdoc />
+    public MemoryMetrics GetMemory()
 	{
 		(ulong totalBytes, ulong availableBytes) = WindowsCapabilities.GetPhysicalMemory();
 		ulong usedBytes = totalBytes - availableBytes;
diff --git a/Engine/Source/Programs/Horde/HordeAgent/Services/WindowsCapabilities.cs b/Engine/Source/Programs/Horde/HordeAgent/Services/WindowsCapabilities.cs
--- a/Engine/Source/Programs/Horde/HordeAgent/Services/WindowsCapabilities.cs
+++ b/Engine/Source/Programs/Horde/HordeAgent/Services/WindowsCapabilities.cs
@@ -89,8 +89,64 @@ public class WindowsCapabilities(IOptionsMonitor<AgentSettings> settings, ILogge
 		}
 		return (0, 0);
 	}
-	
-	private static int GetNumLogicalCores()
+
+    private sealed class CpuMetrics
+    {
+        internal static FILETIME Idle, System, User;
+
+        internal static bool IsFirstRun()
+        {
+            return Idle.IsZero && System.IsZero && User.IsZero;
+        }
+    }
+
+    /// <summary>
+    /// Gets details of the current system's CPU
+    /// </summary>
+    /// <returns>Non-zero if the function succeeds, zero if it fails. On success, the output parameters contain cumulative idle, system, and user times.</returns>
+    public static (float idle, float system, float user) GetCpuMetrics()
+    {
+        try
+        {
+            if (GetSystemTimes(out FILETIME idle, out FILETIME system, out FILETIME user))
+            {
+                // On first run, assign CPU metrics and return (mimics PerformanceCounter logic).
+                if (CpuMetrics.IsFirstRun())
+                {
+                    // Snapshot CPU metrics and return.
+                    CpuMetrics.Idle = idle;
+                    CpuMetrics.System = system;
+                    CpuMetrics.User = user;
+                    return (0.0f, 0.0f, 0.0f);
+                }
+
+                // Compute deltas.
+                ulong idleDelta = idle.ToUInt64() - CpuMetrics.Idle.ToUInt64();
+                ulong systemDelta = system.ToUInt64() - CpuMetrics.System.ToUInt64();
+                ulong userDelta = user.ToUInt64() - CpuMetrics.User.ToUInt64();
+                ulong total = systemDelta + userDelta;
+
+                // Compute percentages.
+                float idlePercentage = total > 0 ? (idleDelta * 100.0f) / total : 0.0f;
+                float systemPercentage = total > 0 ? ((systemDelta - idleDelta) * 100.0f) / total : 0.0f;
+                float userPercentage = total > 0 ? (userDelta * 100.0f) / total : 0.0f;
+
+                // Update previous snapshots.
+                CpuMetrics.Idle = idle;
+                CpuMetrics.System = system;
+                CpuMetrics.User = user;
+
+                return (idlePercentage, systemPercentage, userPercentage);
+            }
+        }
+        catch
+        {
+            // Ignore any error.
+        }
+        return (0.0f, 0.0f, 0.0f);
+    }
+
+    private static int GetNumLogicalCores()
 	{
 		int total = 0;
 		ushort groupCount = GetActiveProcessorGroupCount();
@@ -188,6 +244,19 @@ public class WindowsCapabilities(IOptionsMonitor<AgentSettings> settings, ILogge
 		public ulong ullAvailVirtual;
 		public ulong ullAvailExtendedVirtual;
 	}
+
+    [DllImport("kernel32.dll", SetLastError = true)]
+    [return: MarshalAs(UnmanagedType.Bool)]
+    static extern bool GetSystemTimes(out FILETIME idle, out FILETIME system, out FILETIME user);
+
+    [StructLayout(LayoutKind.Sequential)]
+    internal struct FILETIME
+    {
+        public uint dwLowDateTime;
+        public uint dwHighDateTime;
+        public ulong ToUInt64() => ((ulong)dwHighDateTime << 32) | dwLowDateTime;
+        public bool IsZero => dwLowDateTime == 0 && dwHighDateTime == 0;
+    }
 	#pragma warning enable
 	
 	#endregion P/Invoke
