[Horde][BuildHealth][Dashboard] Increase per stream limit for historical views.

diff --git a/Engine/Source/Programs/Horde/HordeDashboard/src/components/buildhealth/BuildhealthView.tsx b/Engine/Source/Programs/Horde/HordeDashboard/src/components/buildhealth/BuildhealthView.tsx
--- a/Engine/Source/Programs/Horde/HordeDashboard/src/components/buildhealth/BuildhealthView.tsx
+++ b/Engine/Source/Programs/Horde/HordeDashboard/src/components/buildhealth/BuildhealthView.tsx
@@ -18,6 +18,7 @@ import { stepOutcomeTableClasses } from "./stepoutcome/StepOutcomeSharedUICompon
 import { runInAction } from "mobx";
 import { StepOutcomeLegend } from "./stepoutcome/StepOutcomeLegend";
 import { BuildHealthOptionsController } from "./BuildHealthOptionsController";
+import { StepOutcomeDataHandler } from "./stepoutcome/StepOutcomeDataHandler";
 
 // #region -- Dropdown Components --
 
@@ -798,7 +799,6 @@ const BuildHealthSidebar: React.FC<{ handler: BuildHealthDataHandler, buildHealt
     // #endregion -- Property Use Effects --
 
     let timeComboWidth = 150;
-
     return (
         <Stack style={{ minWidth: 375, maxWidth: 500, paddingRight: 18 }}>
             <Stack className={hordeClasses.modal} tokens={{ childrenGap: 12 }}>
@@ -845,7 +845,7 @@ const BuildHealthSidebar: React.FC<{ handler: BuildHealthDataHandler, buildHealt
                                 });
                             }} />
                         </TooltipHost>
-                        <TooltipHost content="Adjust the historical job history time span."
+                        <TooltipHost content={`Adjust the historical job history time span. Note: This is a best effort due to load limitations (we constrain a max job result count: ${StepOutcomeDataHandler.MAX_JOB_COUNT_DEFAULT}).`}
                             calloutProps={{
                                 directionalHint: DirectionalHint.rightCenter
                             }}
diff --git a/Engine/Source/Programs/Horde/HordeDashboard/src/components/buildhealth/stepoutcome/StepOutcomeDataHandler.ts b/Engine/Source/Programs/Horde/HordeDashboard/src/components/buildhealth/stepoutcome/StepOutcomeDataHandler.ts
--- a/Engine/Source/Programs/Horde/HordeDashboard/src/components/buildhealth/stepoutcome/StepOutcomeDataHandler.ts
+++ b/Engine/Source/Programs/Horde/HordeDashboard/src/components/buildhealth/stepoutcome/StepOutcomeDataHandler.ts
@@ -109,7 +109,7 @@ export class StepOutcomeDataHandler extends PollBase {
     private readonly STREAM_JOBS_FILTER = "id,streamId,name,change,templateId,state,createTime,updateTime,preflightCommitId,preflightDescription,cancellationReason,labels,defaultLabel";
     private readonly STREAM_JOBS_BATCH_STEP_FILTER = "batches.steps.id,batches.steps.name,batches.steps.state,batches.steps.outcome,batches.steps.startTime,batches.steps.finishTime,batches.steps.error"
     private readonly STREAM_CHANGES_FILTER = "id,dateUtc";
-    private readonly MAX_JOB_COUNT_DEFAULT = 150;
+    static readonly MAX_JOB_COUNT_DEFAULT = 400;
 
     // #region -- Private Members -- 
 
@@ -182,7 +182,7 @@ export class StepOutcomeDataHandler extends PollBase {
 
         const streamList = this.activeFilter.streams ?? [];
         const modifiedAfter = this.activeFilter.jobHistorySpan.start.toISOString();
-        const jobCount = this.activeFilter.maxJobCount ?? this.MAX_JOB_COUNT_DEFAULT;
+        const jobCount = this.activeFilter.maxJobCount ?? StepOutcomeDataHandler.MAX_JOB_COUNT_DEFAULT;
 
         // Group jobs by streamId
         const jobsByStream: Record<string, string[]> = {};
