Horde: Fix issue with assignment/acknowledgement

diff --git a/Engine/Source/Programs/Horde/HordeDashboard/src/components/IssueViewV2.tsx b/Engine/Source/Programs/Horde/HordeDashboard/src/components/IssueViewV2.tsx
--- a/Engine/Source/Programs/Horde/HordeDashboard/src/components/IssueViewV2.tsx
+++ b/Engine/Source/Programs/Horde/HordeDashboard/src/components/IssueViewV2.tsx
@@ -2399,8 +2399,8 @@ const MarkFixedModal: React.FC<{ onClose: () => void, autoFix?:boolean } > = ({
          rootCommitId: rootCommitId ?? "",
          rootCauseCategory: rootCauseCategory,
          rootCauseSummary: rootCauseSummary,
-         ownerId: !issue.ownerInfo ? dashboard.userId : issue.ownerInfo.id,
-         acknowledged: !issue.ownerInfo || issue.ownerInfo?.id === dashboard.userId ? true : undefined,
+         //ownerId: !issue.ownerInfo ? dashboard.userId : issue.ownerInfo.id,
+         //acknowledged: !issue.ownerInfo || issue.ownerInfo?.id === dashboard.userId ? true : undefined,
          resolved: clearFixed ? false : (fixCL != undefined && fixCL >= 0),
          duplicateIssueId: dupeId
       }
