[Horde][Analytics] Sending up schema version as well.

diff --git a/Engine/Source/Programs/Horde/Plugins/Build/HordeServer.Build/Jobs/JobTelemetry.cs b/Engine/Source/Programs/Horde/Plugins/Build/HordeServer.Build/Jobs/JobTelemetry.cs
--- a/Engine/Source/Programs/Horde/Plugins/Build/HordeServer.Build/Jobs/JobTelemetry.cs
+++ b/Engine/Source/Programs/Horde/Plugins/Build/HordeServer.Build/Jobs/JobTelemetry.cs
@@ -17,6 +17,11 @@ namespace HordeServer.Jobs
 		/// </summary>
 		public string EventName { get; init; }
 
+		/// <summary>
+		/// The schema version of the telemetry.
+		/// </summary>
+		public int SchemaVersion { get; init; }
+
 		/// <summary>
 		/// The job Id.
 		/// </summary>
@@ -120,6 +125,7 @@ namespace HordeServer.Jobs
 		public JobSummaryTelemetry(IJob job, DateTime finishTimeUtc)
 		{
 			EventName = DefaultEventName;
+			SchemaVersion = CurrentSchemaVersion;
 			JobId = job.Id;
 			TemplateId = job.TemplateId;
 			StreamId = job.StreamId;
@@ -153,6 +159,11 @@ namespace HordeServer.Jobs
 		/// </summary>
 		public string EventName { get; init; }
 
+		/// <summary>
+		/// The schema version of the telemetry.
+		/// </summary>
+		public int SchemaVersion { get; init; }
+
 		/// <summary>
 		/// The job Id.
 		/// </summary>
@@ -226,6 +237,7 @@ namespace HordeServer.Jobs
 		public JobLabelSummaryTelemetry(IJob job, IGraph graph)
 		{
 			EventName = DefaultEventName;
+			SchemaVersion = CurrentSchemaVersion;
 			SuccessfulLabels = [];
 			WarningLabels = [];
 			FailedLabels = [];
