[Horde][Dashboard] Use full stream name for tab title if available

diff --git a/Engine/Source/Programs/Horde/HordeDashboard/src/components/StreamView.tsx b/Engine/Source/Programs/Horde/HordeDashboard/src/components/StreamView.tsx
--- a/Engine/Source/Programs/Horde/HordeDashboard/src/components/StreamView.tsx
+++ b/Engine/Source/Programs/Horde/HordeDashboard/src/components/StreamView.tsx
@@ -377,7 +377,7 @@ const StreamViewInner: React.FC = observer(() => {
       }
    ];
 
-   const crumbTitle = `Horde: //${stream.project?.name}/${stream.name}`;
+   const crumbTitle = `Horde: ${stream.fullname ?? `//${stream.project?.name}/${stream.name}`}`;
 
    const onRenderPlainCard = (tab: GetStreamTabResponse) => {
       return <Stack>
