Horde: Build health make issue suspect max days and scoring time configuration, also bump defaults

diff --git a/Engine/Source/Programs/Horde/Plugins/Build/HordeServer.Build/BuildConfig.cs b/Engine/Source/Programs/Horde/Plugins/Build/HordeServer.Build/BuildConfig.cs
--- a/Engine/Source/Programs/Horde/Plugins/Build/HordeServer.Build/BuildConfig.cs
+++ b/Engine/Source/Programs/Horde/Plugins/Build/HordeServer.Build/BuildConfig.cs
@@ -89,6 +89,16 @@ namespace HordeServer
 		/// </summary>
 		public string IssueFixedTag { get; set; } = "#horde";
 
+		/// <summary>
+		/// The maximum days since node success to scan for suspects
+		/// </summary>
+		public int IssueFindSuspectsMaxSuccessDays { get; set; } = 14;
+
+		/// <summary>
+		/// The maximum number of seconds to spend scording suspects
+		/// </summary>
+		public int IssueFindSuspectsMaxScoringSeconds { get; set; } = 20;
+
 		private readonly Dictionary<ProjectId, ProjectConfig> _projectLookup = new Dictionary<ProjectId, ProjectConfig>();
 		private readonly Dictionary<StreamId, StreamConfig> _streamLookup = new Dictionary<StreamId, StreamConfig>();
 
diff --git a/Engine/Source/Programs/Horde/Plugins/Build/HordeServer.Build/Issues/IssueService.cs b/Engine/Source/Programs/Horde/Plugins/Build/HordeServer.Build/Issues/IssueService.cs
--- a/Engine/Source/Programs/Horde/Plugins/Build/HordeServer.Build/Issues/IssueService.cs
+++ b/Engine/Source/Programs/Horde/Plugins/Build/HordeServer.Build/Issues/IssueService.cs
@@ -1409,9 +1409,9 @@ namespace HordeServer.Issues
 
 					TimeSpan delta = firstFailureCommit.DateUtc - lastSuccessCommit.DateUtc;
 					
-					if (Math.Abs(delta.TotalDays) >= 7)
+					if (Math.Abs(delta.TotalDays) >= _buildConfig.CurrentValue.IssueFindSuspectsMaxSuccessDays)
 					{
-						_logger.LogInformation("Ignoring commits for {StreamName}: {LastSuccess} and {FirstFailure}, delta is {DeltaDays}", streamConfig.Name, lastSuccess, firstFailure.Name, delta.TotalDays);
+						_logger.LogWarning("Ignoring commits for {StreamName}: {LastSuccess} and {FirstFailure}, delta is {DeltaDays}", streamConfig.Name, lastSuccess, firstFailure.Name, delta.TotalDays);
 						return suspects;
 					}
 
@@ -1479,9 +1479,9 @@ namespace HordeServer.Issues
 					foreach (ICommit commit in changes)
 					{
 						TimeSpan scoringDuration = DateTime.Now - scoringStart;
-						if (scoringDuration.TotalSeconds >= 10)
+						if (scoringDuration.TotalSeconds >= _buildConfig.CurrentValue.IssueFindSuspectsMaxScoringSeconds)
 						{
-							_logger.LogWarning("Scoring suspect commits took longer than 10 seconds, giving up");
+							_logger.LogWarning("Scoring suspect commits took longer than {ScoringSeconds} seconds, giving up", _buildConfig.CurrentValue.IssueFindSuspectsMaxScoringSeconds);
 							break;
 						}
 
