Horde: Include job name in <title> for job details view

diff --git a/Engine/Source/Programs/Horde/HordeDashboard/src/components/jobDetailsV2/JobDetailViewV2.tsx b/Engine/Source/Programs/Horde/HordeDashboard/src/components/jobDetailsV2/JobDetailViewV2.tsx
--- a/Engine/Source/Programs/Horde/HordeDashboard/src/components/jobDetailsV2/JobDetailViewV2.tsx
+++ b/Engine/Source/Programs/Horde/HordeDashboard/src/components/jobDetailsV2/JobDetailViewV2.tsx
@@ -107,7 +107,7 @@ const JobBreadCrumbs: React.FC<{ jobDetails: JobDetailsV2 }> = observer(({ jobDe
 
    let crumbTitle = `${clText}`;
    if (stream) {
-      crumbTitle = `${stream.fullname} - ${clText}`;
+      crumbTitle = `${stream.fullname} - ${jobData.name} - ${clText}`;
    }
 
    if (stepId) {
