[Horde][BuildHealth][Dashboard] Label support within table

diff --git a/Engine/Source/Programs/Horde/HordeDashboard/src/components/buildhealth/BuildHealthDataHandler.ts b/Engine/Source/Programs/Horde/HordeDashboard/src/components/buildhealth/BuildHealthDataHandler.ts
--- a/Engine/Source/Programs/Horde/HordeDashboard/src/components/buildhealth/BuildHealthDataHandler.ts
+++ b/Engine/Source/Programs/Horde/HordeDashboard/src/components/buildhealth/BuildHealthDataHandler.ts
@@ -59,7 +59,8 @@ interface StepCacheEntry {
  * @todo Data sharing & caching would ideally occur between the BuildHealth View & Step Outcome.
  */
 export class BuildHealthDataHandler implements IHierarchyValidator {
-    private readonly STREAMS_FILTER = "id,state,streamId,name,batches";
+    private readonly STREAMS_FILTER = "id,state,streamId,name";
+    private readonly STREAM_JOBS_BATCH_STEP_FILTER = "batches.steps.id,batches.steps.name,batches.steps.state,batches.steps.outcome,batches.steps.finishTime"
     private readonly STEP_DATA_SAMPLE_SIZE = 5;
 
     // #region -- Private Members --
@@ -222,7 +223,7 @@ export class BuildHealthDataHandler implements IHierarchyValidator {
                 const query: JobStreamQuery = {
                     template: [templateData.templateId],
                     count: this.STEP_DATA_SAMPLE_SIZE,
-                    filter: this.STREAMS_FILTER
+                    filter: `${this.STREAMS_FILTER},${this.STREAM_JOBS_BATCH_STEP_FILTER}`
                 };
 
                 const jobDatas: JobData[] = await backend.getStreamJobs(templateData.streamId, query);
diff --git a/Engine/Source/Programs/Horde/HordeDashboard/src/components/buildhealth/BuildHealthOptions.ts b/Engine/Source/Programs/Horde/HordeDashboard/src/components/buildhealth/BuildHealthOptions.ts
--- a/Engine/Source/Programs/Horde/HordeDashboard/src/components/buildhealth/BuildHealthOptions.ts
+++ b/Engine/Source/Programs/Horde/HordeDashboard/src/components/buildhealth/BuildHealthOptions.ts
@@ -392,6 +392,7 @@ export async function loadBuildHealthOptionsFromParams(
 export class BuildHealthUIOptionsState {
     @observable includeDateAnchors: boolean = false;
     @observable warningsAsSummaryFailure: boolean = false;
+    @observable showLabels: boolean = false;
 
     constructor() {
         makeObservable(this);
diff --git a/Engine/Source/Programs/Horde/HordeDashboard/src/components/buildhealth/BuildHealthOptionsController.ts b/Engine/Source/Programs/Horde/HordeDashboard/src/components/buildhealth/BuildHealthOptionsController.ts
--- a/Engine/Source/Programs/Horde/HordeDashboard/src/components/buildhealth/BuildHealthOptionsController.ts
+++ b/Engine/Source/Programs/Horde/HordeDashboard/src/components/buildhealth/BuildHealthOptionsController.ts
@@ -484,6 +484,15 @@ export class BuildHealthOptionsController {
         this.uiState.warningsAsSummaryFailure = warningsAsSummaryFailure;
     }
 
+    /**
+     * Sets whether the step outcome table shows label values.
+     * @param showLabels True to show label step data.
+     */
+    @action
+    setShowLabels(showLabels: boolean) {
+        this.uiState.showLabels = showLabels;
+    }
+
     /**
      * Gets a options transaction session to make changes to.
      * @returns The @see IBuildHealthOptionWriter that will enact options changes.
diff --git a/Engine/Source/Programs/Horde/HordeDashboard/src/components/buildhealth/BuildhealthView.tsx b/Engine/Source/Programs/Horde/HordeDashboard/src/components/buildhealth/BuildhealthView.tsx
--- a/Engine/Source/Programs/Horde/HordeDashboard/src/components/buildhealth/BuildhealthView.tsx
+++ b/Engine/Source/Programs/Horde/HordeDashboard/src/components/buildhealth/BuildhealthView.tsx
@@ -16,8 +16,8 @@ import { BuildHealthOptionsState, BuildHealthOptionStateDiff, BuildHealthUIOptio
 import { decodeStepKey, decodeTemplateKey, encodeStepKey, encodeTemplateKey } from "./BuildHealthUtilities";
 import { stepOutcomeTableClasses } from "./stepoutcome/StepOutcomeSharedUIComponents";
 import { runInAction } from "mobx";
-import { StepOutcomeLegend } from "./stepoutcome/StepOutcomeLegend";
 import { BuildHealthOptionsController } from "./BuildHealthOptionsController";
+import { StepOutcomeLegend } from "./stepoutcome/StepOutcomeLegend";
 import { StepOutcomeDataHandler } from "./stepoutcome/StepOutcomeDataHandler";
 
 // #region -- Dropdown Components --
@@ -745,6 +745,7 @@ const BuildHealthSidebar: React.FC<{ handler: BuildHealthDataHandler, buildHealt
 
     const [hideDateAnchors, setHideDateAnchors] = useState(() => buildHealthOptions.uiState.includeDateAnchors)
     const [warningsAsSummaryFailure, setWarningsAsSummaryFailure] = useState(() => buildHealthOptions.uiState.warningsAsSummaryFailure)
+    const [showLabels, setShowLabels] = useState(() => buildHealthOptions.uiState.showLabels)
 
     // #ednregion -- Visual Options --
 
@@ -776,6 +777,11 @@ const BuildHealthSidebar: React.FC<{ handler: BuildHealthDataHandler, buildHealt
         buildHealthOptions.setWarningsAsSummaryFailure(isChecked);
     }
 
+    const updateShowLabels = (isChecked) => {
+        setShowLabels(isChecked);
+        buildHealthOptions.setShowLabels(isChecked);
+    }
+
     // #endregion -- Property Callbacks --
 
     // #region -- Property Use Effects --
@@ -944,6 +950,23 @@ const BuildHealthSidebar: React.FC<{ handler: BuildHealthDataHandler, buildHealt
                             onChange={(_, isChecked) => updateWarningsAsSummaryFailure(!!isChecked)}
                         />
                     </TooltipHost>
+                    <TooltipHost content={
+                        <div style={{ maxWidth: 300 }}>
+                            <div>
+                                Warnings are considered success in succes/failure ratio for summary.
+                            </div>
+                        </div>
+                    }
+                        calloutProps={{
+                            directionalHint: DirectionalHint.rightCenter
+                        }}
+                    >
+                        <Checkbox
+                            label="Show Labels"
+                            checked={showLabels}
+                            onChange={(_, isChecked) => updateShowLabels(!!isChecked)}
+                        />
+                    </TooltipHost>
                 </Stack>
                 <Stack
                     styles={{
@@ -993,10 +1016,12 @@ export const BuildHealthView: React.FC = observer(function ConstructBuildHealthV
 
     const uiOptions = useMemo(() => ({
         includeDateAnchors: buildHealthUIState.includeDateAnchors,
-        warningsAsSummaryFailure: buildHealthUIState.warningsAsSummaryFailure
+        warningsAsSummaryFailure: buildHealthUIState.warningsAsSummaryFailure,
+        showLabels: buildHealthUIState.showLabels
     }),
         // eslint-disable-next-line react-hooks/exhaustive-deps -- All necessary dependencies included for MobX observables
         [
+            buildHealthUIState.showLabels,
             buildHealthUIState.includeDateAnchors,
             buildHealthUIState.warningsAsSummaryFailure
         ]);
diff --git a/Engine/Source/Programs/Horde/HordeDashboard/src/components/buildhealth/stepoutcome/StepOutcome.tsx b/Engine/Source/Programs/Horde/HordeDashboard/src/components/buildhealth/stepoutcome/StepOutcome.tsx
--- a/Engine/Source/Programs/Horde/HordeDashboard/src/components/buildhealth/stepoutcome/StepOutcome.tsx
+++ b/Engine/Source/Programs/Horde/HordeDashboard/src/components/buildhealth/stepoutcome/StepOutcome.tsx
@@ -6,8 +6,8 @@ import { observer } from "mobx-react-lite";
 import React, { useEffect, useState } from 'react';
 import { StepOutcomeD3Table } from "./StepOutcomeD3Component";
 import { prettyPrintStepOutcomeFilters, prettyPrintStepOutcomeFiltersSummary, StepOutcomeDataHandler, StepOutcomeFilters, StepOutcomeViewOptions } from "./StepOutcomeDataHandler";
-import { AbortedStr, StepOutcomeTableEntry } from "./StepOutcomeDataTypes";
-import { sharedBorderStyle } from "./StepOutcomeSharedUIComponents";
+import { AbortedStr, OutcomeTableEntry, isLabelOutcome, isStepOutcome, StepOutcomeTableEntry } from "./StepOutcomeDataTypes";
+import { NOT_APPLICABLE_LABEL, sharedBorderStyle } from "./StepOutcomeSharedUIComponents";
 import { isTerminalState, toDate } from "./StepOutcomeUtilities";
 
 // #region -- View Components --
@@ -35,7 +35,7 @@ const stepOutcomeModalClasses = mergeStyleSets(sharedBorderStyle);
 // #region -- Modal Popup View --
 
 type StepOutcomeFocusModalProps = {
-    item: StepOutcomeTableEntry;
+    item: OutcomeTableEntry;
     onDismiss?: (_ev?: React.MouseEvent<HTMLButtonElement>) => any;
 };
 
@@ -43,29 +43,38 @@ type StepOutcomeFocusModalProps = {
  * React Component representing the Step Outcome focus.
  */
 const StepOutcomeFocusModal: React.FC<StepOutcomeFocusModalProps> = ({ item, onDismiss }) => {
-    const startDate = toDate(item.stepResponse.startTime);
+    const startDate = toDate(item.startTime);
 
     const [issueRequestComplete, setIssueRequestComplete] = useState(false);
 
     useEffect(() => {
-        if (!issueRequestComplete) {
+        if (isStepOutcome(item) && !issueRequestComplete) {
             handler.populateStepOutcomeTableEntryIssueData(item, setIssueRequestComplete);
         }
     }, []);
 
+    if (item === undefined) { return; }
+
+    let stepLink: string | null = null;
+    if (isStepOutcome(item)) {
+        stepLink = `/job/${item.jobId}?step=${item.id}`;
+    } else if (isLabelOutcome(item)) {
+        stepLink = `/job/${item.jobId}?label=${item.id}`;
+    }
+
     const details: (string | JSX.Element)[][] = [
-        ['Date (Job)', item.jobCreateTime ? new Date(item.jobCreateTime).toLocaleString() : "N/A"],
-        ['Date (Step)', startDate ? startDate.toLocaleString() : "N/A"],
+        ['Date (Job)', item.jobCreateTime ? new Date(item.jobCreateTime).toLocaleString() : NOT_APPLICABLE_LABEL],
+        ['Date (Step)', startDate ? startDate.toLocaleString() : NOT_APPLICABLE_LABEL],
         ['Stream', item.streamId],
         ['Code Change', item.change],
         ['Job', item.jobName],
-        ['Step', item.stepResponse.name],
-        ['Duration', item.getDurationString()],
-        ['State', item.stepResponse.state],
+        ['Step', item.name],
+        ['Duration', isStepOutcome(item) ? item.getDurationString() : NOT_APPLICABLE_LABEL],
+        ['State', item.state],
         // For simplicity, we will compeltely override the outcome of this step with the job's outcome, as the step was aborted.
-        item.stepResponse.state === AbortedStr ? ['Job Outcome', handler.getJobResponseData(item.jobId)?.summarizedOutcome ?? "-"] : isTerminalState(item) ? ['Outcome', item.stepResponse.outcome] : null,
+        item.state === AbortedStr ? ['Job Outcome', handler.getJobResponseData(item.jobId)?.summarizedOutcome ?? "-"] : isTerminalState(item) ? ['Outcome', item.outcome] : null,
         ['Job Link', <a href={`/job/${item.jobId}`} target="_blank" rel="noopen er noreferrer">View Job</a>],
-        ['Step Link', <a href={`/job/${item.jobId}?step=${item.stepResponse.id}`} target="_blank" rel="noopener noreferrer">View Step</a>],
+        ['Step Link', <a href={`${stepLink}`} target="_blank" rel="noopener noreferrer">View Step</a>],
     ].filter((d): d is (string | JSX.Element)[] => Boolean(d));
 
     let shouldRenderDuplicateColumn = item.duplicateEntries && item.duplicateEntries.length > 0;
@@ -92,7 +101,7 @@ const StepOutcomeFocusModal: React.FC<StepOutcomeFocusModalProps> = ({ item, onD
             dialogContentProps={{
                 type: DialogType.close,
                 onDismiss: onDismiss,
-                title: item.stepResponse.name,
+                title: item.name,
             }}>
             <Stack>
                 <Stack horizontal tokens={{ childrenGap: 16 }} styles={{ root: { width: '100%', flexGrow: 0 } }}>
@@ -124,7 +133,7 @@ const StepOutcomeFocusModal: React.FC<StepOutcomeFocusModalProps> = ({ item, onD
                                             <Stack key={`modal_issue_item_${index}`} tokens={{ childrenGap: 4 }} className={stepOutcomeModalClasses.bottomBorder}>
                                                 {[
                                                     ['Issue', issue.id],
-                                                    ['Issue Link', <a href={`/job/${item.jobId}?step=${item.stepResponse.id}&issue=${issue.id}`} target="_blank" rel="noopener noreferrer">View Issue</a>],
+                                                    ['Issue Link', isStepOutcome(item) ? <a href={`/job/${item.jobId}?step=${item.id}&issue=${issue.id}`} target="_blank" rel="noopener noreferrer">View Issue</a> : NOT_APPLICABLE_LABEL],
                                                     ['Issue Summary', issue.summary],
                                                     ['Status', issue.resolvedAt ? "Resolved" : "Unresolved"],
                                                     ['Resolved Streams', `[${issue.resolvedStreams.join(', ')}]`],
@@ -152,27 +161,36 @@ const StepOutcomeFocusModal: React.FC<StepOutcomeFocusModalProps> = ({ item, onD
                             </Stack>
                             <Stack style={{ flex: 1, overflowY: 'auto' }}>
                                 <Stack tokens={{ childrenGap: 10 }} styles={{ root: { paddingLeft: 16 } }}>
-                                    {item.duplicateEntries.map((relatedEntry, index) => (
-                                        <Stack
-                                            key={`related_job_item${index}`}
-                                            tokens={{ childrenGap: 4 }}
-                                            className={stepOutcomeModalClasses.bottomBorder}
-                                        >
-                                            {[
-                                                ['Date (Job)', relatedEntry.jobCreateTime ? new Date(relatedEntry.jobCreateTime).toLocaleString() : 'N/A'],
-                                                ['Date (Step)', relatedEntry.stepResponse.startTime ? new Date(relatedEntry.stepResponse.startTime).toLocaleString() : 'N/A'],
-                                                ['Job', relatedEntry.jobName],
-                                                ['Job Link', <a href={`/job/${relatedEntry.jobId}`} target="_blank" rel="noopener noreferrer">View Related Job</a>],
-                                                ['Step Link', <a href={`/job/${relatedEntry.jobId}?step=${relatedEntry.stepResponse.id}`} target="_blank" rel="noopener noreferrer">View Related Step</a>],
-                                                ['State', relatedEntry.stepResponse.state],
-                                            ].map(([label, value], j) => (
-                                                <Stack key={j} horizontal tokens={{ childrenGap: 4 }}>
-                                                    <Label styles={{ root: { minWidth: 80, padding: 0 } }}>{label}:</Label>
-                                                    <Text>{value}</Text>
-                                                </Stack>
-                                            ))}
-                                        </Stack>
-                                    ))}
+                                    {item.duplicateEntries.map((relatedEntry, index) => {
+                                        let relatedStepLink: string | null = null;
+                                        if (isStepOutcome(relatedEntry)) {
+                                            relatedStepLink = `/job/${relatedEntry.jobId}?step=${relatedEntry.id}`;
+                                        } else if (isLabelOutcome(item)) {
+                                            relatedStepLink = `/job/${relatedEntry.jobId}?label=${relatedEntry.id}`;
+                                        }
+
+                                        return (
+                                            <Stack
+                                                key={`related_job_item${index}`}
+                                                tokens={{ childrenGap: 4 }}
+                                                className={stepOutcomeModalClasses.bottomBorder}
+                                            >
+                                                {[
+                                                    ['Date (Job)', relatedEntry.jobCreateTime ? new Date(relatedEntry.jobCreateTime).toLocaleString() : 'N/A'],
+                                                    ['Date (Step)', relatedEntry.startTime ? new Date(relatedEntry.startTime).toLocaleString() : 'N/A'],
+                                                    ['Job', relatedEntry.jobName],
+                                                    ['Job Link', <a href={`/job/${relatedEntry.jobId}`} target="_blank" rel="noopener noreferrer">View Related Job</a>],
+                                                    ['Step Link', <a href={`${relatedStepLink}`} target="_blank" rel="noopener noreferrer">View Related Step</a>],
+                                                    ['State', relatedEntry.state],
+                                                ].map(([label, value], j) => (
+                                                    <Stack key={j} horizontal tokens={{ childrenGap: 4 }}>
+                                                        <Label styles={{ root: { minWidth: 80, padding: 0 } }}>{label}:</Label>
+                                                        <Text>{value}</Text>
+                                                    </Stack>
+                                                ))}
+                                            </Stack>
+                                        )
+                                    })}
                                 </Stack>
                             </Stack>
                         </Stack>
diff --git a/Engine/Source/Programs/Horde/HordeDashboard/src/components/buildhealth/stepoutcome/StepOutcomeD3Component.tsx b/Engine/Source/Programs/Horde/HordeDashboard/src/components/buildhealth/stepoutcome/StepOutcomeD3Component.tsx
--- a/Engine/Source/Programs/Horde/HordeDashboard/src/components/buildhealth/stepoutcome/StepOutcomeD3Component.tsx
+++ b/Engine/Source/Programs/Horde/HordeDashboard/src/components/buildhealth/stepoutcome/StepOutcomeD3Component.tsx
@@ -1,12 +1,12 @@
 // Copyright Epic Games, Inc. All Rights Reserved.
 
-import { DefaultButton, DirectionalHint, IconButton, mergeStyleSets, Stack, TooltipHost } from "@fluentui/react";
+import { DirectionalHint, IconButton, mergeStyleSets, Stack, TooltipHost } from "@fluentui/react";
 import * as d3 from "d3";
 import { getHordeStyling, preloadFonts } from "horde/styles/Styles";
 import { useEffect, useRef } from "react";
 import { StepOutcomeDataHandler } from "./StepOutcomeDataHandler";
-import { AbortedStr, ChangeHeader, ColumnHeader, encodeStepNameFromStepNameHeader, isChangeHeader, isDateHeader, isStepOutcome, isSummary, isSummaryHeader, StepNameHeader, StepOutcomeTable, TableEntry } from "./StepOutcomeDataTypes";
-import { getCellAccents, getCellClasses, getCellStyle, getColorRecords, IncrementalRefreshIndicatorPanel, resetCellClasses, sharedBorderStyle } from "./StepOutcomeSharedUIComponents";
+import { AbortedStr, ChangeHeader, ColumnHeader, encodeStepNameFromStepNameHeader, isChangeHeader, isDateHeader, isLabelOutcome, isStepOutcome, isSummary, isSummaryHeader, StepNameHeader, StepOutcomeTable, TableEntry } from "./StepOutcomeDataTypes";
+import { getCellAccents, getCellClasses, getCellStyle, getColorRecords, IncrementalRefreshIndicatorPanel, NOT_APPLICABLE_LABEL, resetCellClasses, sharedBorderStyle } from "./StepOutcomeSharedUIComponents";
 import { getHordeTheme } from "horde/styles/theme";
 import { observer } from "mobx-react-lite";
 import { isTerminalState, toLocalISOString } from "./StepOutcomeUtilities";
@@ -56,11 +56,19 @@ function exportStepOutcomeDataToCSV(stepOutcomeTable: StepOutcomeTable, handler:
     const filename = window.prompt("Enter name of export:", "ProjectExport");
     if (!filename) return;
 
-    const stepOrderedCsvContent = stepOutcomeTable.toRowOrderedCSV(handler);
-    const changeOrderedCsvContent = stepOutcomeTable.toColOrderedCSV(handler);
+    const stepOrderedCsvContent = stepOutcomeTable.toRowOrderedCSV(handler, isStepOutcome);
+    const changeOrderedCsvContent = stepOutcomeTable.toColOrderedCSV(handler, isStepOutcome);
 
     saveFile(stepOrderedCsvContent, `${filename}-StepOrderedData-${defaultName}.csv`);
     saveFile(changeOrderedCsvContent, `${filename}-ChangeOrderedData-${defaultName}.csv`);
+
+    if (stepOutcomeTable.supportsLabelEntries) {
+        const labelOrderedCsvContent = stepOutcomeTable.toRowOrderedCSV(handler, isLabelOutcome);
+        const labelchangeOrderedCsvContent = stepOutcomeTable.toColOrderedCSV(handler, isLabelOutcome);
+
+        saveFile(labelOrderedCsvContent, `${filename}-LabelOrderedData-${defaultName}.csv`);
+        saveFile(labelchangeOrderedCsvContent, `${filename}-LabelChangeOrderedData-${defaultName}.csv`);
+    }
 }
 
 // #endregion -- Step Outcome Table Download Handlers --
@@ -339,7 +347,7 @@ function GenerateD3Table(stepOutcomeTable: StepOutcomeTable, handler: StepOutcom
             });
 
         const rows = scaffold.current.tbody.selectAll("tr")
-            .data(stepOutcomeTable.stepNameRowHeaders as StepNameHeader[])
+            .data([...stepOutcomeTable.getStepRowHeaders()] as StepNameHeader[])
             .join("tr")
             .style("height", "30px")
             .style("padding-top", "4px")
@@ -450,6 +458,7 @@ function GenerateD3Table(stepOutcomeTable: StepOutcomeTable, handler: StepOutcom
                 .style("left", "32px")
                 .style("z-index", "1")
                 .style("background-color", modeColors.content)
+                .style("border", d => d.isLabel ? `1px dashed ${theme.semanticColors.bodyBackgroundHovered}` : `1px solid ${modeColors.background}`)
                 .style("white-space", "nowrap")
                 .style("overflow", "hidden")
                 .style("padding-left", "8px")
@@ -465,7 +474,14 @@ function GenerateD3Table(stepOutcomeTable: StepOutcomeTable, handler: StepOutcom
 
                     if (castedStepNameHeader) {
                         const stepDiv = document.createElement("div");
-                        stepDiv.innerHTML = `<strong>Step:</strong><p class="${tooltipClasses.tooltipLine}">${castedStepNameHeader.stepName}</p>`;
+
+                        if (castedStepNameHeader.isLabel) {
+                            stepDiv.innerHTML = `<strong>Label:</strong><p class="${tooltipClasses.tooltipLine}">${castedStepNameHeader.stepName}</p>`;
+                        }
+                        else {
+                            stepDiv.innerHTML = `<strong>Step:</strong><p class="${tooltipClasses.tooltipLine}">${castedStepNameHeader.stepName}</p>`;
+                        }
+
                         tooltipContent.appendChild(stepDiv);
 
                         const streamDiv = document.createElement("div");
@@ -481,8 +497,6 @@ function GenerateD3Table(stepOutcomeTable: StepOutcomeTable, handler: StepOutcom
                 });
         });
 
-        rows.selectAll("td.step-name").style("border", `1px solid ${modeColors.background}`);
-
         const styleAndBindStepCell = (
             selection: d3.Selection<HTMLDivElement, TableEntry, any, any>
         ) => {
@@ -503,7 +517,7 @@ function GenerateD3Table(stepOutcomeTable: StepOutcomeTable, handler: StepOutcom
                         div.text("").style("background-color", "transparent").style("cursor", "default");
                         return;
                     }
-                    if (isStepOutcome(entry)) { // StepOutcomeTableEntry
+                    if (isStepOutcome(entry) || isLabelOutcome(entry)) { // IOutcomeTableEtnry
                         // UE-324856 - This is a temporary fix until the horde server is fixed.
                         const original = new Date(entry.changeDate!);
                         const corrected = new Date(original.getTime() - 4 * 60 * 60 * 1000);
@@ -525,21 +539,22 @@ function GenerateD3Table(stepOutcomeTable: StepOutcomeTable, handler: StepOutcom
                                     dateDiv.innerHTML += `<p class="${tooltipClasses.tooltipLineWithBorder}">(Local) ${toLocalISOString(corrected) ?? "-"}</p>`
                                     tooltipContent.appendChild(dateDiv);
 
-                                    const durationDiv = document.createElement("div");
-                                    durationDiv.innerHTML = `<strong>Duration:</strong><p class="${tooltipClasses.tooltipLine}">${entry.getDurationString?.() ?? "-"}</p>`;
-                                    tooltipContent.appendChild(durationDiv);
-
+                                    if (isStepOutcome(entry)) {
+                                        const durationDiv = document.createElement("div");
+                                        durationDiv.innerHTML = `<strong>Duration:</strong><p class="${tooltipClasses.tooltipLine}">${entry.getDurationString?.() ?? "-"}</p>`;
+                                        tooltipContent.appendChild(durationDiv);
+                                    }
 
                                     const stateDiv = document.createElement("div");
-                                    stateDiv.innerHTML = `<strong>State:</strong><p class="${tooltipClasses.tooltipLine}">${entry.stepResponse?.state ?? "-"}</p>`;
+                                    stateDiv.innerHTML = `<strong>State:</strong><p class="${tooltipClasses.tooltipLine}">${entry?.state ?? "-"}</p>`;
                                     tooltipContent.appendChild(stateDiv);
 
                                     if (isTerminalState(entry)) {
                                         const outcomeDiv = document.createElement("div");
-                                        outcomeDiv.innerHTML = `<strong>Outcome:</strong><p class="${tooltipClasses.tooltipLine}">${entry.stepResponse?.outcome ?? "-"}</p>`;
+                                        outcomeDiv.innerHTML = `<strong>Outcome:</strong><p class="${tooltipClasses.tooltipLine}">${entry?.outcome ?? "-"}</p>`;
 
                                         // For simplicity, we will compeltely override the outcome of this step with the job's outcome, as the step was aborted.
-                                        if (entry.stepResponse.state === AbortedStr) {
+                                        if (entry.state === AbortedStr) {
                                             outcomeDiv.innerHTML = `<strong>Job Outcome:</strong><p class="${tooltipClasses.tooltipLine}">${handler.getJobResponseData(entry.jobId)?.summarizedOutcome ?? "-"}</p>`;
                                         }
 
@@ -555,7 +570,7 @@ function GenerateD3Table(stepOutcomeTable: StepOutcomeTable, handler: StepOutcom
                                     tooltipContent.appendChild(jobDiv);
 
                                     const stepDiv = document.createElement("div");
-                                    stepDiv.innerHTML = `<strong>Step:</strong><p class="${tooltipClasses.tooltipLine}">${entry.stepResponse?.name ?? "-"}</p>`;
+                                    stepDiv.innerHTML = `<strong>Step:</strong><p class="${tooltipClasses.tooltipLine}">${entry?.name ?? "-"}</p>`;
                                     tooltipContent.appendChild(stepDiv);
 
                                     attachTargettedToolTip(this as HTMLElement, scaffold.current!.tooltip, tooltipContent);
@@ -577,7 +592,7 @@ function GenerateD3Table(stepOutcomeTable: StepOutcomeTable, handler: StepOutcom
                             ratio = numerator / entry.stepsCompleted;
                         }
 
-                        let ratioStr: string = ratio === -1 ? "N/A" : `${(ratio * 100).toFixed(1)}%`;
+                        let ratioStr: string = ratio === -1 ? NOT_APPLICABLE_LABEL : `${(ratio * 100).toFixed(1)}%`;
 
                         div.style("cursor", "default")
                             .style("background-color", summaryColorScale(ratio))
diff --git a/Engine/Source/Programs/Horde/HordeDashboard/src/components/buildhealth/stepoutcome/StepOutcomeDataHandler.ts b/Engine/Source/Programs/Horde/HordeDashboard/src/components/buildhealth/stepoutcome/StepOutcomeDataHandler.ts
--- a/Engine/Source/Programs/Horde/HordeDashboard/src/components/buildhealth/stepoutcome/StepOutcomeDataHandler.ts
+++ b/Engine/Source/Programs/Horde/HordeDashboard/src/components/buildhealth/stepoutcome/StepOutcomeDataHandler.ts
@@ -1,8 +1,8 @@
 // Copyright Epic Games, Inc. All Rights Reserved.
 
-import { GetJobResponse, IssueQuery, JobStreamQuery } from "horde/backend/Api";
+import { GetJobResponse, IssueQuery, JobStepState, JobStreamQuery, LabelState } from "horde/backend/Api";
 import { PollBase } from "horde/backend/PollBase";
-import { JobResponseData, StepOutcomeTable, StepOutcomeTableEntry, UpdatedChangeSummaryData } from "./StepOutcomeDataTypes";
+import { JobResponseData, LabelOutcomeTableEntry, StepOutcomeTable, StepOutcomeTableEntry, UpdatedChangeSummaryData } from "./StepOutcomeDataTypes";
 import backend from "horde/backend";
 import { decodeStepKey, decodeTemplateKey, encodeStepNameFromStrings } from "../BuildHealthUtilities";
 import { action, computed, makeObservable, observable, runInAction } from "mobx";
@@ -84,6 +84,7 @@ export function prettyPrintStepOutcomeFiltersSummary(filters: StepOutcomeFilters
 export interface StepOutcomeViewOptions {
     includeDateAnchors: boolean;
     warningsAsSummaryFailure: boolean;
+    showLabels: boolean;
 }
 
 /**
@@ -106,8 +107,8 @@ export interface StepOutcomeFilters {
  * @todo Data sharing & caching would ideally occur between the BuildHealth View & Step Outcome.
  */
 export class StepOutcomeDataHandler extends PollBase {
-    private readonly STREAM_JOBS_FILTER = "id,streamId,name,change,templateId,state,createTime,updateTime,preflightCommitId,preflightDescription,cancellationReason,labels,defaultLabel";
-    private readonly STREAM_JOBS_BATCH_STEP_FILTER = "batches.steps.id,batches.steps.name,batches.steps.state,batches.steps.outcome,batches.steps.startTime,batches.steps.finishTime,batches.steps.error"
+    private readonly STREAM_JOBS_FILTER = "id,streamId,name,change,templateId,state,createTime,updateTime,abortedByUserInfo,preflightCommitId,preflightDescription,cancellationReason,labels,defaultLabel";
+    private readonly STREAM_JOBS_BATCH_STEP_FILTER = "batches.steps.id,batches.steps.name,batches.steps.state,batches.steps.outcome,batches.steps.startTime,batches.steps.finishTime,batches.steps.warning,batches.steps.error"
     private readonly STREAM_CHANGES_FILTER = "id,dateUtc";
     static readonly MAX_JOB_COUNT_DEFAULT = 400;
 
@@ -183,7 +184,7 @@ export class StepOutcomeDataHandler extends PollBase {
         const streamList = this.activeFilter.streams ?? [];
         const modifiedAfter = this.activeFilter.jobHistorySpan.start.toISOString();
         const jobCount = this.activeFilter.maxJobCount ?? StepOutcomeDataHandler.MAX_JOB_COUNT_DEFAULT;
-
+        const jobStepKeyToState = new Map<string, JobStepState>();
         // Group jobs by streamId
         const jobsByStream: Record<string, string[]> = {};
 
@@ -273,11 +274,18 @@ export class StepOutcomeDataHandler extends PollBase {
                 }
 
                 const batchResponses = jobResponse.batches ?? [];
+                let streamStepSet = stepMap.get(jobResponse.streamId);
+                let templateStepSet = streamStepSet ? streamStepSet.get(jobResponse.templateId!) : undefined;
+
                 for (const batch of batchResponses) {
                     for (const step of batch.steps) {
+                        let unifiedKey: string = jobResponse.id + step.id;
+                        if (!jobStepKeyToState.has(unifiedKey)) {
+                            jobStepKeyToState.set(unifiedKey, step.state);
+                        }
+
                         let fullyQualifiedStepName = encodeStepNameFromStrings(jobResponse.streamId, jobResponse.templateId!, step.name).toLocaleLowerCase();
-                        let streamStepSet = stepMap.get(jobResponse.streamId);
-                        let templateStepSet = streamStepSet ? streamStepSet.get(jobResponse.templateId!) : undefined;
+
                         if (
                             templateStepSet &&
                             templateStepSet.size >= 0 &&
@@ -287,6 +295,22 @@ export class StepOutcomeDataHandler extends PollBase {
                         }
                     }
                 }
+
+                // Do not compute label steps if there are no steps for the template.
+                if (templateStepSet && templateStepSet.size >= 0) {
+                    const labelResponses = jobResponse.labels ?? [];
+                    for (let idx: number = 0; idx < labelResponses.length; ++idx) {
+                        let label = labelResponses[idx];
+                        let synthesizedStepOutcome: JobStepState | LabelState = LabelState.Unspecified;
+
+                        // Labels lack sufficient cancellation & skipped context, and inferring this from their .steps member is often incorrect (steps may not be present)
+                        // We therefore apply a simple heuristic: if there is cancellation info - we anotate the synthesized step as aborted; else skipped.
+                        if (label.state === LabelState.Unspecified) {
+                            synthesizedStepOutcome = jobResponse.abortedByUserInfo !== undefined ? JobStepState.Aborted : JobStepState.Skipped;
+                        }
+                        this.stepOutcomeTableData.addEntry(new LabelOutcomeTableEntry(label, idx, jobResponse.id, jobResponse.name, jobResponse.streamId, jobResponse.createTime.toString(), jobResponse.change, this.changeData.get(jobResponse.change!)?.dateUtc, synthesizedStepOutcome));
+                    }
+                }
             }
         }
 
@@ -312,11 +336,11 @@ export class StepOutcomeDataHandler extends PollBase {
      * @param stepOutcomeTableEntry The step outcome table entry to attempt to populate with issue data.
      */
     async populateStepOutcomeTableEntryIssueData(stepOutcomeTableEntry: StepOutcomeTableEntry, onCompleteCallback?: (boolean) => void): Promise<void> {
-        let unresolvedIssueQuery: IssueQuery = { jobId: stepOutcomeTableEntry.jobId, stepId: stepOutcomeTableEntry.stepResponse.id, resolved: false }
+        let unresolvedIssueQuery: IssueQuery = { jobId: stepOutcomeTableEntry.jobId, stepId: stepOutcomeTableEntry.id, resolved: false }
         const unresolvedIssueResults = await backend.getIssues(unresolvedIssueQuery);
         stepOutcomeTableEntry.issuesData = unresolvedIssueResults;
 
-        let resolvedIssueQuery: IssueQuery = { jobId: stepOutcomeTableEntry.jobId, stepId: stepOutcomeTableEntry.stepResponse.id, resolved: true }
+        let resolvedIssueQuery: IssueQuery = { jobId: stepOutcomeTableEntry.jobId, stepId: stepOutcomeTableEntry.id, resolved: true }
         const resolvedIssueResults = await backend.getIssues(resolvedIssueQuery);
         stepOutcomeTableEntry.issuesData.push(...resolvedIssueResults);
 
diff --git a/Engine/Source/Programs/Horde/HordeDashboard/src/components/buildhealth/stepoutcome/StepOutcomeDataTypes.ts b/Engine/Source/Programs/Horde/HordeDashboard/src/components/buildhealth/stepoutcome/StepOutcomeDataTypes.ts
--- a/Engine/Source/Programs/Horde/HordeDashboard/src/components/buildhealth/stepoutcome/StepOutcomeDataTypes.ts
+++ b/Engine/Source/Programs/Horde/HordeDashboard/src/components/buildhealth/stepoutcome/StepOutcomeDataTypes.ts
@@ -1,9 +1,10 @@
 // Copyright Epic Games, Inc. All Rights Reserved.
 
-import { ChangeSummaryData, GetJobResponse, GetStepResponse, IssueData, JobStepOutcome } from "horde/backend/Api";
+import { ChangeSummaryData, GetJobResponse, GetLabelStateResponse, GetStepResponse, IssueData, JobStepOutcome, JobStepState, LabelOutcome, LabelState } from "horde/backend/Api";
 import { computeDuration, DURATION_NOT_SET, formatDurationFromMs, toDate } from "./StepOutcomeUtilities";
 import { StepOutcomeDataHandler, StepOutcomeViewOptions } from "./StepOutcomeDataHandler";
 import { action, makeObservable, observable } from "mobx";
+import { NOT_APPLICABLE_LABEL } from "./StepOutcomeSharedUIComponents";
 
 export const WarningStr: string = "Warnings";
 export const SuccessStr: string = "Success";
@@ -13,6 +14,8 @@ export const WaitingStr: string = "Waiting";
 export const RunningStr: string = "Running";
 export const ReadyStr: string = "Ready";
 export const AbortedStr: string = "Aborted";
+export const CompletedStepStateStr: string = "Completed";
+export const CompleteLabelStateStr: string = "Complete";
 
 /**
  * Encodes a step outcome table entry to a stream qualified step name.
@@ -20,8 +23,8 @@ export const AbortedStr: string = "Aborted";
  * @param stepOutcomeTableEntry The stepOutcomeTableEntry to encode.
  * @returns The hierarchical stream qualified step name.
  */
-export function encodeStepName(stepOutcomeTableEntry: StepOutcomeTableEntry) {
-    return `${stepOutcomeTableEntry.streamId.toLocaleLowerCase()}${KEY_SEPARATOR}${stepOutcomeTableEntry.stepResponse.name.toLocaleLowerCase()}`;
+export function encodeStepName(stepOutcomeTableEntry: OutcomeTableEntry) {
+    return `${stepOutcomeTableEntry.streamId.toLocaleLowerCase()}${KEY_SEPARATOR}${stepOutcomeTableEntry.name.toLocaleLowerCase()}`;
 }
 
 /**
@@ -76,7 +79,7 @@ export type JobResponseData = GetJobResponse & {
 /**
  * Discriminated union which represents a tabel entry. This can be a @see StepOutcomeTableEntry (the core data in the table) or a @see SummaryTableEntry (the summarization of an entire row).
  */
-export type TableEntry = StepOutcomeTableEntry | SummaryTableEntry | undefined;
+export type TableEntry = OutcomeTableEntry | SummaryTableEntry | undefined;
 
 /**
  * Type predicate to test whether this is a @see StepOutcomeTableEntry .
@@ -87,6 +90,15 @@ export function isStepOutcome(entry: TableEntry): entry is StepOutcomeTableEntry
     return entry instanceof StepOutcomeTableEntry;
 }
 
+/**
+ * Type predicate to test whether this is a @see LabelOutcomeTableEntry .
+ * @param entry The entry to test.
+ * @returns True if it is a @see LabelOutcomeTableEntry - false otherwise.
+ */
+export function isLabelOutcome(entry: TableEntry): entry is LabelOutcomeTableEntry {
+    return entry instanceof LabelOutcomeTableEntry;
+}
+
 /**
  * Type predicate to test whether this is a @see SummaryTableEntry .
  * @param entry The entry to test.
@@ -109,36 +121,108 @@ export class SummaryTableEntry {
 }
 
 /**
- * Model for a table entry.
+ * Base model for table entries.
  */
-export class StepOutcomeTableEntry {
+export abstract class OutcomeTableEntry {
+
     // #region -- Public Members --
 
+    /**
+     * The change number for the table entry.
+     */
     change?: number;
+
+    /**
+     * The change date for the table entry.
+     */
     changeDate: Date | string;
-    stepResponse: StepReponseData;
+
+    /**
+     * The owning job's id.
+     */
     jobId: string;
+
+    /**
+     * The owning job's name.
+     */
     jobName: string;
+
+    /**
+     * The stream id that this table entry belongs to.
+     */
     streamId: string;
+
+    /**
+     * The owning job's creation time.
+     */
     jobCreateTime: string;
+
+    /**
+     * The issue data for the tabel entry.
+     */
     issuesData: IssueData[] = [];
-    duplicateEntries: StepOutcomeTableEntry[] = [];
+
+    /**
+     * Duplicate entries for the table entry.
+     */
+    duplicateEntries: OutcomeTableEntry[] = [];
 
     // #endregion -- Public Members --
 
+    // #region -- Abstract Members --
+
+    /**
+     * Gets the id of the table entry.
+     */
+    abstract get id(): string;
+
+    /**
+     * Gets the start time of the table entry.
+     */
+    abstract get startTime(): Date | string | undefined;
+
+    /**
+     * Gets the name of the table entry.
+     */
+    abstract get name(): string;
+
+    /**
+     * The state of the table entry.
+     */
+    abstract get state(): JobStepState | LabelState;
+
+    /**
+     * The outcome of the table entry.
+     */
+    abstract get outcome(): JobStepOutcome | LabelOutcome;
+
+    /**
+     * Gets the CSV representation of the the table entry. Property ordering matches @see getCSVHeader .
+     * @param relatedJobContext The related job's context object.
+     * @returns CSV representation of the table entry.
+     */
+    abstract toCSVRepresntation(relatedJobContext: JobResponseData | undefined): string;
+
+    /**
+     * Gets the CSV- property names of the step outcome entry. Property ordering matches @see toCSVRepresntation .
+     * @returns The comma separated properties serialization.
+     */
+    abstract getCSVHeader(): string;
+
+    // #endregion -- Abstract Members --
+
     // #region -- Constructor --
 
     /**
      * Constructor.
-     * @param stepResponse The step response object for the table entry.
      * @param jobId The owning job's Id.
      * @param jobName The owning job's name.
      * @param streamId The owning streams's Id.
      * @param jobCreateTime The owning job's create time.
-     * @param change The change associated with the stepResponse.
-     * @param changeDate The change submit date associated with the stepResponse.
+     * @param change The change associated with the table entry.
+     * @param changeDate The change submit date associated with the table entry.
      */
-    constructor(stepResponse: GetStepResponse, jobId: string, jobName: string, streamId: string, jobCreateTime: string, change?: number, changeDate?: Date | string) {
+    constructor(jobId: string, jobName: string, streamId: string, jobCreateTime: string, change?: number, changeDate?: Date | string) {
         if (change !== undefined) {
             this.change = change;
         }
@@ -150,7 +234,6 @@ export class StepOutcomeTableEntry {
             this.changeDate = jobCreateTime;
         }
 
-        this.stepResponse = { ...stepResponse, stepDurationMs: computeDuration(stepResponse.startTime, stepResponse.finishTime) };
         this.jobId = jobId;
         this.jobName = jobName;
         this.streamId = streamId;
@@ -158,27 +241,188 @@ export class StepOutcomeTableEntry {
     }
 
     // #endregion -- Constructor --
+}
 
-    // #region -- Public API --
+/**
+ * Model for a label table entry.
+ */
+export class LabelOutcomeTableEntry extends OutcomeTableEntry {
+    // #region -- Private Members --
+
+    private labelResponse: GetLabelStateResponse;
+    private labelIdx: number;
+    private synthesizedStepState?: LabelState | JobStepState;
+
+    // #endregion -- Private Members --
+
+    // #region -- Constructor --
 
     /**
-     * Gets the duration of the underyling @see StepReponseData in string representation.
-     * @returns The string representation of the step duration.
+     * Constructor.
+     * @param labelResponse The label response object for the table entry.
+     * @param labelIdx The label index for the job.
+     * @param jobId The owning job's Id.
+     * @param jobName The owning job's name.
+     * @param streamId The owning streams's Id.
+     * @param jobCreateTime The owning job's create time.
+     * @param change The change associated with the table entry.
+     * @param changeDate The change submit date associated with the table entry.
+     * @param synthesizedStepState The synthesized step state if applicable.
      */
-    getDurationString(): string {
-        let durationStr: string = this.stepResponse.state === "Skipped" ? "N/A" : "Ongoing";
+    constructor(labelResponse: GetLabelStateResponse, labelIdx: number, jobId: string, jobName: string, streamId: string, jobCreateTime: string, change?: number, changeDate?: Date | string, synthesizedStepState?: LabelState | JobStepState) {
+        super(jobId, jobName, streamId, jobCreateTime, change, changeDate);
 
-        if (this.stepResponse.stepDurationMs != DURATION_NOT_SET) {
-            durationStr = formatDurationFromMs(this.stepResponse.stepDurationMs);
+        this.labelIdx = labelIdx;
+        this.labelResponse = labelResponse;
+        this.synthesizedStepState = synthesizedStepState;
+    }
+
+    // #endregion -- Constructor --
+
+    // #region -- Abstract API --
+
+    /**
+     * @inheritdoc
+     */
+    get id(): string {
+        return this.labelIdx.toString();
+    }
+
+    /**
+     * @inheritdoc
+     */
+    get name(): string {
+        if (this.labelResponse?.dashboardName !== undefined) {
+            return `${this.labelResponse?.dashboardCategory}-${this.labelResponse?.dashboardName}`;
         }
+        return "";
+    }
 
-        return durationStr;
+    /**
+     * @inheritdoc
+     * Note: Labels have a smaller stable state domain than Steps. Labels don't have a concept of skip / cancelled (due to complexities & edge cases), so we attempt to use a synthesized step state in such a case.
+     * This is inferred from other metadata associated with the job, and is provided at construction of the lable outcome entry.
+     */
+    get state(): LabelState | JobStepState {
+        if (this.labelResponse.state === LabelState.Unspecified) {
+            return this.synthesizedStepState ?? LabelState.Unspecified;
+        }
+        if (this.labelResponse.state !== undefined && this.labelResponse.state !== null) {
+            return this.labelResponse.state;
+        }
+        return LabelState.Unspecified;
     }
 
     /**
-     * Gets the CSV representation of the the table entry. Property ordering matches @see getCSVHeader .
-     * @param relatedJobContext The related job's context object.
-     * @returns CSV representation of the table entry.
+     * @inheritdoc
+     */
+    get outcome(): LabelOutcome {
+        return this.labelResponse.outcome ?? LabelOutcome.Unspecified;
+    }
+
+    /**
+     * @inheritdoc
+     */
+    get startTime(): Date | string | undefined {
+        return undefined;
+    }
+
+    /**
+     * @inheritdoc
+     */
+    toCSVRepresntation(relatedJobContext: JobResponseData | undefined): string {
+        let baseUrl = window.location.origin;
+        let jobUrl = `${baseUrl}/job/${this.jobId}`;
+        let labelUrl = `${baseUrl}/job/${this.jobId}?label=${this.labelIdx}`;
+        return `${this.change},${this.streamId},${this.jobName},${this.jobId},${this.labelResponse.dashboardCategory},${this.labelResponse.dashboardName},${this.id},${this.state},${this.outcome},${this.jobCreateTime},${relatedJobContext?.summarizedOutcome},${relatedJobContext?.state},${jobUrl},${labelUrl}\n`;
+    }
+
+    /**
+     * @inheritdoc
+     */
+    getCSVHeader(): string {
+        return "Change,StreamId,JobName,JobId,LabelCategory,LabelName,LabelIdx,StepState,StepOutcome,JobStart,JobSummaryOutcome,JobState,JobUrl,LabelUrl\n";
+    }
+
+    // #endregion -- Abstract API --
+}
+
+/**
+ * Model for a table entry.
+ */
+export class StepOutcomeTableEntry extends OutcomeTableEntry {
+    // #region -- Private Members --
+
+    private stepResponse: StepReponseData;
+
+    // #endregion -- Private Members --
+
+    // #region -- Constructor --
+
+    /**
+     * Constructor.
+     * @param stepResponse The step response object for the table entry.
+     * @param jobId The owning job's Id.
+     * @param jobName The owning job's name.
+     * @param streamId The owning streams's Id.
+     * @param jobCreateTime The owning job's create time.
+     * @param change The change associated with the table entry.
+     * @param changeDate The change submit date associated with the table entry.
+     */
+    constructor(stepResponse: GetStepResponse, jobId: string, jobName: string, streamId: string, jobCreateTime: string, change?: number, changeDate?: Date | string) {
+        super(jobId, jobName, streamId, jobCreateTime, change, changeDate);
+
+        this.stepResponse = { ...stepResponse, stepDurationMs: computeDuration(stepResponse.startTime, stepResponse.finishTime) };
+    }
+
+    // #endregion -- Constructor --
+
+    // #region -- Abstract API --
+
+    /**
+     * @inheritdoc
+     */
+    get id(): string {
+        return this.stepResponse.id;
+    }
+
+    /**
+     * @inheritdoc
+     */
+    get name(): string {
+        return this.stepResponse?.name ?? "";
+    }
+
+    /**
+     * @inheritdoc
+     */
+    get starTime() {
+        return this.stepResponse?.startTime;
+    }
+
+    /**
+     * @inheritdoc
+     */
+    get state(): JobStepState {
+        return this.stepResponse.state;
+    }
+
+    /**
+     * @inheritdoc
+     */
+    get outcome(): JobStepOutcome {
+        return this.stepResponse.outcome;
+    }
+
+    /**
+     * @inheritdoc
+     */
+    get startTime() {
+        return this.stepResponse.startTime;
+    }
+
+    /**
+     * @inheritdoc
      */
     toCSVRepresntation(relatedJobContext: JobResponseData | undefined): string {
         let baseUrl = window.location.origin;
@@ -188,12 +432,30 @@ export class StepOutcomeTableEntry {
     }
 
     /**
-     * Gets the CSV property names of the step outcome entry. Property ordering matches @see toCSVRepresntation .
-     * @returns The comma separated properties serialization.
+     * @inheritdoc
      */
-    static getCSVHeader() {
+    getCSVHeader(): string {
         return "Change,StreamId,JobName,JobId,StepName,StepId,StepState,StepOutcome,Duration(Sec),JobStart,StepStart,JobSummaryOutcome,JobState,JobUrl,StepUrl\n";
     }
+
+    // #endregion -- Abstract API --
+
+    // #region -- Public API --
+
+    /**
+     * Gets the duration of the underyling @see StepReponseData in string representation.
+     * @returns The string representation of the step duration.
+     */
+    getDurationString(): string {
+        let durationStr: string = this.stepResponse.state === SkippedStr ? NOT_APPLICABLE_LABEL : (this.stepResponse.state === WaitingStr || this.stepResponse.state === ReadyStr ? "Not Started" : "Ongoing");
+
+        if (this.stepResponse.stepDurationMs != DURATION_NOT_SET) {
+            durationStr = formatDurationFromMs(this.stepResponse.stepDurationMs);
+        }
+
+        return durationStr;
+    }
+
     // #endregion -- Public API --
 }
 
@@ -264,6 +526,7 @@ export type DateHeader = {
 export type StepNameHeader = {
     stepName: string;
     streamId: string;
+    isLabel: boolean;
 }
 
 // #region -- Table Header Types --
@@ -281,6 +544,11 @@ export class StepOutcomeTable {
      */
     supportsDateAnchoredChanges: boolean = true;
 
+    /**
+     * Sets whether @see getStepRowHeaders returns label table entries for the given data.
+     */
+    supportsLabelEntries: boolean = false;
+
     /**
      * Change column headers.
      */
@@ -293,7 +561,7 @@ export class StepOutcomeTable {
      */
     stepRowLookup: Map<string, number> = new Map<string, number>();
     stepRowNumberLookup: Map<number, string> = new Map<number, string>();
-    stepNameRowHeaders: StepNameHeader[] = [];
+    private stepNameRowHeaders: StepNameHeader[] = [];
 
     /**
      * SummaryLookup maps a normalized step name to a summary entry.
@@ -308,7 +576,7 @@ export class StepOutcomeTable {
     /**
      * Table data structure
      */
-    tableEntries: (StepOutcomeTableEntry | null)[][] = [];
+    tableEntries: (OutcomeTableEntry | null)[][] = [];
 
     // #endregion --  Public Members --
 
@@ -318,7 +586,7 @@ export class StepOutcomeTable {
 
     // #region -- Private API --
 
-    private ensureRow(source: (StepOutcomeTableEntry | null)[], targetSize: number) {
+    private ensureRow(source: (OutcomeTableEntry | null)[], targetSize: number) {
         while (targetSize > source.length) {
             source.push(null);
         }
@@ -394,7 +662,7 @@ export class StepOutcomeTable {
      * @param tableEntry The table entry to add to the data table.
      * @returns If tableEntry parameter is invalid, returns.
      */
-    addEntry(tableEntry: StepOutcomeTableEntry): void {
+    addEntry(tableEntry: OutcomeTableEntry): void {
         if (tableEntry.change === undefined) {
             return;
         }
@@ -408,7 +676,7 @@ export class StepOutcomeTable {
             insertionRow = this.stepNameRowHeaders.length;
             this.stepRowLookup.set(normalizedStepName, insertionRow); // cache the normalized step name -> row number
             this.stepRowNumberLookup.set(insertionRow, normalizedStepName); // cache the row number -> normalized step name
-            this.stepNameRowHeaders.push({ streamId: tableEntry.streamId, stepName: tableEntry.stepResponse.name });
+            this.stepNameRowHeaders.push({ streamId: tableEntry.streamId, stepName: tableEntry.name, isLabel: isLabelOutcome(tableEntry) });
         }
 
         let insertionCol = this.changeColLookup.get(tableEntry.change);
@@ -439,8 +707,8 @@ export class StepOutcomeTable {
             // Attempt to sort by step sort time first. This fails under two scenarios:
             // 1. The step has Not Started (so the startTime will be null and thus, a delay in this taking top position)
             // 2. The step has been skipped (and we must fall back to job creation time as a sorting mechanism - startTime will never be !null).
-            let startA: number = tableEntry.stepResponse.startTime ? new Date(tableEntry.stepResponse.startTime).getTime() : 0;
-            let startB: number = existingEntry.stepResponse.startTime ? new Date(existingEntry.stepResponse.startTime).getTime() : 0;
+            let startA: number = tableEntry.startTime ? new Date(tableEntry.startTime).getTime() : 0;
+            let startB: number = existingEntry.startTime ? new Date(existingEntry.startTime).getTime() : 0;
 
             // If either is invalid (most likely due to the Not Started case, but secondarily if it's skipped/aborted), fall back to considering the job created time.
             // Note: this will not be captured on the edge case of step-retry only. The best recommendation for this edge case is for users is to review the job details which they will have access
@@ -484,25 +752,27 @@ export class StepOutcomeTable {
 
             for (let colIndex = 0; colIndex < row.length; colIndex++) {
                 const cell = row[colIndex];
-                if (cell == null) {
+
+                if (cell === null) {
                     continue;
                 }
+
                 summaryTableEntry.steps++;
 
-                if (cell.stepResponse.state === WaitingStr || cell.stepResponse.state === RunningStr) {
+                if (cell.state === WaitingStr || cell.state === RunningStr) {
                     summaryTableEntry.stepsPending++;
                     continue;
                 }
 
                 summaryTableEntry.stepsCompleted++;
 
-                if (cell.stepResponse.outcome == WarningStr) {
+                if (cell.outcome == WarningStr) {
                     summaryTableEntry.stepsWarning++;
                 }
-                else if (cell.stepResponse.outcome == FailureStr) {
+                else if (cell.outcome == FailureStr || cell.state == AbortedStr) {
                     summaryTableEntry.stepsFail++;
                 }
-                else if (cell.stepResponse.outcome == SuccessStr) {
+                else if (cell.outcome == SuccessStr) {
                     summaryTableEntry.stepsPass++;
                 }
             }
@@ -541,10 +811,10 @@ export class StepOutcomeTable {
      * @param rowName The step name to obtain.
      * @returns * Generator of TableEntries.
      */
-    getStepOutputRow(rowName: string): Generator<StepOutcomeTableEntry>;
-    getStepOutputRow(rowIndex: number): Generator<StepOutcomeTableEntry>;
-    getStepOutputRow(row: string | number): Generator<StepOutcomeTableEntry>;
-    *getStepOutputRow(row: string | number): Generator<StepOutcomeTableEntry | null | undefined> {
+    getStepOutputRow(rowName: string): Generator<OutcomeTableEntry>;
+    getStepOutputRow(rowIndex: number): Generator<OutcomeTableEntry>;
+    getStepOutputRow(row: string | number): Generator<OutcomeTableEntry>;
+    *getStepOutputRow(row: string | number): Generator<OutcomeTableEntry | null | undefined> {
         let rowIdx: number | undefined;
 
         if (typeof row === 'string') {
@@ -576,20 +846,20 @@ export class StepOutcomeTable {
      * @param changeNumber The change number.
      * @returns * Generator of all Table Entries for the specific change.
      */
-    *getChangeColumn(changeNumber: number): Generator<StepOutcomeTableEntry> {
+    *getChangeColumn(changeNumber: number): Generator<OutcomeTableEntry> {
 
         let colIdx: (number | undefined) = this.changeColLookup.get(changeNumber);
         if (colIdx === undefined) {
             return;
         }
 
-        for (const entry of this.stepNameRowHeaders) {
+        for (const entry of this.getStepRowHeaders()) {
             let rowIdx: (number | undefined) = this.stepRowLookup.get(encodeStepNameFromStepNameHeader(entry));
             if (rowIdx === undefined) {
                 continue;
             }
 
-            let tableEntry: (StepOutcomeTableEntry | null) = this.tableEntries[rowIdx][colIdx];
+            let tableEntry: (OutcomeTableEntry | null) = this.tableEntries[rowIdx][colIdx];
             if (tableEntry !== null) {
                 yield tableEntry;
             }
@@ -629,26 +899,54 @@ export class StepOutcomeTable {
         }
     }
 
+    /**
+     * Gets all of the step row headers for the table.
+     * @returns * Generator of all the step row headers. If @see this.supportsLabelEntries , these items will also be returned.
+     */
+    *getStepRowHeaders(): Generator<StepNameHeader> {
+        for (let idx: number = 0; idx < this.stepNameRowHeaders.length; ++idx) {
+            if (!this.supportsLabelEntries && this.stepNameRowHeaders[idx].isLabel) {
+                continue;
+            }
+            yield this.stepNameRowHeaders[idx];
+        }
+    }
+
     /**
      * Generates the csv string representation of the step outcome table, in row order.
      * @param contextHandler The datahandler used to obtain job response data.
      * @returns A string in CSV format that represents the underlying step outcome table, row order.
      */
-    toRowOrderedCSV(contextHandler: StepOutcomeDataHandler): string {
-        let resultString: string = StepOutcomeTableEntry.getCSVHeader();
+    toRowOrderedCSV<U extends OutcomeTableEntry>(contextHandler: StepOutcomeDataHandler, guard: (entry: OutcomeTableEntry) => entry is U): string {
+        let resultString: string | undefined;
         for (let rowIndex = 0; rowIndex < this.tableEntries.length; rowIndex++) {
             const row = this.tableEntries[rowIndex];
             for (let i: number = 0; i < row.length; ++i) {
                 if (!row[i]) {
                     continue;
                 }
-                let tableEntry: StepOutcomeTableEntry = row[i]!;
+                let tableEntry: OutcomeTableEntry = row[i]!;
+
+                if (!guard(tableEntry)) {
+                    continue;
+                }
+
+                let castedEntry: U = tableEntry as U;
+
+                if (castedEntry === undefined) {
+                    continue;
+                }
+
+                if (resultString == undefined) {
+                    resultString = castedEntry.getCSVHeader();
+                }
+
                 let jobResponseData: JobResponseData | undefined = contextHandler.getJobResponseData(tableEntry.jobId);
-                resultString += tableEntry.toCSVRepresntation(jobResponseData);
+                resultString += castedEntry.toCSVRepresntation(jobResponseData);
             }
         }
 
-        return resultString;
+        return resultString!;
     }
 
     /**
@@ -656,21 +954,38 @@ export class StepOutcomeTable {
      * @param contextHandler The datahandler used to obtain job response data.
      * @returns A string in CSV format that represents the underlying step outcome table, column order.
      */
-    toColOrderedCSV(contextHandler: StepOutcomeDataHandler): string {
-        let resultString: string = StepOutcomeTableEntry.getCSVHeader();
+    toColOrderedCSV<U extends OutcomeTableEntry>(contextHandler: StepOutcomeDataHandler, guard: (entry: OutcomeTableEntry) => entry is U): string {
+        let resultString: string | undefined;
         for (let changeIdx = 0; changeIdx < this.changeColHeaders.length; changeIdx++) {
             const changeEntries = [...this.getChangeColumn(this.changeColHeaders[changeIdx].change)];
             for (let i: number = 0; i < changeEntries.length; ++i) {
+
                 if (!changeEntries[i]) {
                     continue;
                 }
-                let tableEntry: StepOutcomeTableEntry = changeEntries[i]!;
+
+                let tableEntry: OutcomeTableEntry = changeEntries[i]!;
+
+                if (!guard(tableEntry)) {
+                    continue;
+                }
+
+                let castedEntry: OutcomeTableEntry = tableEntry as OutcomeTableEntry;
+
+                if (castedEntry === undefined) {
+                    continue;
+                }
+
+                if (resultString == undefined) {
+                    resultString = castedEntry.getCSVHeader();
+                }
+
                 let jobResponseData: JobResponseData | undefined = contextHandler.getJobResponseData(tableEntry.jobId);
-                resultString += tableEntry.toCSVRepresntation(jobResponseData);
+                resultString += castedEntry.toCSVRepresntation(jobResponseData);
             }
         }
 
-        return resultString;
+        return resultString!;
     }
 
     /**
@@ -680,6 +995,7 @@ export class StepOutcomeTable {
     @action
     setViewModelOptions(viewModelOptions: StepOutcomeViewOptions) {
         this.viewModelOptions = viewModelOptions;
+        this.supportsLabelEntries = viewModelOptions.showLabels;
         this.supportsDateAnchoredChanges = viewModelOptions.includeDateAnchors;
     }
 
diff --git a/Engine/Source/Programs/Horde/HordeDashboard/src/components/buildhealth/stepoutcome/StepOutcomeLegend.tsx b/Engine/Source/Programs/Horde/HordeDashboard/src/components/buildhealth/stepoutcome/StepOutcomeLegend.tsx
--- a/Engine/Source/Programs/Horde/HordeDashboard/src/components/buildhealth/stepoutcome/StepOutcomeLegend.tsx
+++ b/Engine/Source/Programs/Horde/HordeDashboard/src/components/buildhealth/stepoutcome/StepOutcomeLegend.tsx
@@ -7,72 +7,62 @@ import { getCellStyle, getCellAccents, getCellClasses } from "./StepOutcomeShare
 
 export const StepOutcomeLegend: React.FC = () => {
     const fakeSuccessEntry = {
-        stepResponse: {
-            state: JobStepState.Completed,
-            outcome: JobStepOutcome.Success
-        }
-    } as Pick<StepOutcomeTableEntry, "stepResponse"> as StepOutcomeTableEntry;
+        state: JobStepState.Completed,
+        outcome: JobStepOutcome.Success
+    } as unknown as StepOutcomeTableEntry;
+
+    const fakeSuccessMultiRunEntry = {
+        state: JobStepState.Completed,
+        outcome: JobStepOutcome.Success,
+        duplicateEntries: ["One", "Two", "Three"]
+    } as unknown as StepOutcomeTableEntry;
 
     const fakeFailureEntry = {
-        stepResponse: {
-            state: JobStepState.Completed,
-            outcome: JobStepOutcome.Failure
-        }
-    } as Pick<StepOutcomeTableEntry, "stepResponse"> as StepOutcomeTableEntry;
+        state: JobStepState.Completed,
+        outcome: JobStepOutcome.Failure
+    } as unknown as StepOutcomeTableEntry;
 
     const fakeWarningEntry = {
-        stepResponse: {
-            state: JobStepState.Completed,
-            outcome: JobStepOutcome.Warnings
-        }
-    } as Pick<StepOutcomeTableEntry, "stepResponse"> as StepOutcomeTableEntry;
+        state: JobStepState.Completed,
+        outcome: JobStepOutcome.Warnings
+    } as unknown as StepOutcomeTableEntry;
 
     const fakeRunningEntry = {
-        stepResponse: {
-            state: JobStepState.Running,
-            outcome: JobStepOutcome.Unspecified
-        }
-    } as Pick<StepOutcomeTableEntry, "stepResponse"> as StepOutcomeTableEntry;
+        state: JobStepState.Running,
+        outcome: JobStepOutcome.Unspecified
+    } as unknown as StepOutcomeTableEntry;
 
     const fakeRunningSummaryData = {
         summarizedOutcome: JobStepOutcome.Unspecified
     } as Partial<JobResponseData> as JobResponseData | undefined;
 
     const fakeSkippedEntry = {
-        stepResponse: {
-            state: JobStepState.Skipped,
-            outcome: JobStepOutcome.Failure
-        }
-    } as Pick<StepOutcomeTableEntry, "stepResponse"> as StepOutcomeTableEntry;
+        state: JobStepState.Skipped,
+        outcome: JobStepOutcome.Failure
+    } as unknown as StepOutcomeTableEntry;
 
     const fakeSkippedSummaryData = {
         summarizedOutcome: JobStepOutcome.Failure
     } as Partial<JobResponseData> as JobResponseData | undefined;
 
     const fakeWaitingEntry = {
-        stepResponse: {
-            state: JobStepState.Waiting,
-            outcome: JobStepOutcome.Success
-        }
-    } as Pick<StepOutcomeTableEntry, "stepResponse"> as StepOutcomeTableEntry;
+        state: JobStepState.Waiting,
+        outcome: JobStepOutcome.Success
+    } as unknown as StepOutcomeTableEntry;
 
     const fakeReadyEntry = {
-        stepResponse: {
-            state: JobStepState.Ready,
-            outcome: JobStepOutcome.Success
-        }
-    } as Pick<StepOutcomeTableEntry, "stepResponse"> as StepOutcomeTableEntry;
+        state: JobStepState.Ready,
+        outcome: JobStepOutcome.Success
+    } as unknown as StepOutcomeTableEntry;
 
     const fakeReadyWatingummaryData = {
         summarizedOutcome: JobStepOutcome.Unspecified
     } as Partial<JobResponseData> as JobResponseData | undefined;
 
     const fakeAbortedEntry = {
-        stepResponse: {
-            state: JobStepState.Aborted,
-            outcome: JobStepOutcome.Failure
-        }
-    } as Pick<StepOutcomeTableEntry, "stepResponse"> as StepOutcomeTableEntry;
+        state: JobStepState.Aborted,
+        outcome: JobStepOutcome.Failure
+    } as unknown as StepOutcomeTableEntry;
 
     const fakeAbortedFailureSummaryData = {
         summarizedOutcome: JobStepOutcome.Failure
@@ -95,7 +85,7 @@ export const StepOutcomeLegend: React.FC = () => {
                         <tr>
                             <td>
                                 <TooltipHost content="When a step has run to completion." calloutProps={{ directionalHint: DirectionalHint.topCenter }}>
-                                    <div style={{ width: "100%", height: "100%" }}>Complete Exection (solid fill)</div>
+                                    <div style={{ width: "100%", height: "100%" }}>Complete</div>
                                 </TooltipHost>
                             </td>
                             <td style={{
@@ -117,7 +107,7 @@ export const StepOutcomeLegend: React.FC = () => {
                         <tr>
                             <td>
                                 <TooltipHost content="When a step cannot execute (dependency failed), is currently executing, or hasn't yet executed." calloutProps={{ directionalHint: DirectionalHint.topCenter }}>
-                                    <div style={{ width: "100%", height: "100%" }}>Incomplete Execution (outline)</div>
+                                    <div style={{ width: "100%", height: "100%" }}>Incomplete</div>
                                 </TooltipHost>
                             </td>
                             <td style={{
@@ -141,8 +131,8 @@ export const StepOutcomeLegend: React.FC = () => {
                         </tr>
                         <tr>
                             <td>
-                                <TooltipHost content="When a step as been aborted." calloutProps={{ directionalHint: DirectionalHint.topCenter }}>
-                                    <div style={{ width: "100%", height: "100%" }}>Cancelled Execution (strike-through)</div>
+                                <TooltipHost content="When a step as been cancelled." calloutProps={{ directionalHint: DirectionalHint.topCenter }}>
+                                    <div style={{ width: "100%", height: "100%" }}>Cancelled</div>
                                 </TooltipHost>
                             </td>
                             <td className={getCellClasses(fakeAbortedEntry, fakeAbortedFailureSummaryData).join(" ")}
@@ -151,21 +141,34 @@ export const StepOutcomeLegend: React.FC = () => {
                                     height: 20,
                                     backgroundColor: getCellStyle(fakeAbortedEntry).bgColor,
                                     boxShadow: getCellAccents(fakeAbortedEntry, fakeAbortedFailureSummaryData).boxShadow,
-                                }}><TooltipHost content="The step has been aborted, and other steps in the job have failed." ><div style={{ width: "100%", height: "100%" }} /></TooltipHost></td>
+                                }}><TooltipHost content="The step has been cancelled, and other steps in the job have failed." ><div style={{ width: "100%", height: "100%" }} /></TooltipHost></td>
                             <td className={getCellClasses(fakeAbortedEntry, fakeAbortedWarningSummaryData).join(" ")}
                                 style={{
                                     width: 20,
                                     height: 20,
                                     backgroundColor: getCellStyle(fakeAbortedEntry).bgColor,
                                     boxShadow: getCellAccents(fakeAbortedEntry, fakeAbortedWarningSummaryData).boxShadow,
-                                }}><TooltipHost content="The step has been aborted, and other steps in the job encountered warnings." ><div style={{ width: "100%", height: "100%" }} /></TooltipHost></td>
+                                }}><TooltipHost content="The step has been cancelled, and other steps in the job encountered warnings." ><div style={{ width: "100%", height: "100%" }} /></TooltipHost></td>
                             <td className={getCellClasses(fakeAbortedEntry, fakeAbortedSuccessSummaryData).join(" ")}
                                 style={{
                                     width: 20,
                                     height: 20,
                                     backgroundColor: getCellStyle(fakeAbortedEntry).bgColor,
                                     boxShadow: getCellAccents(fakeAbortedEntry, fakeAbortedSuccessSummaryData).boxShadow,
-                                }}><TooltipHost content="The step has been aborted, but all other steps in the job completed with success." ><div style={{ width: "100%", height: "100%" }} /></TooltipHost></td>
+                                }}><TooltipHost content="The step has been cancelled, but all other steps in the job completed with success." ><div style={{ width: "100%", height: "100%" }} /></TooltipHost></td>
+                        </tr>
+                        <tr>
+                            <td>
+                                <TooltipHost content="When a step & changelist have multiple runs." calloutProps={{ directionalHint: DirectionalHint.topCenter }}>
+                                    <div style={{ width: "100%", height: "100%" }}>Multi-value</div>
+                                </TooltipHost>
+                            </td>
+                            <td className={getCellClasses(fakeSuccessMultiRunEntry, fakeAbortedSuccessSummaryData).join(" ")}
+                            style={{
+                                width: 20,
+                                height: 20,
+                                backgroundColor: getCellStyle(fakeSuccessMultiRunEntry).bgColor,
+                            }}><TooltipHost content="The step & change has multiple entires." ><div style={{ width: "100%", height: "100%" }} /></TooltipHost></td>
                         </tr>
                     </tbody>
                 </table>
diff --git a/Engine/Source/Programs/Horde/HordeDashboard/src/components/buildhealth/stepoutcome/StepOutcomeSharedUIComponents.tsx b/Engine/Source/Programs/Horde/HordeDashboard/src/components/buildhealth/stepoutcome/StepOutcomeSharedUIComponents.tsx
--- a/Engine/Source/Programs/Horde/HordeDashboard/src/components/buildhealth/stepoutcome/StepOutcomeSharedUIComponents.tsx
+++ b/Engine/Source/Programs/Horde/HordeDashboard/src/components/buildhealth/stepoutcome/StepOutcomeSharedUIComponents.tsx
@@ -4,11 +4,12 @@ import { mergeStyleSets, Spinner, SpinnerSize } from "@fluentui/react";
 import dashboard, { StatusColor } from "horde/backend/Dashboard";
 import { observer } from "mobx-react-lite";
 import { StepOutcomeDataHandler } from "./StepOutcomeDataHandler";
-import { AbortedStr, FailureStr, JobResponseData, SkippedStr, StepOutcomeTableEntry, SuccessStr, WarningStr } from "./StepOutcomeDataTypes";
+import { AbortedStr, FailureStr, OutcomeTableEntry, JobResponseData, SkippedStr, SuccessStr, WarningStr, CompletedStepStateStr, CompleteLabelStateStr } from "./StepOutcomeDataTypes";
 import { isTerminalState } from "./StepOutcomeUtilities";
 
 // #region -- Constants --
 
+export const NOT_APPLICABLE_LABEL: string = "N/A";
 export const ROW_HEADER_NAME: string = "stepName" as const;
 
 // #endregion -- Constants --
@@ -51,13 +52,13 @@ export const stepOutcomeTableClasses = mergeStyleSets({
         position: "relative",
         selectors: {
             "&::after": {
-                content: "'*'",
+                content: "'.'",
                 position: "absolute",
-                top: 4,
-                right: -4,
-                color: dashboard.darktheme ? "#ffffff" : "#000000",
+                top: -9,
+                right: -2,
+                color: colors["Ready"],
                 fontWeight: "bold",
-                fontSize: 25,
+                fontSize: 40,
                 lineHeight: 1,
                 pointerEvents: "none",
                 userSelect: "none",
@@ -66,7 +67,7 @@ export const stepOutcomeTableClasses = mergeStyleSets({
     },
     ...Object.fromEntries(
         Object.entries(colors).map(([className, color]) => [
-            className, // whatever key your getColorRecords() provides
+            className,
             {
                 ...baseX,
                 selectors: {
@@ -118,14 +119,14 @@ export function resetCellClasses(): string[] {
  * @param jobState The job response data which will be used to style applicable table entries.
  * @returns The classes to apply to the cell.
  */
-export function getCellClasses(entry: StepOutcomeTableEntry | null, jobState: JobResponseData | undefined): string[] {
+export function getCellClasses(entry: OutcomeTableEntry | null, jobState: JobResponseData | undefined): string[] {
     const classes: string[] = [];
 
     if (!entry || jobState?.summarizedOutcome === undefined) {
         return classes;
     }
 
-    if (entry.stepResponse.state === AbortedStr) {
+    if (entry.state === AbortedStr) {
         classes.push(stepOutcomeTableClasses[jobState?.summarizedOutcome]);
     }
 
@@ -142,7 +143,7 @@ export function getCellClasses(entry: StepOutcomeTableEntry | null, jobState: Jo
  * @param jobState The job response data which will be used to style applicable table entries.
  * @returns The accents to apply to the cell.
  */
-export function getCellAccents(entry: StepOutcomeTableEntry | null, jobState: JobResponseData | undefined): { boxShadow: string } {
+export function getCellAccents(entry: OutcomeTableEntry | null, jobState: JobResponseData | undefined): { boxShadow: string } {
     const standard = { boxShadow: "none" };
     if (!entry || jobState?.summarizedOutcome === undefined) {
         return standard;
@@ -150,11 +151,11 @@ export function getCellAccents(entry: StepOutcomeTableEntry | null, jobState: Jo
     const colors: Record<string, string> = getColorRecords();
 
     // If we have skipped or aborted the step, we use the summarized outcome to help with coloring.
-    if (entry.stepResponse.state == SkippedStr || entry.stepResponse.state == AbortedStr) {
+    if (entry.state == SkippedStr || entry.state == AbortedStr) {
         return { boxShadow: `inset 0 0 0 4px ${colors[jobState?.summarizedOutcome]}` };
     }
     else if (!isTerminalState(entry)) {
-        return { boxShadow: `inset 0 0 0 4px ${colors[entry.stepResponse.state]}` };
+        return { boxShadow: `inset 0 0 0 4px ${colors[entry.state]}` };
     }
 
     return standard;
@@ -166,7 +167,7 @@ export function getCellAccents(entry: StepOutcomeTableEntry | null, jobState: Jo
  * @param viewModelOptions The view model options to consider while applying styles.
  * @returns The matching style given the state of the entry. 
  */
-export function getCellStyle(entry: StepOutcomeTableEntry | null): { bgColor: string } {
+export function getCellStyle(entry: OutcomeTableEntry | null): { bgColor: string } {
     const colors: Record<string, string> = getColorRecords();
 
     let resultColor: string = 'transparent';
@@ -174,8 +175,7 @@ export function getCellStyle(entry: StepOutcomeTableEntry | null): { bgColor: st
         return { bgColor: resultColor };
     }
 
-    const { stepResponse } = entry;
-    const { outcome, state } = stepResponse;
+    const { outcome, state } = entry;
 
     // We do not fill the cells for Aborted, Skipped, or those about to run (or currently running).
     if (state === AbortedStr || state === SkippedStr || !isTerminalState(entry)) {
@@ -194,7 +194,7 @@ export function getCellStyle(entry: StepOutcomeTableEntry | null): { bgColor: st
         resultColor = colors[WarningStr];
     }
     // A step can have a completed outcome, but merely be waiting to run, running, or ready.
-    else if (outcome === SuccessStr && state === "Completed") {
+    else if (outcome === SuccessStr && (state === CompletedStepStateStr || state === CompleteLabelStateStr)) {
         resultColor = colors[SuccessStr];
     }
     // The edge cases have been satisfied, yield to the state to define result view color.
diff --git a/Engine/Source/Programs/Horde/HordeDashboard/src/components/buildhealth/stepoutcome/StepOutcomeUtilities.ts b/Engine/Source/Programs/Horde/HordeDashboard/src/components/buildhealth/stepoutcome/StepOutcomeUtilities.ts
--- a/Engine/Source/Programs/Horde/HordeDashboard/src/components/buildhealth/stepoutcome/StepOutcomeUtilities.ts
+++ b/Engine/Source/Programs/Horde/HordeDashboard/src/components/buildhealth/stepoutcome/StepOutcomeUtilities.ts
@@ -1,7 +1,7 @@
 // Copyright Epic Games, Inc. All Rights Reserved.
 
 import { GetJobResponse, JobStepOutcome, JobStepState } from "horde/backend/Api";
-import { ReadyStr, RunningStr, StepOutcomeTableEntry, WaitingStr } from "./StepOutcomeDataTypes";
+import { OutcomeTableEntry, ReadyStr, RunningStr, WaitingStr } from "./StepOutcomeDataTypes";
 
 export const DURATION_NOT_SET: number = -1;
 
@@ -10,8 +10,8 @@ export const DURATION_NOT_SET: number = -1;
  * @param stepOutcomeEntry The entry to consider.
  * @returns True if the step has run to completion, false otherwise.
  */
-export function isTerminalState(stepOutcomeEntry: StepOutcomeTableEntry): boolean {
-    if (stepOutcomeEntry.stepResponse.state === ReadyStr || stepOutcomeEntry.stepResponse.state === RunningStr || stepOutcomeEntry.stepResponse.state === WaitingStr) {
+export function isTerminalState(stepOutcomeEntry: OutcomeTableEntry): boolean {
+    if (stepOutcomeEntry.state === ReadyStr || stepOutcomeEntry.state === RunningStr || stepOutcomeEntry.state === WaitingStr) {
         return false;
     }
 
