[Horde][BuildHealth][Server] Added suspect finding info to issue audit log
#DisabledValidations CheckCopyright, CheckUnacceptableWords

diff --git a/Engine/Source/Programs/Horde/Plugins/Build/HordeServer.Build/Issues/IssueService.cs b/Engine/Source/Programs/Horde/Plugins/Build/HordeServer.Build/Issues/IssueService.cs
--- a/Engine/Source/Programs/Horde/Plugins/Build/HordeServer.Build/Issues/IssueService.cs
+++ b/Engine/Source/Programs/Horde/Plugins/Build/HordeServer.Build/Issues/IssueService.cs
@@ -975,25 +975,8 @@ namespace HordeServer.Issues
 			{
 				foreach (IssueAuditLogEntry e in uniqueEntries)
 				{
-					object[]? args = e.GetArgs() ?? [];
-
 					#pragma warning disable CA2254 // Template should be a static 
-					if (e.GetLogLevel() == LogLevel.Error)
-					{
-						auditLogger.LogError(e.GetMessageTemplate(), args);
-					}
-					else if (e.GetLogLevel() == LogLevel.Warning)
-					{
-						auditLogger.LogWarning(e.GetMessageTemplate(), args);
-					}
-					else if (e.GetLogLevel() == LogLevel.Information)
-					{
-						auditLogger.LogInformation(e.GetMessageTemplate(), args);
-					}
-					else
-					{
-						auditLogger.LogDebug(e.GetMessageTemplate(), args);
-					}
+					auditLogger.Log(e.GetLogLevel(), e.GetMessageTemplate(), e.GetArgs() ?? []);
 					#pragma warning restore CA2254 // Template should be a static 
 				}
 			}
@@ -1049,8 +1032,11 @@ namespace HordeServer.Issues
 				IJobStepRef? prevJob = await _jobStepRefs.GetPrevStepForNodeAsync(job.StreamId, job.TemplateId, node.Name, job.CommitId, JobStepOutcome.Success, true);
 				if (prevJob != null)
 				{
+					// Find first audit logger in event group. Doesn't matter which as they're all emitted to same issue audit log.
+					ILogger? suspectAuditLogger = eventGroup.Events.FirstOrDefault(e => e.AuditLogger != null)?.AuditLogger;
+
 					spanData.LastSuccess = new NewIssueStepData(prevJob);
-					spanData.Suspects = await FindSuspectsForSpanAsync(streamConfig, spanData.Fingerprint, spanData.LastSuccess.CommitId, spanData.FirstFailure.CommitId, cancellationToken);
+					spanData.Suspects = await FindSuspectsForSpanAsync(streamConfig, spanData.Fingerprint, spanData.LastSuccess.CommitId, spanData.FirstFailure.CommitId, suspectAuditLogger, cancellationToken);
 				}
 
 				IJobStepRef? nextJob = await _jobStepRefs.GetNextStepForNodeAsync(job.StreamId, job.TemplateId, node.Name, job.CommitId, JobStepOutcome.Success, true);
@@ -1386,9 +1372,10 @@ namespace HordeServer.Issues
 		/// <param name="fingerprint">The fingerprint for the span</param>
 		/// <param name="lastSuccess">The last successful commit</param>
 		/// <param name="firstFailure">The first failing commit</param>
+		/// <param name="issueLogger">The issue's audit logger</param>
 		/// <param name="cancellationToken">Cancellation token for the operation</param>
 		/// <returns>List of suspects</returns>
-		async Task<List<NewIssueSpanSuspectData>> FindSuspectsForSpanAsync(StreamConfig streamConfig, IIssueFingerprint fingerprint, CommitId lastSuccess, CommitId firstFailure, CancellationToken cancellationToken)
+		async Task<List<NewIssueSpanSuspectData>> FindSuspectsForSpanAsync(StreamConfig streamConfig, IIssueFingerprint fingerprint, CommitId lastSuccess, CommitId firstFailure, ILogger? issueLogger, CancellationToken cancellationToken)
 		{
 			using TelemetrySpan span = _tracer.StartActiveSpan($"{nameof(IssueService)}.{nameof(FindSuspectsForSpanAsync)}");
 			
@@ -1403,7 +1390,7 @@ namespace HordeServer.Issues
 
 					if (firstFailureCommit == null || lastSuccessCommit == null)
 					{
-						_logger.LogError("Failed to get commits for {StreamName}: {LastSuccess} and {FirstFailure}", streamConfig.Name, lastSuccess, firstFailure);
+						LoggerUtils.MultiLog([(_logger, LogLevel.Error), (issueLogger, LogLevel.Information)], "Failed to get commits for {StreamName}: {LastSuccess} and {FirstFailure}", streamConfig.Name, lastSuccess, firstFailure);
 						return suspects;
 					}
 
@@ -1411,11 +1398,11 @@ namespace HordeServer.Issues
 					
 					if (Math.Abs(delta.TotalDays) >= _buildConfig.CurrentValue.IssueFindSuspectsMaxSuccessDays)
 					{
-						_logger.LogWarning("Ignoring commits for {StreamName}: {LastSuccess} and {FirstFailure}, delta is {DeltaDays}", streamConfig.Name, lastSuccess, firstFailure.Name, delta.TotalDays);
+						LoggerUtils.MultiLog([(_logger, LogLevel.Warning), (issueLogger, LogLevel.Debug)], "Ignoring commits for {StreamName}: {LastSuccess} and {FirstFailure}, delta is {DeltaDays}", streamConfig.Name, lastSuccess, firstFailure.Name, delta.TotalDays);
 						return suspects;
 					}
 
-					_logger.LogDebug("Querying for changes in {StreamName} between {LastSuccess} and {FirstFailure}", streamConfig.Name, lastSuccess, firstFailure);
+					LoggerUtils.MultiLog([(_logger, LogLevel.Debug), (issueLogger, LogLevel.Debug)], "Querying for changes in {StreamName} between {LastSuccess} and {FirstFailure}", streamConfig.Name, lastSuccess, firstFailure);
 
 					// Use fingerprint type to get commit tags and narrow search window
 					List<CommitTag>? tags = null;
@@ -1426,7 +1413,7 @@ namespace HordeServer.Issues
 
 					// Get the submitted changes before this job
 					List<ICommit> changes = await _commitService.GetCollection(streamConfig).FindAsync(lastSuccess, includeMinCommit: false, firstFailure, includeMaxCommit: true, MaxChanges, tags, cancellationToken: cancellationToken).ToListAsync(cancellationToken);
-					_logger.LogDebug("Found {NumResults} changes", changes.Count);
+					LoggerUtils.MultiLog([(_logger, LogLevel.Debug), (issueLogger, LogLevel.Debug)], "Found {NumResults} changes", changes.Count);
 
 					// Get all the parameters used to rank suspects
 					FileFilter filter = new FileFilter(fingerprint.ChangeFilter.Split(";", StringSplitOptions.RemoveEmptyEntries | StringSplitOptions.TrimEntries));
@@ -1481,23 +1468,23 @@ namespace HordeServer.Issues
 						TimeSpan scoringDuration = DateTime.Now - scoringStart;
 						if (scoringDuration.TotalSeconds >= _buildConfig.CurrentValue.IssueFindSuspectsMaxScoringSeconds)
 						{
-							_logger.LogWarning("Scoring suspect commits took longer than {ScoringSeconds} seconds, giving up", _buildConfig.CurrentValue.IssueFindSuspectsMaxScoringSeconds);
+							LoggerUtils.MultiLog([(_logger, LogLevel.Warning), (issueLogger, LogLevel.Information)], "Scoring suspect commits took longer than {ScoringSeconds} seconds, giving up", _buildConfig.CurrentValue.IssueFindSuspectsMaxScoringSeconds);
 							break;
 						}
 
 						int rank = -1;
 						try
 						{
-							rank = await GetChangeScoreAsync(commit, filter, suspectFiles, suspectNotes, suspectSymbols, _tracer);
+							rank = await GetChangeScoreAsync(commit, filter, suspectFiles, suspectNotes, suspectSymbols, _tracer, issueLogger);
 						}
 						catch (NullReferenceException nullRefEx)
 						{
-							_logger.LogError("Exception ranking suspect change {Commit}: {Message} - {Callstack}", commit.Id, nullRefEx.Message, nullRefEx.StackTrace);
+							LoggerUtils.MultiLog([(_logger, LogLevel.Error), (issueLogger, LogLevel.Information)], "Exception ranking suspect change {CommitId}: {Message} - {Callstack}", commit.Id, nullRefEx.Message, nullRefEx.StackTrace);
 
 						}
 						catch (EpicGames.Perforce.PerforceException perforceEx)
 						{
-							_logger.LogError("Exception ranking suspect change {Commit}: {Message} - {Callstack}", commit.Id, perforceEx.Message, perforceEx.StackTrace);
+							LoggerUtils.MultiLog([(_logger, LogLevel.Error), (issueLogger, LogLevel.Information)], "Exception ranking suspect change {CommitId}: {Message} - {Callstack}", commit.Id, perforceEx.Message, perforceEx.StackTrace);
 
 							// re-throw any perforce exceptions that aren't known or expected here
 							if (!perforceEx.Message.Contains("does not contain any files in", System.StringComparison.OrdinalIgnoreCase))
@@ -1518,16 +1505,17 @@ namespace HordeServer.Issues
 							suspects.Add(suspect);
 						}
 					}
+					issueLogger?.LogDebug("Max rank of {MaxRank} achieved during suspect finding", maxRank);
 				}
 			}
 			catch (Exception ex)
 			{
-				_logger.LogError(ex, "Exception finding suspects for {StreamName}: {LastSuccess} and {FirstFailure}", streamConfig.Name, lastSuccess, firstFailure.Name);
+				LoggerUtils.MultiLog([(_logger, LogLevel.Error), (issueLogger, LogLevel.Information)], ex.ToString(), "Exception finding suspects for {StreamName}: {LastSuccess} and {FirstFailure}", streamConfig.Name, lastSuccess, firstFailure.Name);
 			}
 			return suspects;
 		}
 
-		static async ValueTask<int> GetChangeScoreAsync(ICommit commit, FileFilter filter, HashSet<string> suspectFiles, HashSet<string> suspectNotes, HashSet<string> suspectSymbols, Tracer tracer)
+		static async ValueTask<int> GetChangeScoreAsync(ICommit commit, FileFilter filter, HashSet<string> suspectFiles, HashSet<string> suspectNotes, HashSet<string> suspectSymbols, Tracer tracer, ILogger? issueLogger)
 		{
 			using TelemetrySpan span = tracer.StartActiveSpan($"{nameof(IssueService)}.{nameof(GetChangeScoreAsync)}");
 			
@@ -1544,10 +1532,12 @@ namespace HordeServer.Issues
 					if (suspectFiles.Contains(fileName))
 					{
 						fileRank += 20;
+						issueLogger?.LogDebug("Commit {CommitId} suspect rank increased based on file: {File}", commit.Id, fileName);
 					}
 					else if (suspectNotes.Contains(fileName))
 					{
 						fileRank += 10;
+						issueLogger?.LogDebug("Commit {CommitId} suspect rank increased based on note: {Note}", commit.Id, fileName);
 					}
 
 					if (suspectSymbols.Count > 0)
@@ -1557,6 +1547,7 @@ namespace HordeServer.Issues
 						{
 							fileRank = Math.Max(fileRank, 10 + (10 * matches));
 							fileRank += 10 + (10 * matches);
+							issueLogger?.LogDebug("Commit {CommitId} suspect rank increased based on {Matches} symbol matches.", commit.Id, matches);
 						}
 					}
 
@@ -1763,7 +1754,8 @@ namespace HordeServer.Issues
 				if (job.CommitId < span.FirstFailure.CommitId && (span.LastSuccess == null || job.CommitId > span.LastSuccess.CommitId))
 				{
 					NewIssueStepData newLastSuccess = new NewIssueStepData(job.CommitId, IssueSeverity.Unspecified, job.Name, job.Id, batch.Id, step.Id, step.StartTimeUtc ?? default, step.LogId, null, false);
-					List<NewIssueSpanSuspectData> newSuspects = await FindSuspectsForSpanAsync(streamConfig, span.Fingerprint, job.CommitId, span.FirstFailure.CommitId, cancellationToken);
+					ILogger issueLogger = _issueCollection.GetLogger(span.IssueId);
+					List<NewIssueSpanSuspectData> newSuspects = await FindSuspectsForSpanAsync(streamConfig, span.Fingerprint, job.CommitId, span.FirstFailure.CommitId, issueLogger, cancellationToken);
 
 					if (await _issueCollection.TryUpdateSpanAsync(span, newLastSuccess: newLastSuccess, newSuspects: newSuspects, cancellationToken: cancellationToken) == null)
 					{
diff --git a/Engine/Source/Programs/Shared/EpicGames.Horde/Logs/LoggerUtils.cs b/Engine/Source/Programs/Shared/EpicGames.Horde/Logs/LoggerUtils.cs
new file mode 100644
--- /dev/null
+++ b/Engine/Source/Programs/Shared/EpicGames.Horde/Logs/LoggerUtils.cs
@@ -0,0 +1,28 @@
+﻿// Copyright Epic Games, Inc. All Rights Reserved.
+
+using Microsoft.Extensions.Logging;
+
+namespace EpicGames.Horde.Logs
+{
+	/// <summary>
+	/// Helper methods for logger functionality
+	/// </summary>
+	public static class LoggerUtils
+	{
+		/// <summary>
+		/// Helper method to fan a single message out to a list loggers
+		/// </summary>
+		/// <param name="loggerAndLevels">Array of tuples containing the logger and log level for the message.</param>
+		/// <param name="message">Message template to be logged.</param>
+		/// <param name="args">Structured logging arguments to inject in the message template.</param>
+		public static void MultiLog((ILogger? logger, LogLevel logLevel)[] loggerAndLevels, string message, params object?[] args)
+		{
+			foreach ((ILogger? logger, LogLevel level) in loggerAndLevels)
+			{
+				#pragma warning disable CA2254 // Template should be a static 
+				logger?.Log(level, message, args);
+				#pragma warning restore CA2254 // Template should be a static 
+			}
+		}
+	}
+}
