[Horde][BuildHealth][Dashboard] Update labels for select/clear all buttons + remove the clear-all filtered subset behaviour.

diff --git a/Engine/Source/Programs/Horde/HordeDashboard/src/components/buildhealth/BuildhealthView.tsx b/Engine/Source/Programs/Horde/HordeDashboard/src/components/buildhealth/BuildhealthView.tsx
--- a/Engine/Source/Programs/Horde/HordeDashboard/src/components/buildhealth/BuildhealthView.tsx
+++ b/Engine/Source/Programs/Horde/HordeDashboard/src/components/buildhealth/BuildhealthView.tsx
@@ -355,6 +355,9 @@ const MultiListParameter: React.FC<{ handler: BuildHealthDataHandler, buildHealt
 
             // Select & Clear all for the given filtered list.
             if (option.itemType === DropdownMenuItemType.Header && option.text === SELECT_CLEAR_ALL_SPECIAL_GROUP) {
+                let fitleredOptionList = filteredOptionsRef.current.length == doptions.length ? undefined : filteredOptionsRef.current;
+                let selectAllLabel = (fitleredOptionList !== undefined && fitleredOptionList.length > 0) ? "Select Filtered" : "Select All";
+
                 return (
                     <Stack horizontal tokens={{ childrenGap: 8 }} styles={{
                         root:
@@ -366,7 +369,7 @@ const MultiListParameter: React.FC<{ handler: BuildHealthDataHandler, buildHealt
                         }
                     }}>
                         {/* We have opted to use buttons here as we: 1. don't want to persist the state of a select all option; 2. the semantics of storing a "Select/Clear All" value gets complicated when the action is applied to a filtered set; */}
-                        <DefaultButton styles={{ root: { flex: 1, height: 20 } }} text="Select All"
+                        <DefaultButton title="Selects all options (or currently filtered) within the group." styles={{ root: { flex: 1, height: 20 } }} text={selectAllLabel}
                             onClick={() => {
                                 if (!option) return;
 
@@ -380,7 +383,7 @@ const MultiListParameter: React.FC<{ handler: BuildHealthDataHandler, buildHealt
                                 option && onSelectAll?.(synthesizedOption, fitleredOptionList)
                             }}
                         />
-                        <DefaultButton styles={{ root: { flex: 1, height: 20 } }} text="Clear All"
+                        <DefaultButton title="Clears all options within the group." styles={{ root: { flex: 1, height: 20 } }} text="Clear All"
                             onClick={() => {
                                 if (!option) return;
 
@@ -389,9 +392,8 @@ const MultiListParameter: React.FC<{ handler: BuildHealthDataHandler, buildHealt
                                     selected: false,
                                 };
 
-                                // Only send a suggested list if we actually have filtered something.
-                                let fitleredOptionList = filteredOptionsRef.current.length == doptions.length ? undefined : filteredOptionsRef.current;
-                                option && onSelectAll?.(synthesizedOption, fitleredOptionList)
+                                // We do not send a suggested list for clear all; we just clear all.
+                                option && onSelectAll?.(synthesizedOption, undefined)
                             }} />
                     </Stack>
                 );
