[Horde][BuildHealth][Server] Add overload of MultiLog function to capture exception object and use it for suspect finding exceptions

diff --git a/Engine/Source/Programs/Horde/Plugins/Build/HordeServer.Build/Issues/IssueService.cs b/Engine/Source/Programs/Horde/Plugins/Build/HordeServer.Build/Issues/IssueService.cs
--- a/Engine/Source/Programs/Horde/Plugins/Build/HordeServer.Build/Issues/IssueService.cs
+++ b/Engine/Source/Programs/Horde/Plugins/Build/HordeServer.Build/Issues/IssueService.cs
@@ -1372,7 +1372,7 @@ namespace HordeServer.Issues
 		/// <param name="fingerprint">The fingerprint for the span</param>
 		/// <param name="lastSuccess">The last successful commit</param>
 		/// <param name="firstFailure">The first failing commit</param>
-		/// <param name="issueLogger">The issue's audit logger</param>
+		/// <param name="issueLogger">The audit logger for the issue that the span is associated with</param>
 		/// <param name="cancellationToken">Cancellation token for the operation</param>
 		/// <returns>List of suspects</returns>
 		async Task<List<NewIssueSpanSuspectData>> FindSuspectsForSpanAsync(StreamConfig streamConfig, IIssueFingerprint fingerprint, CommitId lastSuccess, CommitId firstFailure, ILogger? issueLogger, CancellationToken cancellationToken)
@@ -1510,7 +1510,7 @@ namespace HordeServer.Issues
 			}
 			catch (Exception ex)
 			{
-				LoggerUtils.MultiLog([(_logger, LogLevel.Error), (issueLogger, LogLevel.Information)], ex.ToString(), "Exception finding suspects for {StreamName}: {LastSuccess} and {FirstFailure}", streamConfig.Name, lastSuccess, firstFailure.Name);
+				LoggerUtils.MultiLog([(_logger, LogLevel.Error), (issueLogger, LogLevel.Information)], ex, "Exception finding suspects for {StreamName}: {LastSuccess} and {FirstFailure}", streamConfig.Name, lastSuccess, firstFailure.Name);
 			}
 			return suspects;
 		}
diff --git a/Engine/Source/Programs/Shared/EpicGames.Horde/Logs/LoggerUtils.cs b/Engine/Source/Programs/Shared/EpicGames.Horde/Logs/LoggerUtils.cs
--- a/Engine/Source/Programs/Shared/EpicGames.Horde/Logs/LoggerUtils.cs
+++ b/Engine/Source/Programs/Shared/EpicGames.Horde/Logs/LoggerUtils.cs
@@ -1,5 +1,6 @@
 ﻿// Copyright Epic Games, Inc. All Rights Reserved.
 
+using System;
 using Microsoft.Extensions.Logging;
 
 namespace EpicGames.Horde.Logs
@@ -24,5 +25,22 @@ namespace EpicGames.Horde.Logs
 				#pragma warning restore CA2254 // Template should be a static 
 			}
 		}
+
+		/// <summary>
+		/// Helper method to fan a single message out to a list loggers
+		/// </summary>
+		/// <param name="loggerAndLevels">Array of tuples containing the logger and log level for the message.</param>
+		/// <param name="exception">Exception object to include in logging.</param>
+		/// <param name="message">Message template to be logged.</param>
+		/// <param name="args">Structured logging arguments to inject in the message template.</param>
+		public static void MultiLog((ILogger? logger, LogLevel logLevel)[] loggerAndLevels, Exception exception, string message, params object?[] args)
+		{
+			foreach ((ILogger? logger, LogLevel level) in loggerAndLevels)
+			{
+				#pragma warning disable CA2254 // Template should be a static 
+				logger?.Log(level, exception, message, args);
+				#pragma warning restore CA2254 // Template should be a static 
+			}
+		}
 	}
 }
