Horde: Fix compiler warnings

#skipci
#rnx

diff --git a/Engine/Source/Programs/Horde/Plugins/Analytics/HordeServer.Analytics.Integration.Tests/Telemetry/TelemetryIntegrationTests.cs b/Engine/Source/Programs/Horde/Plugins/Analytics/HordeServer.Analytics.Integration.Tests/Telemetry/TelemetryIntegrationTests.cs
--- a/Engine/Source/Programs/Horde/Plugins/Analytics/HordeServer.Analytics.Integration.Tests/Telemetry/TelemetryIntegrationTests.cs
+++ b/Engine/Source/Programs/Horde/Plugins/Analytics/HordeServer.Analytics.Integration.Tests/Telemetry/TelemetryIntegrationTests.cs
@@ -11,7 +11,6 @@ using EpicGames.Horde.Jobs.Templates;
 using EpicGames.Horde.Projects;
 using EpicGames.Horde.Streams;
 using EpicGames.Horde.Telemetry;
-using EpicGames.Perforce;
 using HordeCommon;
 using HordeServer.Analytics.Tests;
 using HordeServer.Issues;
@@ -26,7 +25,6 @@ using HordeServer.Telemetry.Metrics;
 using HordeServer.Tests;
 using HordeServer.Tests.Issues;
 using Microsoft.Extensions.DependencyInjection;
-using StackExchange.Redis;
 
 namespace HordeServer.Analytics.Integration.Tests.Telemetry
 {
