Horde: Add ability to clean up orphaned agent session IDs

diff --git a/Engine/Source/Programs/Horde/HordeServer/Server/DebugController.cs b/Engine/Source/Programs/Horde/HordeServer/Server/DebugController.cs
--- a/Engine/Source/Programs/Horde/HordeServer/Server/DebugController.cs
+++ b/Engine/Source/Programs/Horde/HordeServer/Server/DebugController.cs
@@ -11,6 +11,7 @@ using System.Threading;
 using System.Threading.Tasks;
 using System.Web;
 using EpicGames.Core;
+using HordeServer.Agents;
 using HordeServer.Configuration;
 using HordeServer.Utilities;
 using JetBrains.Profiler.SelfApi;
@@ -37,6 +38,7 @@ namespace HordeServer.Server
 
 		private readonly IMongoService _mongoService;
 		private readonly ConfigService _configService;
+		private readonly IAgentScheduler _agentScheduler;
 		private readonly IOptionsSnapshot<GlobalConfig> _globalConfig;
 		private readonly ILogger<DebugController> _logger;
 
@@ -46,11 +48,13 @@ namespace HordeServer.Server
 		public DebugController(
 			IMongoService mongoService,
 			ConfigService configService,
+			IAgentScheduler agentScheduler,
 			IOptionsSnapshot<GlobalConfig> globalConfig,
 			ILogger<DebugController> logger)
 		{
 			_mongoService = mongoService;
 			_configService = configService;
+			_agentScheduler = agentScheduler;
 			_globalConfig = globalConfig;
 			_logger = logger;
 		}
@@ -381,5 +385,24 @@ namespace HordeServer.Server
 			string stringArg = "hello";
 			throw new Exception($"Message: numberArg:{numberArg}, stringArg:{stringArg}");
 		}
+
+		/// <summary>
+		/// Clean up orphaned agent sessions
+		/// </summary>
+		/// <returns></returns>
+		[HttpGet]
+		[Route("/api/v1/debug/agent-orphaned-session-cleanup")]
+		public async Task<ActionResult<object>> AgentSessionCleanupAsync(CancellationToken cancellationToken = default)
+		{
+			if (!_globalConfig.Value.Authorize(ServerAclAction.Debug, User))
+			{
+				return Forbid(ServerAclAction.Debug);
+			}
+
+			return Ok(new 
+			{ 
+				SessionsRemoved = await _agentScheduler.RemoveOrphanedSessionsAsync(cancellationToken),
+			});
+		}
 	}
 }
diff --git a/Engine/Source/Programs/Horde/Plugins/Compute/HordeServer.Compute.Tests/Agents/AgentSchedulerTests.cs b/Engine/Source/Programs/Horde/Plugins/Compute/HordeServer.Compute.Tests/Agents/AgentSchedulerTests.cs
--- a/Engine/Source/Programs/Horde/Plugins/Compute/HordeServer.Compute.Tests/Agents/AgentSchedulerTests.cs
+++ b/Engine/Source/Programs/Horde/Plugins/Compute/HordeServer.Compute.Tests/Agents/AgentSchedulerTests.cs
@@ -9,8 +9,10 @@ using HordeCommon.Rpc;
 using HordeCommon.Rpc.Messages;
 using HordeServer.Agents;
 using HordeServer.Agents.Sessions;
+using HordeServer.Server;
 using HordeServer.Utilities;
 using Microsoft.Extensions.DependencyInjection;
+using StackExchange.Redis;
 
 namespace HordeServer.Tests.Agents
 {
@@ -129,6 +131,25 @@ namespace HordeServer.Tests.Agents
 			Assert.AreEqual(5, sessionIds2.Length);
 		}
 
+		[TestMethod]
+		public async Task FindTerminatedSessionsAsync()
+		{
+			RpcSession? session = await AgentScheduler.TryCreateSessionAsync(new AgentId("foo"), SessionIdUtils.GenerateNewId(), new RpcAgentCapabilities());
+			Assert.IsNotNull(session);
+			Assert.AreEqual(0, (await AgentScheduler.RemoveOrphanedSessionsAsync()).Count);
+
+			SessionId orphanedId = SessionIdUtils.GenerateNewId();
+			IDatabase database = GetRedisServiceSingleton().GetDatabase();
+			string orphanedKey = HordeServer.Agents.AgentScheduler.SessionStateKey(orphanedId.ToString());
+			await database.StringSetAsync(orphanedKey, "not-relevant");
+
+			// Ensure only orphaned session are removed and valid ones are left intact
+			List<SessionId> orphanedIds = await AgentScheduler.RemoveOrphanedSessionsAsync();
+			CollectionAssert.AreEquivalent(new List<SessionId> { orphanedId }, orphanedIds);
+			session = await AgentScheduler.TryGetSessionAsync(session.SessionId);
+			Assert.IsNotNull(session);
+		}
+
 		[TestMethod]
 		public async Task AddSessionsToExistingQueueTestAsync()
 		{
diff --git a/Engine/Source/Programs/Horde/Plugins/Compute/HordeServer.Compute/Agents/AgentScheduler.cs b/Engine/Source/Programs/Horde/Plugins/Compute/HordeServer.Compute/Agents/AgentScheduler.cs
--- a/Engine/Source/Programs/Horde/Plugins/Compute/HordeServer.Compute/Agents/AgentScheduler.cs
+++ b/Engine/Source/Programs/Horde/Plugins/Compute/HordeServer.Compute/Agents/AgentScheduler.cs
@@ -17,7 +17,7 @@ using StackExchange.Redis;
 
 namespace HordeServer.Agents
 {
-	class AgentScheduler : IAgentScheduler, IHostedService, IAsyncDisposable
+	internal class AgentScheduler : IAgentScheduler, IHostedService, IAsyncDisposable
 	{
 		/// <summary>
 		/// Time after which we expire filters
@@ -79,6 +79,14 @@ namespace HordeServer.Agents
 				=> new($"compute:sessions:{SessionId}:state");
 		}
 
+		internal static string SessionStateKey(string sessionId) => $"compute:sessions:{sessionId}:state";
+
+		SessionId FromSessionStateKey(string sessionStateKey)
+		{
+			string[] parts = sessionStateKey.Split(":");
+			return parts.Length == 4 ? SessionId.Parse(parts[2]) : throw new ArgumentException($"Unable to extract session ID from key: {sessionStateKey}");
+		}
+
 		record struct LeaseCollectionKeys
 		{
 			public RedisSetKey<LeaseId> Current
@@ -392,7 +400,7 @@ namespace HordeServer.Agents
 			// Execute the transaction
 			if (!await transaction.ExecuteAsync().WaitAsync(cancellationToken))
 			{
-				_logger.LogDebug("Transaction failed for update of session {SessionId}", sessionId);
+				_logger.LogDebug("Redis transaction failed for update of session {SessionId}", sessionId);
 				return null;
 			}
 
@@ -430,6 +438,32 @@ namespace HordeServer.Agents
 			}
 		}
 
+		/// <inheritdoc/>
+		public async Task<List<SessionId>> RemoveOrphanedSessionsAsync(CancellationToken cancellationToken = default)
+		{
+			IServer server = _redisService.GetConnection().GetServers().First();
+			IDatabase database = _redisService.GetDatabase();
+
+			// Sessions with a state key
+			HashSet<SessionId> sessionStates = await server
+				.KeysAsync(database: database.Database, pattern: SessionStateKey("*"))
+				.Select(key => FromSessionStateKey(key.ToString()))
+				.ToHashSetAsync(cancellationToken);
+			
+			// Sessions registered in the "compute:sessions:current" hash key
+			HashSet<SessionId> registeredSessions = (await database.HashGetAllAsync(Keys.Sessions.Current))
+				.Select(x => x.Name).ToHashSet();
+
+			List<SessionId> sessionIds = sessionStates.Except(registeredSessions).ToList();
+			foreach (SessionId sid in sessionIds)
+			{
+				_logger.LogInformation("Deleting orphaned session ID {SessionId}", sid);
+				await database.KeyDeleteAsync(sid.ToString());
+			}
+
+			return sessionIds;
+		}
+
 		#endregion
 		#region Filters
 
diff --git a/Engine/Source/Programs/Horde/Plugins/Compute/HordeServer.Compute/Agents/IAgentScheduler.cs b/Engine/Source/Programs/Horde/Plugins/Compute/HordeServer.Compute/Agents/IAgentScheduler.cs
--- a/Engine/Source/Programs/Horde/Plugins/Compute/HordeServer.Compute/Agents/IAgentScheduler.cs
+++ b/Engine/Source/Programs/Horde/Plugins/Compute/HordeServer.Compute/Agents/IAgentScheduler.cs
@@ -123,6 +123,14 @@ namespace HordeServer.Agents
 		/// <returns>List of expired sessions</returns>
 		IAsyncEnumerable<RpcSession> FindExpiredSessionsAsync(CancellationToken cancellationToken = default);
 
+		/// <summary>
+		/// Remove session states that are not registered in the "current session key"
+		/// This can happen when upgrading Horde to a newer version and old agents that created new sessions after that.
+		/// Clearing Redis can also fix this issue.
+		/// </summary>
+		/// <param name="cancellationToken">Cancellation token</param>
+		public Task<List<SessionId>> RemoveOrphanedSessionsAsync(CancellationToken cancellationToken = default);
+
 		/// <summary>
 		/// Finds all active agent lease IDs
 		/// </summary>
