* Handle digestion of extra test description and test recipe
* Add test recipe REST API endpoint
* Surface error fingerprints to the test session

#rnx

#code-analyzed

diff --git a/Engine/Source/Programs/Horde/Plugins/Build/HordeServer.Build.Tests/Jobs/TestData/TestDataTests.cs b/Engine/Source/Programs/Horde/Plugins/Build/HordeServer.Build.Tests/Jobs/TestData/TestDataTests.cs
--- a/Engine/Source/Programs/Horde/Plugins/Build/HordeServer.Build.Tests/Jobs/TestData/TestDataTests.cs
+++ b/Engine/Source/Programs/Horde/Plugins/Build/HordeServer.Build.Tests/Jobs/TestData/TestDataTests.cs
@@ -3,6 +3,7 @@
 using System.Security.Claims;
 using EpicGames.Horde;
 using EpicGames.Horde.Commits;
+using EpicGames.Horde.Issues;
 using EpicGames.Horde.Jobs;
 using EpicGames.Horde.Jobs.Graphs;
 using EpicGames.Horde.Jobs.Templates;
@@ -95,11 +96,14 @@ namespace HordeServer.Tests.Jobs.TestData
 
 			await SetConfigAsync(globalConfig);
 
+			NodeAnnotations workflowAnnotations = new NodeAnnotations();
+			workflowAnnotations.WorkflowId = new WorkflowId("test-workflow-id");
+
 			List<INode> nodes = new List<INode>();
-			nodes.Add(MockNode("Update Version Files", NodeAnnotations.Empty));
-			nodes.Add(MockNode("Compile UnrealHeaderTool Win64", NodeAnnotations.Empty));
-			nodes.Add(MockNode("Compile ShooterGameEditor Win64", NodeAnnotations.Empty));
-			nodes.Add(MockNode("Cook ShooterGame Win64", NodeAnnotations.Empty));
+			nodes.Add(MockNode("Update Version Files", workflowAnnotations));
+			nodes.Add(MockNode("Compile UnrealHeaderTool Win64", workflowAnnotations));
+			nodes.Add(MockNode("Compile ShooterGameEditor Win64", workflowAnnotations));
+			nodes.Add(MockNode("Cook ShooterGame Win64", workflowAnnotations));
 
 			Mock<INodeGroup> grp = new Mock<INodeGroup>(MockBehavior.Strict);
 			grp.SetupGet(x => x.Nodes).Returns(nodes);
@@ -130,6 +134,8 @@ namespace HordeServer.Tests.Jobs.TestData
 					step.SetupGet(x => x.NodeIdx).Returns(nodeIdx);
 					step.SetupGet(x => x.LogId).Returns(logFile.Id);
 					step.SetupGet(x => x.StartTimeUtc).Returns(DateTime.UtcNow + time);
+					step.SetupGet(x => x.Name).Returns(@group.Nodes[nodeIdx].Name);
+					step.SetupGet(x => x.Annotations).Returns(@group.Nodes[nodeIdx].Annotations);
 
 					steps.Add(step.Object);
 				}
@@ -396,6 +402,7 @@ namespace HordeServer.Tests.Jobs.TestData
 			Assert.IsTrue(testnames.Any(t => t.Key == "UE.Automation(Group:Editor) EngineTest"));
 			Assert.IsTrue(testnames.Any(t => t.Key == "UE.Automation(Group:Editor) Lyra"));
 			Assert.AreEqual("Editor", testnames[0].Name);
+			Assert.IsNull(testnames[0].Team);
 
 			GetMetadataRequest metaRequest = new GetMetadataRequest() { Entries = new() { { "Platform", "Win64" } } };
 			ActionResult<List<GetTestMetadataResponse>> metaResult = await TestDataController.GetTestMetadataV2Async(metaRequest);
@@ -427,7 +434,8 @@ namespace HordeServer.Tests.Jobs.TestData
 			Assert.AreEqual(3, tests[0].Phases.Count);
 			Assert.IsTrue(tests[0].Phases.Any(i => i.Name == "EditConditionParser.EvaluateBool"));
 			Assert.IsTrue(tests[0].Phases.Any(i => i.Name == "Editor.Failing.Test"));
-			string phaseId = tests[0].Phases.First(i => i.Name == "Editor.Failing.Test").Id.ToString();
+			GetPhaseResponse phase = tests[0].Phases.First(i => i.Name == "Editor.Failing.Test");
+			string phaseId = phase.Id.ToString();
 
 			GetTestSessionsRequest sessionsRequest = new GetTestSessionsRequest() { StreamIds = streamIds.ToList(), TestIds = testnames.Select(t => t.Id.ToString()).ToList() };
 			ActionResult<List<GetTestSessionResponse>> testSessionResult = await TestDataController.GetTestSessionsV2Async(sessionsRequest);
@@ -438,6 +446,12 @@ namespace HordeServer.Tests.Jobs.TestData
 			Assert.AreEqual(TestOutcome.Failure, sessions[0].Outcome);
 			Assert.IsTrue(sessions[0].TagIds?.All(id => tags.Any(t => t.Id == id)));
 
+			Assert.IsNotNull(sessions[0].ErrorFingerprints);
+			Assert.AreEqual(1, sessions[0].ErrorFingerprints!.Count);
+			GetErrorFingerprintResponse errorFingerprint = sessions[0].ErrorFingerprints![0];
+			Assert.AreEqual(1, errorFingerprint.Phases.Count);
+			Assert.AreEqual(phase.Key, errorFingerprint.Phases[0]);
+
 			GetPhaseSessionsRequest phaseSessionRequest = new GetPhaseSessionsRequest() { StreamIds = streamIds.ToList(), PhaseIds = tests.SelectMany(t => t.Phases.Select(p => p.Id)).ToList() };
 			ActionResult<List<GetTestPhaseSessionResponse>> phaseSessionResult = await TestDataController.GetTestPhaseSessionsV2Async(phaseSessionRequest);
 			Assert.IsNotNull(phaseSessionResult);
@@ -447,6 +461,8 @@ namespace HordeServer.Tests.Jobs.TestData
 			Assert.AreEqual(TestPhaseOutcome.Failed, phaseSessions.First(p => p.PhaseRef == phaseId).Outcome);
 			Assert.IsTrue(phaseSessions.Any(s => s.TagIds?.All(id => tags.Any(t => t.Id == id)) ?? false));
 
+			Assert.IsTrue(phaseSessions.First(p => p.PhaseRef == phaseId).ErrorFingerprint == errorFingerprint.Key);
+
 			// Get the sessions for one phase in one stream
 			phaseSessionRequest = new GetPhaseSessionsRequest() { StreamIds = [streamIds[0]], PhaseIds = [tests[0].Phases[0].Id] };
 			phaseSessionResult = await TestDataController.GetTestPhaseSessionsV2Async(phaseSessionRequest);
@@ -472,6 +488,79 @@ namespace HordeServer.Tests.Jobs.TestData
 			Assert.AreEqual(4, testDataIds.Value.Count);
 		}
 
+		[TestMethod]
+		public async Task SessionReportTestV2RecipeAsync()
+		{
+			await TestDataService.StartAsync(CancellationToken.None);
+
+			string[] streamIds = new string[] { _mainStreamId.ToString(), _releaseStreamId.ToString() };
+
+			IJob job = CreateJob(_mainStreamId, 105, "Test Build", _graph);
+			IJobStep step = job.Batches[0].Steps[0];
+			IJob job2 = CreateJob(_releaseStreamId, 105, "Test Build", _graph);
+			IJobStep step2 = job2.Batches[0].Steps[0];
+			IJob job3 = CreateJob(_mainStreamId, 106, "Test Build", _graph);
+			IJobStep step3 = job3.Batches[0].Steps[0];
+
+			BsonDocument testData = BsonDocument.Parse(String.Join('\n', _sessionTestDataV2Lines4));
+			List<(string key, BsonDocument)> data = new List<(string key, BsonDocument)>();
+			foreach (BsonValue item in testData.GetValue("items").AsBsonArray.ToList())
+			{
+				BsonDocument value = item.AsBsonDocument.GetValue("data").AsBsonDocument;
+				data.Add((item.AsBsonDocument.GetValue("key").AsString, value));
+			}
+
+			BsonDocument testData2 = BsonDocument.Parse(String.Join('\n', _sessionTestDataV2Lines4));
+			List<(string key, BsonDocument)> data2 = new List<(string key, BsonDocument)>();
+			foreach (BsonValue item in testData2.GetValue("items").AsBsonArray.ToList())
+			{
+				BsonDocument value = item.AsBsonDocument.GetValue("data").AsBsonDocument;
+				data2.Add((item.AsBsonDocument.GetValue("key").AsString, value));
+			}
+
+			await TestDataCollectionV2.AddAsync(job, step, data.ToArray());
+			await TestDataCollectionV2.AddAsync(job2, step2, data.ToArray());
+			await TestDataCollectionV2.AddAsync(job3, step3, data2.ToArray());
+
+			await TestDataService.TickForTestingAsync();
+
+			GetStreamTestsRequest streamRequest = new GetStreamTestsRequest() { StreamIds = streamIds.ToList() };
+			ActionResult<List<GetTestSessionStreamResponse>> streamResult = await TestDataController.GetTestStreamsV2Async(streamRequest);
+			Assert.IsNotNull(streamResult);
+			Assert.IsNotNull(streamResult.Value);
+			List<GetTestSessionStreamResponse> streams = streamResult.Value;
+			Assert.AreNotEqual(0, streams.Count);
+
+			GetTestsRequest testRequest = new GetTestsRequest() { TestIds = streams[0].Tests.Select(t => t.Id).ToList() };
+			ActionResult<List<GetTestNameResponse>> testnameResult = await TestDataController.GetTestNameRefV2Async(testRequest);
+			Assert.IsNotNull(testnameResult);
+			Assert.IsNotNull(testnameResult.Value);
+			List<GetTestNameResponse> testnames = testnameResult.Value;
+			Assert.AreEqual(1, testnames.Count);
+			Assert.IsTrue(testnames.Any(t => t.Key == "UE.Automation(Group:Editor) Lyra"));
+			Assert.AreEqual("Editor", testnames[0].Name);
+			Assert.AreEqual("UE", testnames[0].Team);
+			Assert.AreEqual("UE Functional Test", testnames[0].Harness);
+			Assert.AreEqual("test intent", testnames[0].Intent);
+			Assert.AreEqual("functional", testnames[0].Profile);
+
+			GetTestRecipesRequest recipeRequest = new GetTestRecipesRequest() { TestIds = testnames.Select(t => t.Id).ToList() };
+			ActionResult<List<GetTestRecipeResponse>> recipeResult = await TestDataController.GetTestRecipeRefV2Async(recipeRequest);
+			Assert.IsNotNull(recipeResult);
+			Assert.IsNotNull(recipeResult.Value);
+			List<GetTestRecipeResponse> testRecipes = recipeResult.Value;
+			Assert.AreEqual(2, testRecipes.Count);
+			Assert.AreEqual("test-workflow-id", testRecipes[0].WorkflowId);
+			Assert.AreEqual(_mainStreamId.ToString(), testRecipes[0].StreamId);
+			Assert.AreEqual("test-template", testRecipes[0].JobTemplateId);
+			Assert.AreEqual("Update Version Files", testRecipes[0].JobStep);
+			Assert.AreEqual(10, testRecipes[0].TimeoutMinutes);
+			Assert.IsTrue(testRecipes[0].Params?.Any(p => p.Key == "group" && p.Value.Equals("editor")));
+			Assert.IsTrue(testRecipes[0].Params?.Any(p => p.Key == "project" && p.Value.Equals("Lyra")));
+
+			Assert.AreEqual(_releaseStreamId.ToString(), testRecipes[1].StreamId);
+		}
+
 		private readonly string[] _simpleTestDataLines =
 		{
 			@"{",
@@ -1066,5 +1155,102 @@ namespace HordeServer.Tests.Jobs.TestData
 			@"    ]",
 			@"}"
 		};
+
+		private readonly string[] _sessionTestDataV2Lines4 =
+		{
+			@"{",
+			@"  ""items"": [",
+			@"        {",
+			@"            ""key"": ""UE.Automation(Group:Editor) Lyra"",",
+			@"            ""data"": {",
+			@"   		      ""version"" : 2,",
+			@"                ""summary"": {",
+			@"                   ""testName"": ""Editor"",",
+			@"        		     ""dateTime"": ""2025-03-03T18:14:04"",",
+			@"		             ""timeElapseSec"": 207.02324,",
+			@"                   ""commandline"": ""run ue.automation -group=editor -project=Lyra"",",
+			@"        		     ""phasesTotalCount"": 3,",
+			@"		             ""phasesSucceededCount"": 2,",
+			@"        		     ""phasesUndefinedCount"": 0,",
+			@"		             ""phasesFailedCount"": 1",
+			@"        		  },",
+			@"                ""phases"": [",
+			@"                    {",
+			@"                        ""key"": ""874a6711"",",
+			@"                        ""name"": ""EditConditionParser.EvaluateBitFlags"",",
+			@"                        ""dateTime"": ""2025-03-03T18:07:55"",",
+			@"                        ""timeElapseSec"": 0.0879574,",
+			@"                        ""outcome"": ""Success"",",
+			@"                        ""deviceKeys"": [",
+			@"                            ""9AF6B8CC484096D351EC56ABDFBB919F""",
+			@"                        ]",
+			@"                    },",
+			@"                    {",
+			@"                        ""key"": ""03f772e3"",",
+			@"                        ""name"": ""EditConditionParser.EvaluateBool"",",
+			@"                        ""dateTime"": ""2025-03-03T18:08:08"",",
+			@"                        ""timeElapseSec"": 0.008163903,",
+			@"                        ""outcome"": ""Success"",",
+			@"                        ""deviceKeys"": [",
+			@"                            ""9AF6B8CC484096D351EC56ABDFBB919F""",
+			@"                        ]",
+			@"                   },",
+			@"                   {",
+			@"                        ""key"": ""62d99b42"",",
+			@"                        ""name"": ""Editor.Failing.Test"",",
+			@"                        ""dateTime"": ""2025-03-03T18:08:14"",",
+			@"                        ""timeElapseSec"": 0.1961886,",
+			@"                        ""outcome"": ""Failed"",",
+			@"                        ""errorFingerprint"": ""a6efc189"",",
+			@"                        ""deviceKeys"": [",
+			@"                            ""9AF6B8CC484096D351EC56ABDFBB919F""",
+			@"                        ],",
+			@"                        ""eventStreamPath"": ""Engine/Programs/AutomationTool/Saved/Logs/UE.EditorAutomation(RunTest=Group_Editor)_(Win64_Development_Editor)/EventStreams/62d99b42.json""",
+			@"                   }",
+			@"                ]",
+			@"                ""description"": {",
+			@"                    ""intent"": ""test intent"",",
+ 			@"                    ""profile"": ""functional"",",
+			@"                    ""harness"": ""UE Functional Test"",",
+ 			@"                    ""team"": ""UE""",
+			@"                }",
+			@"                ""recipe"": {",
+			@"                    ""nodeClass"": ""testClass"",",
+ 			@"                    ""version"": ""1.0.0"",",
+ 			@"                    ""timeoutMinutes"": 10,",
+			@"                    ""harness"": ""UE Functional Test"",",
+ 			@"                    ""params"": {",
+			@"                      ""group"": ""editor"",",
+			@"                      ""project"": ""Lyra""",
+			@"                    }",
+			@"                }",
+			@"                ""metadata"": {",
+			@"                    ""Platform"": ""Win64"",",
+ 			@"                    ""BuildTarget"": ""Editor"",",
+			@"                    ""Configuration"": ""Development"",",
+ 			@"                    ""Project"": ""Lyra"",",
+			@"                    ""RHI"": ""default""",
+			@"                }",
+			@"                ""devices"": {",
+			@"                  ""9AF6B8CC484096D351EC56ABDFBB919F"": {",
+			@"                    ""name"": ""M-SIHLOYQHDOUKO"",",
+			@"                    ""appInstanceLogPath"": ""Engine/Programs/AutomationTool/Saved/Logs/UE.EditorAutomation(RunTest=Group_Editor)_(Win64_Development_Editor)/Editor/EditorOutput.log"",",
+			@"                    ""metadata"": {",
+			@"                      ""platform"": ""WindowsEditor"",",
+			@"                      ""os_version"": ""Windows 10 (22H2) [10.0.19045.5487] "",",
+			@"                      ""model"": ""Default"",",
+			@"                      ""gpu"": ""Parsec Virtual Display Adapter"",",
+			@"                      ""cpumodel"": ""AMD EPYC 7R32"",",
+			@"                      ""ram_in_gb"": ""127"",",
+			@"                      ""render_mode"": ""D3D12_SM6"",",
+			@"                      ""rhi"": ""D3D12 (SM6)""",
+			@"                    }",
+			@"                  }",
+			@"                },",
+			@"            }",
+			@"        }",
+			@"    ]",
+			@"}"
+		};
 	}
 }
diff --git a/Engine/Source/Programs/Horde/Plugins/Build/HordeServer.Build/Jobs/TestData/ITestData.cs b/Engine/Source/Programs/Horde/Plugins/Build/HordeServer.Build/Jobs/TestData/ITestData.cs
--- a/Engine/Source/Programs/Horde/Plugins/Build/HordeServer.Build/Jobs/TestData/ITestData.cs
+++ b/Engine/Source/Programs/Horde/Plugins/Build/HordeServer.Build/Jobs/TestData/ITestData.cs
@@ -1,6 +1,7 @@
 // Copyright Epic Games, Inc. All Rights Reserved.
 
 using EpicGames.Horde.Commits;
+using EpicGames.Horde.Issues;
 using EpicGames.Horde.Jobs;
 using EpicGames.Horde.Jobs.Templates;
 using EpicGames.Horde.Jobs.TestData;
@@ -282,6 +283,26 @@ namespace HordeServer.Jobs.TestData
 		/// The display name of the test
 		/// </summary>
 		string Name { get; }
+
+		/// <summary>
+		/// The team associated with the test
+		/// </summary>
+		string? Team { get; }
+
+		/// <summary>
+		/// The test intent
+		/// </summary>
+		string? Intent { get; }
+
+		/// <summary>
+		/// The harness used by the test
+		/// </summary>
+		string? Harness { get; }
+
+		/// <summary>
+		/// The test profile (unit test, functional, performance, etc)
+		/// </summary>
+		string? Profile { get; }
 	}
 
 	/// <summary>
@@ -342,6 +363,78 @@ namespace HordeServer.Jobs.TestData
 		IReadOnlyList<ITestMetaEntry> Entries { get; }
 	}
 
+	/// <summary>
+	/// A param entry key value pair
+	/// </summary>
+	public interface ITestParamEntry
+	{
+		/// <summary>
+		/// The param key
+		/// </summary>
+		string Key { get; }
+
+		/// <summary>
+		/// The param value
+		/// </summary>
+		object Value { get; }
+	}
+
+	/// <summary>
+	/// A test recipe reference
+	/// </summary>
+	public interface ITestRecipeRef
+	{
+		/// <summary>
+		/// The test recipe id
+		/// </summary>
+		TestRecipeId Id { get; }
+
+		/// <summary>
+		/// The test name id associated with this recipe
+		/// </summary>
+		TestId NameRef { get; }
+
+		/// <summary>
+		/// The stream id associated with this recipe
+		/// </summary>
+		StreamId StreamId { get; }
+
+		/// <summary>
+		/// The job template associated with this recipe
+		/// </summary>
+		TemplateId JobTemplateId { get; }
+
+		/// <summary>
+		/// The job step node name associated with this recipe
+		/// </summary>
+		string JobStep { get; }
+
+		/// <summary>
+		/// The workflow associated with this recipe
+		/// </summary>
+		WorkflowId? WorkflowId { get; }
+
+		/// <summary>
+		/// The test node class
+		/// </summary>
+		string NodeClass { get; }
+
+		/// <summary>
+		/// The test version
+		/// </summary>
+		string Version { get; }
+
+		/// <summary>
+		/// The test timeout in minute
+		/// </summary>
+		float TimeoutMinutes { get; }
+
+		/// <summary>
+		/// The test params
+		/// </summary>
+		IReadOnlyList<ITestParamEntry>? Params { get; }
+	}
+
 	/// <summary>
 	/// A test tag reference
 	/// </summary>
@@ -373,6 +466,11 @@ namespace HordeServer.Jobs.TestData
 		/// </summary>
 		TestMetaId Metadata { get; }
 
+		/// <summary>
+		/// The recipe associated with this test session
+		/// </summary>
+		TestRecipeId? Recipe { get; }
+
 		/// <summary>
 		/// The tags associated with the test session
 		/// </summary>
@@ -446,6 +544,27 @@ namespace HordeServer.Jobs.TestData
 		/// The number of phases that failed during the test session
 		/// </summary>
 		int PhasesFailedCount { get; }
+
+		/// <summary>
+		/// The unique error fingerprints collected from the phase events
+		/// </summary>
+		IReadOnlyList<IErrorFingerprintEntry>? ErrorFingerprints { get; }
+	}
+
+	/// <summary>
+	/// A error fingerprint entry
+	/// </summary>
+	public interface IErrorFingerprintEntry
+	{
+		/// <summary>
+		/// The error fingerprint key
+		/// </summary>
+		string Key { get; }
+
+		/// <summary>
+		/// The phase keys associated with the error fingerprint
+		/// </summary>
+		IReadOnlyList<string> Phases { get; }
 	}
 
 	/// <summary>
diff --git a/Engine/Source/Programs/Horde/Plugins/Build/HordeServer.Build/Jobs/TestData/ITestDataCollection.cs b/Engine/Source/Programs/Horde/Plugins/Build/HordeServer.Build/Jobs/TestData/ITestDataCollection.cs
--- a/Engine/Source/Programs/Horde/Plugins/Build/HordeServer.Build/Jobs/TestData/ITestDataCollection.cs
+++ b/Engine/Source/Programs/Horde/Plugins/Build/HordeServer.Build/Jobs/TestData/ITestDataCollection.cs
@@ -226,6 +226,16 @@ namespace HordeServer.Jobs.TestData
 		/// <returns></returns>
 		Task<IReadOnlyList<ITestTagRef>> FindTestTagsAsync(TestTagId[]? tagIds = null, string[]? names = null, CancellationToken cancellationToken = default);
 
+		/// <summary>
+		/// Find test recipe references
+		/// </summary>
+		/// <param name="recipeIds"></param>
+		/// <param name="testIds"></param>
+		/// <param name="streamIds"></param>
+		/// <param name="cancellationToken">Cancellation token for the operation</param>
+		/// <returns></returns>
+		Task<IReadOnlyList<ITestRecipeRef>> FindTestRecipesAsync(TestRecipeId[]? recipeIds = null, TestId[]? testIds = null, StreamId[]? streamIds = null, CancellationToken cancellationToken = default);
+
 		/// <summary>
 		/// Find test phase sessions from a set of test phases
 		/// </summary>
diff --git a/Engine/Source/Programs/Horde/Plugins/Build/HordeServer.Build/Jobs/TestData/TestDataCollectionV2.cs b/Engine/Source/Programs/Horde/Plugins/Build/HordeServer.Build/Jobs/TestData/TestDataCollectionV2.cs
--- a/Engine/Source/Programs/Horde/Plugins/Build/HordeServer.Build/Jobs/TestData/TestDataCollectionV2.cs
+++ b/Engine/Source/Programs/Horde/Plugins/Build/HordeServer.Build/Jobs/TestData/TestDataCollectionV2.cs
@@ -1,6 +1,7 @@
 // Copyright Epic Games, Inc. All Rights Reserved.
 
 using EpicGames.Horde.Commits;
+using EpicGames.Horde.Issues;
 using EpicGames.Horde.Jobs;
 using EpicGames.Horde.Jobs.Templates;
 using EpicGames.Horde.Jobs.TestData;
@@ -82,6 +83,98 @@ namespace HordeServer.Jobs.TestData
 			}
 		}
 
+		/// <summary>
+		/// Param key value pair
+		/// </summary>
+		class TestParamEntry : ITestParamEntry
+		{
+			[BsonRequired, BsonElement("k")]
+			public string Key { get; set; }
+
+			[BsonRequired, BsonElement("v")]
+			public object Value { get; set; }
+
+			private TestParamEntry()
+			{
+				Key = String.Empty;
+				Value = false;
+			}
+
+			public TestParamEntry(string key, object data)
+			{
+				Key = key;
+				Value = data;
+			}
+		}
+
+		/// <summary>
+		/// Test recipe reference document
+		/// </summary>
+		class TestRecipeDocument : ITestRecipeRef, ITestExpire
+		{
+			[BsonRequired, BsonId]
+			public TestRecipeId Id { get; set; }
+
+			[BsonRequired, BsonElement("t")]
+			public TestId NameRef { get; set; }
+
+			[BsonRequired, BsonElement("sid")]
+			public StreamId StreamId { get; set; }
+
+			[BsonRequired, BsonElement("n")]
+			public string NodeClass { get; set; }
+
+			[BsonRequired, BsonElement("v")]
+			public string Version { get; set; }
+
+			[BsonRequired, BsonElement("o")]
+			public float TimeoutMinutes { get; set; }
+
+			[BsonRequired, BsonElement("jt")]
+			public TemplateId JobTemplateId { get; set; }
+			
+			[BsonRequired, BsonElement("js")]
+			public string JobStep { get; set; }
+
+			[BsonIgnoreIfNull, BsonElement("w")]
+			public WorkflowId? WorkflowId { get; set; }
+
+			[BsonIgnoreIfNull, BsonElement("p")]
+			public List<TestParamEntry>? Params { get; set; }
+			IReadOnlyList<ITestParamEntry>? ITestRecipeRef.Params => Params;
+
+			[BsonRequired, BsonElement("s")]
+			public DateTime LastSeenUtc { get; set; }
+
+			private TestRecipeDocument()
+			{
+				NodeClass = String.Empty;
+				Version = String.Empty;
+				JobStep = String.Empty;
+			}
+
+			public TestRecipeDocument(TestId testId, StreamId streamId, TemplateId jobTemplate, string jobStep, string nodeClass, string version, float timeoutMinutes, WorkflowId? workflow = null)
+			{
+				Id = TestRecipeId.GenerateNewId();
+				NameRef = testId;
+				StreamId = streamId;
+				JobTemplateId = jobTemplate;
+				JobStep = jobStep;
+				NodeClass = nodeClass;
+				Version = version;
+				TimeoutMinutes = timeoutMinutes;
+				WorkflowId = workflow;
+				LastSeenUtc = DateTime.UtcNow;
+			}
+
+			public IReadOnlyList<ITestParamEntry> UpdateParams(IReadOnlyDictionary<string, object> @params)
+			{
+				Params = @params.Select(i => new TestParamEntry(i.Key, i.Value)).ToList();
+
+				return Params;
+			}
+		}
+
 		/// <summary>
 		/// Test tag reference document
 		/// </summary>
@@ -123,6 +216,18 @@ namespace HordeServer.Jobs.TestData
 			[BsonRequired, BsonElement("k")]
 			public string Key { get; set; }
 
+			[BsonIgnoreIfNull, BsonElement("t")]
+			public string? Team { get; set; }
+
+			[BsonIgnoreIfNull, BsonElement("i")]
+			public string? Intent { get; set; }
+
+			[BsonIgnoreIfNull, BsonElement("h")]
+			public string? Harness { get; set;  }
+
+			[BsonIgnoreIfNull, BsonElement("p")]
+			public string? Profile { get; set; }
+
 			[BsonRequired, BsonElement("s")]
 			public DateTime LastSeenUtc { get; set; }
 
@@ -139,6 +244,14 @@ namespace HordeServer.Jobs.TestData
 				Name = name;
 				LastSeenUtc = DateTime.UtcNow;
 			}
+
+			public void AddDescription(string? team, string? intent, string? harness, string? profile)
+			{
+				Team = team;
+				Intent = intent;
+				Harness = harness;
+				Profile = profile;
+			}
 		}
 
 		/// <summary>
@@ -212,6 +325,31 @@ namespace HordeServer.Jobs.TestData
 			}
 		}
 
+		/// <summary>
+		/// Store the error fingerprint entries
+		/// </summary>
+		class TestErrorfingerprintEntry : IErrorFingerprintEntry
+		{
+			[BsonRequired, BsonElement("k")]
+
+			public string Key { get; set; }
+
+			[BsonRequired, BsonElement("p")]
+			public IReadOnlyList<string> Phases { get; set; }
+
+			private TestErrorfingerprintEntry()
+			{
+				Key = String.Empty;
+				Phases = new List<string>();
+			}
+
+			public TestErrorfingerprintEntry(string key, IReadOnlyList<string> phases)
+			{
+				Key = key;
+				Phases = phases;
+			}
+		}
+
 		/// <summary>
 		/// Store the test session overall result
 		/// </summary>
@@ -226,9 +364,16 @@ namespace HordeServer.Jobs.TestData
 			[BsonRequired, BsonElement("m")]
 			public TestMetaId Metadata { get; set; }
 
+			[BsonIgnoreIfNull, BsonElement("r")]
+			public TestRecipeId? Recipe { get; set; }
+
 			[BsonIgnoreIfNull, BsonElement("tg")]
 			public IReadOnlyList<TestTagId>? Tags { get; set; }
 
+			[BsonIgnoreIfNull, BsonElement("ef")]
+			public List<TestErrorfingerprintEntry>? ErrorFingerprints { get; set; }
+			IReadOnlyList<IErrorFingerprintEntry>? ITestSession.ErrorFingerprints => ErrorFingerprints;
+
 			[BsonRequired, BsonElement("did")]
 			public ObjectId TestDataId { get; set; }
 
@@ -420,6 +565,7 @@ namespace HordeServer.Jobs.TestData
 
 		readonly IMongoCollection<TestDataDocument> _testDataDocuments;
 		readonly IMongoCollection<TestMetadataDocument> _testMeta;
+		readonly IMongoCollection<TestRecipeDocument> _testRecipes;
 		readonly IMongoCollection<TestTagDocument> _testTags;
 		readonly IMongoCollection<TestNameDocument> _tests;
 		readonly IMongoCollection<TestPhaseDocument> _testPhases;
@@ -447,6 +593,10 @@ namespace HordeServer.Jobs.TestData
 			metaIndexes.Add(keys => keys.Ascending("m.k").Ascending("m.v"));
 			_testMeta = mongoService.GetCollection<TestMetadataDocument>("TestDataV2.Metadata", metaIndexes);
 
+			List<MongoIndex<TestRecipeDocument>> recipeIndexes = new List<MongoIndex<TestRecipeDocument>>();
+			recipeIndexes.Add(keys => keys.Ascending(x => x.NameRef).Ascending(x => x.StreamId).Ascending(x => x.JobTemplateId).Ascending(x => x.JobStep), unique: true);
+			_testRecipes = mongoService.GetCollection<TestRecipeDocument>("TestDataV2.TestRecipes", recipeIndexes);
+
 			List<MongoIndex<TestTagDocument>> tagIndexes = new List<MongoIndex<TestTagDocument>>();
 			tagIndexes.Add(keys => keys.Ascending(x => x.Name), unique: true);
 			_testTags = mongoService.GetCollection<TestTagDocument>("TestDataV2.TestTags", tagIndexes);
@@ -606,6 +756,30 @@ namespace HordeServer.Jobs.TestData
 			return await _testTags.Find(filter).ToListAsync(cancellationToken);
 		}
 
+		/// <inheritdoc/>
+		public async Task<IReadOnlyList<ITestRecipeRef>> FindTestRecipesAsync(TestRecipeId[]? recipeIds = null, TestId[]? testIds = null, StreamId[]? streamIds = null, CancellationToken cancellationToken = default)
+		{
+			FilterDefinition<TestRecipeDocument> filter = FilterDefinition<TestRecipeDocument>.Empty;
+			FilterDefinitionBuilder<TestRecipeDocument> filterBuilder = Builders<TestRecipeDocument>.Filter;
+
+			if (recipeIds != null && recipeIds.Length > 0)
+			{
+				filter &= filterBuilder.In(x => x.Id, recipeIds);
+			}
+
+			if (testIds != null && testIds.Length > 0)
+			{
+				filter &= filterBuilder.In(x => x.NameRef, testIds);
+			}
+
+			if (streamIds != null && streamIds.Length > 0)
+			{
+				filter &= filterBuilder.In(x => x.StreamId, streamIds);
+			}
+
+			return await _testRecipes.Find(filter).ToListAsync(cancellationToken);
+		}
+
 		/// <inheritdoc/>
 		public async Task<IReadOnlyList<ITestPhaseSession>> FindTestPhaseSessionsAsync(StreamId[] streamIds, TestPhaseId[] phaseIds, DateTime? minCreateTime = null, DateTime? maxCreateTime = null, CommitId? minCommitId = null, CommitId? maxCommitId = null, CancellationToken cancellationToken = default)
 		{
@@ -719,6 +893,8 @@ namespace HordeServer.Jobs.TestData
 				public string DateTime { get; set; } = String.Empty;
 				[BsonRequired, BsonElement("timeElapseSec")]
 				public double TimeElapseSec { get; set; }
+				[BsonIgnoreIfNull, BsonElement("commandline")]
+				public string? CommandLine { get; set; }
 				[BsonRequired, BsonElement("phasesTotalCount")]
 				public int PhasesTotalCount { get; set; }
 				[BsonRequired, BsonElement("phasesSucceededCount")]
@@ -751,6 +927,30 @@ namespace HordeServer.Jobs.TestData
 				public List<string>? Tags { get; set; }
 			}
 
+			public class TestDescription
+			{
+				[BsonIgnoreIfNull, BsonElement("team")]
+				public string? Team { get; set; }
+				[BsonIgnoreIfNull, BsonElement("intent")]
+				public string? Intent { get; set; }
+				[BsonIgnoreIfNull, BsonElement("profile")]
+				public string? Profile { get; set; }
+				[BsonIgnoreIfNull, BsonElement("harness")]
+				public string? Harness { get; set; }
+			}
+
+			public class TestRecipe
+			{
+				[BsonRequired, BsonElement("nodeClass")]
+				public string NodeClass { get; set; } = String.Empty;
+				[BsonRequired, BsonElement("version")]
+				public string Version { get; set; } = String.Empty;
+				[BsonIgnoreIfNull, BsonElement("timeoutMinutes")]
+				public float? TimeoutMinutes { get; set; } = 0;
+				[BsonIgnoreIfNull, BsonElement("params")]
+				public Dictionary<string, object>? Params { get; set; } = new Dictionary<string, object>();
+			}
+
 			[BsonRequired, BsonElement("summary")]
 			public TestSessionSummary? Summary { get; set; }
 			[BsonRequired, BsonElement("phases")]
@@ -759,6 +959,10 @@ namespace HordeServer.Jobs.TestData
 			public Dictionary<string, string> Metadata { get; set; } = new Dictionary<string, string>();
 			[BsonIgnoreIfNull, BsonElement("tags")]
 			public List<string>? Tags { get; set; }
+			[BsonIgnoreIfNull, BsonElement("description")]
+			public TestDescription? Description { get; set; }
+			[BsonIgnoreIfNull, BsonElement("recipe")]
+			public TestRecipe? Recipe { get; set; }
 		}
 
 		private async Task AddTestReportDataAsync(IJob job, IJobStep step, IReadOnlyList<TestDataDocument> data, CancellationToken cancellationToken = default)
@@ -790,7 +994,14 @@ namespace HordeServer.Jobs.TestData
 
 				IReadOnlyList<ITestTagRef>? tags = await AddTestTagsAsync(testSession.Tags, cancellationToken);
 				ITestMetaRef metadata = await AddTestMetadataAsync(testSession.Metadata, cancellationToken);
-				ITestNameRef testRef = await AddTestAsync(document.Key, testSession.Summary.TestName, metadata.Id, tags?.Select(t => t.Id).ToList(), document.StreamId, cancellationToken);
+				ITestNameRef testRef = await AddTestAsync(document.Key, testSession.Summary.TestName, metadata.Id, tags?.Select(t => t.Id).ToList(), testSession.Description, document.StreamId, cancellationToken);
+
+				// test recipe
+				ITestRecipeRef? recipe = null;
+				if (testSession.Recipe != null)
+				{
+					recipe = await AddTestRecipeAsync(testRef.Id, testSession.Recipe, job, step, document.StreamId, cancellationToken);
+				}
 
 				// test session
 				TestSessionDocument sessionDoc = new TestSessionDocument(document.StreamId, document.Id, testRef.Id, metadata.Id);
@@ -815,6 +1026,10 @@ namespace HordeServer.Jobs.TestData
 				sessionDoc.Duration = TimeSpan.FromSeconds(testSession.Summary.TimeElapseSec);
 				// tags
 				sessionDoc.Tags = tags?.Select(t => t.Id).ToList();
+				// recipe
+				sessionDoc.Recipe = recipe?.Id;
+				// error fingerprints
+				Dictionary<string, List<string>> errorFingerprints = new();
 
 				// phase session
 				List<TestPhaseSessionDocument> phaseSessions = new();
@@ -840,6 +1055,22 @@ namespace HordeServer.Jobs.TestData
 					phaseDoc.Duration = TimeSpan.FromSeconds(phase.TimeElapseSec);
 					// tags
 					phaseDoc.Tags = tags?.Where(tag => phase.Tags?.Contains(tag.Name) ?? false).Select(tag => tag.Id).ToList();
+					// collect any error fingerprint
+					if (phase.ErrorFingerprint != null && errorFingerprints.Count < 50)
+					{
+						if (!errorFingerprints.ContainsKey(phase.ErrorFingerprint))
+						{
+							errorFingerprints.Add(phase.ErrorFingerprint, new List<string>() { phase.Key });
+						}
+						else
+						{
+							List<string>? phases = errorFingerprints.GetValueOrDefault(phase.ErrorFingerprint);
+							if (phases != null && phases.Count < 50)
+							{
+								phases.Add(phase.Key);
+							}
+						}
+					}
 
 					phaseSessions.Add(phaseDoc);
 				}
@@ -848,6 +1079,11 @@ namespace HordeServer.Jobs.TestData
 					await _testPhaseSessions.InsertManyAsync(phaseSessions, null, cancellationToken);
 				}
 
+				if (errorFingerprints.Count > 0)
+				{
+					sessionDoc.ErrorFingerprints = errorFingerprints.Select(E => new TestErrorfingerprintEntry(E.Key, E.Value)).ToList();
+				}
+
 				await _testSessions.InsertOneAsync(sessionDoc, null, cancellationToken);
 			}
 		}
@@ -884,6 +1120,52 @@ namespace HordeServer.Jobs.TestData
 			return meta;
 		}
 
+		private async Task<ITestRecipeRef?> AddTestRecipeAsync(TestId nameRef, AutomatedTestSessionData.TestRecipe recipe, IJob job, IJobStep jobStep, StreamId streamId, CancellationToken cancellationToken)
+		{
+			WorkflowId? workflow = jobStep.Annotations?.WorkflowId;
+			TemplateId jobTemplate = job.TemplateId;
+			string stepName = jobStep.Name;
+
+			FilterDefinition<TestRecipeDocument> filter = FilterDefinition<TestRecipeDocument>.Empty;
+			FilterDefinitionBuilder<TestRecipeDocument> filterBuilder = Builders<TestRecipeDocument>.Filter;
+
+			filter &= filterBuilder.Eq(i => i.NameRef, nameRef);
+			filter &= filterBuilder.Eq(i => i.StreamId, streamId);
+			filter &= filterBuilder.Eq(i => i.JobTemplateId, jobTemplate);
+			filter &= filterBuilder.Eq(i => i.JobStep, stepName);
+
+			TestRecipeDocument? recipeDoc = await _testRecipes.Find(filter).FirstOrDefaultAsync(cancellationToken);
+
+			if (recipeDoc == null)
+			{
+				recipeDoc = new TestRecipeDocument(nameRef, streamId, jobTemplate, stepName, recipe.NodeClass, recipe.Version, recipe.TimeoutMinutes ?? 0, workflow);
+				if (recipe.Params != null)
+				{
+					recipeDoc.UpdateParams(recipe.Params);
+				}
+				await _testRecipes.InsertOneAsync(recipeDoc, null, cancellationToken);
+			}
+			// update only once per day
+			else if (recipeDoc.LastSeenUtc < DateTime.UtcNow.AddDays(-1))
+			{
+				UpdateDefinitionBuilder<TestRecipeDocument> updateBuilder = Builders<TestRecipeDocument>.Update;
+				List<UpdateDefinition<TestRecipeDocument>> updates = new List<UpdateDefinition<TestRecipeDocument>>();
+				updates.Add(updateBuilder.Set(x => x.LastSeenUtc, DateTime.UtcNow));
+				updates.Add(updateBuilder.Set(x => x.NodeClass, recipe.NodeClass));
+				updates.Add(updateBuilder.Set(x => x.Version, recipe.Version));
+				updates.Add(updateBuilder.Set(x => x.TimeoutMinutes, recipe.TimeoutMinutes ?? 0));
+				updates.Add(updateBuilder.Set(x => x.WorkflowId, workflow));
+				if (recipe.Params != null)
+				{
+					recipeDoc.UpdateParams(recipe.Params);
+					updates.Add(updateBuilder.Set(x => x.Params, recipeDoc.Params));
+				}
+				await _testRecipes.FindOneAndUpdateAsync(x => x.Id == recipeDoc.Id, updateBuilder.Combine(updates), null, cancellationToken);
+			}
+
+			return recipeDoc;
+		}
+
 		private async Task<IReadOnlyList<ITestTagRef>?> AddTestTagsAsync(IReadOnlyList<string>? tags, CancellationToken cancellationToken)
 		{
 			if (tags == null)
@@ -919,7 +1201,7 @@ namespace HordeServer.Jobs.TestData
 			return storedTags;
 		}
 
-		private async Task<ITestNameRef> AddTestAsync(string key, string name, TestMetaId metadataId, IReadOnlyList<TestTagId>? tags, StreamId streamId, CancellationToken cancellationToken)
+		private async Task<ITestNameRef> AddTestAsync(string key, string name, TestMetaId metadataId, IReadOnlyList<TestTagId>? tags, AutomatedTestSessionData.TestDescription? description, StreamId streamId, CancellationToken cancellationToken)
 		{
 			FilterDefinitionBuilder<TestNameDocument> filterBuilder = Builders<TestNameDocument>.Filter;
 
@@ -928,6 +1210,10 @@ namespace HordeServer.Jobs.TestData
 			if (test == null)
 			{
 				test = new TestNameDocument(key, name);
+				if (description != null)
+				{
+					test.AddDescription(description.Team, description.Intent, description.Harness, description.Profile);
+				}
 				await _tests.InsertOneAsync(test, null, cancellationToken);
 			}
 			// update only once per day
@@ -937,6 +1223,13 @@ namespace HordeServer.Jobs.TestData
 				List<UpdateDefinition<TestNameDocument>> updates = new List<UpdateDefinition<TestNameDocument>>();
 				updates.Add(updateBuilder.Set(x => x.LastSeenUtc, DateTime.UtcNow));
 				updates.Add(updateBuilder.Set(x => x.Name, name));
+				if (description != null)
+				{
+					updates.Add(updateBuilder.Set(x => x.Team, description.Team));
+					updates.Add(updateBuilder.Set(x => x.Intent, description.Intent));
+					updates.Add(updateBuilder.Set(x => x.Harness, description.Harness));
+					updates.Add(updateBuilder.Set(x => x.Profile, description.Profile));
+				}
 				await _tests.FindOneAndUpdateAsync(x => x.Id == test.Id, updateBuilder.Combine(updates), null, cancellationToken);
 			}
 
@@ -1049,18 +1342,19 @@ namespace HordeServer.Jobs.TestData
 				return;
 			}
 
-			DateTime expireTime = DateTime.Now.AddMonths(-retainMonths);
+			DateTime expireTime = DateTime.UtcNow.AddMonths(-retainMonths);
 
 			// expire and prune root test data
 			long testsExpired = await ExpireCollectionAsync(_tests, expireTime, cancellationToken);
 			long phasesExpired = await ExpireCollectionAsync(_testPhases, expireTime, cancellationToken);
 			long metadataExpired = await ExpireCollectionAsync(_testMeta, expireTime, cancellationToken);
 			long tagsExpired = await ExpireCollectionAsync(_testTags, expireTime, cancellationToken);
+			long recipesExpired = await ExpireCollectionAsync(_testRecipes, expireTime, cancellationToken);
 
 			// Check whether we need to prune
-			if (testsExpired > 0 || phasesExpired > 0 || metadataExpired > 0 || tagsExpired > 0)
+			if (testsExpired > 0 || phasesExpired > 0 || metadataExpired > 0 || tagsExpired > 0 || recipesExpired > 0)
 			{
-				_logger.LogInformation("Expired {TestsExpired} tests, {PhasesExpired} suites, {MetadataExpired} meta, {TagsExpired} tags", testsExpired, phasesExpired, metadataExpired, tagsExpired);
+				_logger.LogInformation("Expired {TestsExpired} tests, {PhasesExpired} suites, {MetadataExpired} meta, {TagsExpired} tags, {RecipesExpired} recipes", testsExpired, phasesExpired, metadataExpired, tagsExpired, recipesExpired);
 				await PruneDataAsync(cancellationToken);
 			}
 
@@ -1110,26 +1404,37 @@ namespace HordeServer.Jobs.TestData
 		public async Task PruneDataAsync(CancellationToken cancellationToken = default)
 		{
 			HashSet<TestId> testIds = new HashSet<TestId>();
-			foreach (TestId testId in await (await _tests.DistinctAsync(x => x.Id, Builders<TestNameDocument>.Filter.Empty, null, cancellationToken)).ToListAsync(cancellationToken))
+			IAsyncCursor<TestId> distinctTests = await _tests.DistinctAsync(x => x.Id, Builders<TestNameDocument>.Filter.Empty, null, cancellationToken);
+			foreach (TestId testId in await distinctTests.ToListAsync(cancellationToken))
 			{
 				testIds.Add(testId);
 			}
 			_logger.LogInformation("Test data pruning to {TestIdCount} test name ids", testIds.Count);
 
 			HashSet<TestMetaId> metadataIds = new HashSet<TestMetaId>();
-			foreach (TestMetaId metaId in await (await _testMeta.DistinctAsync(x => x.Id, Builders<TestMetadataDocument>.Filter.Empty, null, cancellationToken)).ToListAsync(cancellationToken))
+			IAsyncCursor<TestMetaId> distinctMetadata = await _testMeta.DistinctAsync(x => x.Id, Builders<TestMetadataDocument>.Filter.Empty, null, cancellationToken);
+			foreach (TestMetaId metaId in await distinctMetadata.ToListAsync(cancellationToken))
 			{
 				metadataIds.Add(metaId);
 			}
 			_logger.LogInformation("Test data pruning to {MetadataIdCount} test metadata ids", metadataIds.Count);
 
 			HashSet<TestTagId> tagIds = new HashSet<TestTagId>();
-			foreach (TestTagId tagId in await (await _testTags.DistinctAsync(x => x.Id, Builders<TestTagDocument>.Filter.Empty, null, cancellationToken)).ToListAsync(cancellationToken))
+			IAsyncCursor<TestTagId> distinctTags = await _testTags.DistinctAsync(x => x.Id, Builders<TestTagDocument>.Filter.Empty, null, cancellationToken);
+			foreach (TestTagId tagId in await distinctTags.ToListAsync(cancellationToken))
 			{
 				tagIds.Add(tagId);
 			}
 			_logger.LogInformation("Test data pruning to {TagIdCount} test tag ids", tagIds.Count);
 
+			HashSet<TestRecipeId> recipeIds = new HashSet<TestRecipeId>();
+			IAsyncCursor<TestRecipeId> distinctRecipes = await _testRecipes.DistinctAsync(x => x.Id, Builders<TestRecipeDocument>.Filter.Empty, null, cancellationToken);
+			foreach (TestRecipeId recipeId in await distinctRecipes.ToListAsync(cancellationToken))
+			{
+				recipeIds.Add(recipeId);
+			}
+			_logger.LogInformation("Test data pruning to {RecipeIdCount} test recipe ids", recipeIds.Count);
+
 			int testStreamsDeleted = 0;
 
 			// update test session streams
diff --git a/Engine/Source/Programs/Horde/Plugins/Build/HordeServer.Build/Jobs/TestData/TestDataController.cs b/Engine/Source/Programs/Horde/Plugins/Build/HordeServer.Build/Jobs/TestData/TestDataController.cs
--- a/Engine/Source/Programs/Horde/Plugins/Build/HordeServer.Build/Jobs/TestData/TestDataController.cs
+++ b/Engine/Source/Programs/Horde/Plugins/Build/HordeServer.Build/Jobs/TestData/TestDataController.cs
@@ -1,5 +1,6 @@
 // Copyright Epic Games, Inc. All Rights Reserved.
 
+using System.IO;
 using System.Text.Json;
 using EpicGames.Core;
 using EpicGames.Horde.Commits;
@@ -644,7 +645,11 @@ namespace HordeServer.Jobs.TestData
 			{
 				Id = x.Id.ToString(),
 				Name = x.Name,
-				Key = x.Key
+				Key = x.Key,
+				Team = x.Team,
+				Intent = x.Intent,
+				Harness = x.Harness,
+				Profile = x.Profile,
 			}).ToList();
 		}
 
@@ -690,6 +695,39 @@ namespace HordeServer.Jobs.TestData
 			}).ToList();
 		}
 
+		/// <summary>
+		/// Get test recipes from provided ids
+		/// </summary>
+		/// <param name="request"></param>
+		/// <returns></returns>
+		[HttpPost]
+		[Route("/api/v3/testdata/recipes")]
+		[ProducesResponseType(typeof(List<GetTestRecipeResponse>), 200)]
+		public async Task<ActionResult<List<GetTestRecipeResponse>>> GetTestRecipeRefV2Async([FromBody] GetTestRecipesRequest request)
+		{
+			TestRecipeId[]? recipeIds = request.RecipeIds?.Select(x => TestRecipeId.Parse(x)).ToHashSet().ToArray();
+			TestId[]? testIds = request.TestIds?.Select(x => TestId.Parse(x)).ToHashSet().ToArray();
+			StreamId[]? streamIds = request.StreamIds?.Select(x => new StreamId(x)).ToHashSet().ToArray();
+
+			IReadOnlyList<ITestRecipeRef> testValues = await _testDataService.FindTestRecipesV2Async(recipeIds, testIds, streamIds);
+
+			return testValues.Where(
+				x => _buildConfig.Value.TryGetStream(x.StreamId, out StreamConfig? streamConfig) && streamConfig.Authorize(JobAclAction.ViewJob, User)
+			).Select(x => new GetTestRecipeResponse
+			{
+				Id = x.Id.ToString(),
+				NameRef = x.NameRef.ToString(),
+				StreamId = x.StreamId.ToString(),
+				JobTemplateId = x.JobTemplateId.ToString(),
+				JobStep = x.JobStep,
+				NodeClass = x.NodeClass,
+				Version = x.Version,
+				TimeoutMinutes = x.TimeoutMinutes,
+				WorkflowId = x.WorkflowId?.ToString(),
+				Params = x.Params?.ToDictionary(i => i.Key, i => i.Value)
+			}).ToList();
+		}
+
 		/// <summary>
 		/// Get test phases from provide test keys and refs
 		/// </summary>
@@ -789,7 +827,13 @@ namespace HordeServer.Jobs.TestData
 					PhasesTotalCount = testData.PhasesTotalCount,
 					PhasesSucceededCount = testData.PhasesSucceededCount,
 					PhasesUndefinedCount = testData.PhasesUndefinedCount,
-					PhasesFailedCount = testData.PhasesFailedCount
+					PhasesFailedCount = testData.PhasesFailedCount,
+					ErrorFingerprints = testData.ErrorFingerprints?.Select(
+						e => new GetErrorFingerprintResponse()
+						{
+							Key = e.Key,
+							Phases = e.Phases.ToList()
+						} ).ToList()
 				});
 			}
 
diff --git a/Engine/Source/Programs/Horde/Plugins/Build/HordeServer.Build/Jobs/TestData/TestDataService.cs b/Engine/Source/Programs/Horde/Plugins/Build/HordeServer.Build/Jobs/TestData/TestDataService.cs
--- a/Engine/Source/Programs/Horde/Plugins/Build/HordeServer.Build/Jobs/TestData/TestDataService.cs
+++ b/Engine/Source/Programs/Horde/Plugins/Build/HordeServer.Build/Jobs/TestData/TestDataService.cs
@@ -274,6 +274,19 @@ namespace HordeServer.Jobs.TestData
 			return await _testDataV2.FindTestTagsAsync(tagIds, names, cancellationToken);
 		}
 
+		/// <summary>
+		/// Find test recipe references
+		/// </summary>
+		/// <param name="recipeIds"></param>
+		/// <param name="testIds"></param>
+		/// <param name="streamIds"></param>
+		/// <param name="cancellationToken">Cancellation token for the operation</param>
+		/// <returns></returns>
+		public async Task<IReadOnlyList<ITestRecipeRef>> FindTestRecipesV2Async(TestRecipeId[]? recipeIds = null, TestId[]? testIds = null, StreamId[]? streamIds = null, CancellationToken cancellationToken = default)
+		{
+			return await _testDataV2.FindTestRecipesAsync(recipeIds, testIds, streamIds, cancellationToken);
+		}
+
 		/// <summary>
 		/// Find test phase sessions from a set of test phases
 		/// </summary>
diff --git a/Engine/Source/Programs/Horde/Plugins/Build/HordeServer.Build/Jobs/TestData/TestRecipeId.cs b/Engine/Source/Programs/Horde/Plugins/Build/HordeServer.Build/Jobs/TestData/TestRecipeId.cs
new file mode 100644
--- /dev/null
+++ b/Engine/Source/Programs/Horde/Plugins/Build/HordeServer.Build/Jobs/TestData/TestRecipeId.cs
@@ -0,0 +1,62 @@
+// Copyright Epic Games, Inc. All Rights Reserved.
+
+using System.ComponentModel;
+using EpicGames.Core;
+using HordeServer.Utilities;
+using MongoDB.Bson;
+
+namespace HordeServer.Jobs.TestData
+{
+	/// <summary>
+	/// Identifier for a test recipe
+	/// </summary>
+	/// <param name="Id">Id to construct from</param>
+	[JsonSchemaString]
+	[TypeConverter(typeof(ObjectIdTypeConverter<TestRecipeId, TestRecipeIdConverter>))]
+	[ObjectIdConverter(typeof(TestRecipeIdConverter))]
+	public record struct TestRecipeId(ObjectId Id)
+	{
+		/// <summary>
+		/// Default empty value for test recipe id
+		/// </summary>
+		public static TestRecipeId Empty { get; } = default;
+
+		/// <summary>
+		/// Creates a new <see cref="TestRecipeId"/>
+		/// </summary>
+		public static TestRecipeId GenerateNewId() => new TestRecipeId(ObjectId.GenerateNewId());
+
+		/// <inheritdoc cref="ObjectId.Parse(System.String)"/>
+		public static TestRecipeId Parse(string text) => new TestRecipeId(ObjectId.Parse(text));
+
+		/// <inheritdoc cref="ObjectId.TryParse(System.String, out ObjectId)"/>
+		public static bool TryParse(string text, out TestRecipeId id)
+		{
+			if (ObjectId.TryParse(text, out ObjectId objectId))
+			{
+				id = new TestRecipeId(objectId);
+				return true;
+			}
+			else
+			{
+				id = default;
+				return false;
+			}
+		}
+
+		/// <inheritdoc/>
+		public override string ToString() => Id.ToString();
+	}
+
+	/// <summary>
+	/// Converter to and from <see cref="ObjectId"/> instances.
+	/// </summary>
+	class TestRecipeIdConverter : ObjectIdConverter<TestRecipeId>
+	{
+		/// <inheritdoc/>
+		public override TestRecipeId FromObjectId(ObjectId id) => new TestRecipeId(id);
+
+		/// <inheritdoc/>
+		public override ObjectId ToObjectId(TestRecipeId value) => value.Id;
+	}
+}
diff --git a/Engine/Source/Programs/Shared/EpicGames.Horde/Jobs/TestData/TestDataMessages.cs b/Engine/Source/Programs/Shared/EpicGames.Horde/Jobs/TestData/TestDataMessages.cs
--- a/Engine/Source/Programs/Shared/EpicGames.Horde/Jobs/TestData/TestDataMessages.cs
+++ b/Engine/Source/Programs/Shared/EpicGames.Horde/Jobs/TestData/TestDataMessages.cs
@@ -331,6 +331,27 @@ namespace EpicGames.Horde.Jobs.TestData
 		public List<string>? TagNames { get; set; }
 	}
 
+	/// <summary>
+	/// Get test recipe request
+	/// </summary>
+	public class GetTestRecipesRequest
+	{
+		/// <summary>
+		/// Test recipe ids to get
+		/// </summary>
+		public List<string>? RecipeIds { get; set; }
+
+		/// <summary>
+		/// Test ids to get
+		/// </summary>
+		public List<string>? TestIds { get; set; }
+
+		/// <summary>
+		/// Stream ids to get
+		/// </summary>
+		public List<string>? StreamIds { get; set; }
+	}
+
 	/// <summary>
 	/// Get test phases from test keys and ids request
 	/// </summary>
@@ -602,6 +623,26 @@ namespace EpicGames.Horde.Jobs.TestData
 		/// The key of the test
 		/// </summary>
 		public string Key { get; set; } = String.Empty;
+
+		/// <summary>
+		/// The team associated with this test
+		/// </summary>
+		public string? Team { get; set; }
+
+		/// <summary>
+		/// The intent associated with this test
+		/// </summary>
+		public string? Intent { get; set; }
+
+		/// <summary>
+		/// The harness associated with this test
+		/// </summary>
+		public string? Harness { get; set; }
+
+		/// <summary>
+		/// The profile associated with this test
+		/// </summary>
+		public string? Profile { get; set; }
 	}
 
 	/// <summary>
@@ -636,6 +677,62 @@ namespace EpicGames.Horde.Jobs.TestData
 		public string Name { get; set; } = String.Empty;
 	}
 
+	/// <summary>
+	/// A test recipe ref
+	/// </summary>
+	public class GetTestRecipeResponse
+	{
+		/// <summary>
+		/// Recipe unique id
+		/// </summary>
+		public string Id { get; set; } = String.Empty;
+
+		/// <summary>
+		/// The test ref id associated with this recipe
+		/// </summary>
+		public string NameRef { get; set; } = String.Empty;
+
+		/// <summary>
+		/// The stream id associated with this recipe
+		/// </summary>
+		public string StreamId { get; set; } = String.Empty;
+
+		/// <summary>
+		/// The job template id associated with this recipe
+		/// </summary>
+		public string JobTemplateId { get; set; } = String.Empty;
+
+		/// <summary>
+		/// The job step name associated with this recipe
+		/// </summary>
+		public string JobStep { get; set; } = String.Empty;
+
+		/// <summary>
+		/// The test node class associated with this recipe
+		/// </summary>
+		public string NodeClass { get; set; } = String.Empty;
+
+		/// <summary>
+		/// The test version associated with this recipe
+		/// </summary>
+		public string Version { get; set; } = String.Empty;
+
+		/// <summary>
+		/// The test timeout in minutes associated with this recipe
+		/// </summary>
+		public float TimeoutMinutes { get; set; } = 0;
+
+		/// <summary>
+		/// The workflow id associated with this recipe
+		/// </summary>
+		public string? WorkflowId { get; set; }
+
+		/// <summary>
+		/// The test parameters associated with this recipe
+		/// </summary>
+		public Dictionary<string, object>? Params { get; set; }
+	}
+
 	/// <summary>
 	/// A test phase from a test
 	/// </summary>
@@ -798,6 +895,27 @@ namespace EpicGames.Horde.Jobs.TestData
 		/// The step job id associated with the test session
 		/// </summary>
 		public string StepId { get; set; } = String.Empty;
+
+		/// <summary>
+		/// The error fingerprints associated with the test session
+		/// </summary>
+		public List<GetErrorFingerprintResponse>? ErrorFingerprints { get; set; }
+	}
+
+	/// <summary>
+	/// Response object for error fingerprints
+	/// </summary>
+	public class GetErrorFingerprintResponse
+	{
+		/// <summary>
+		/// The error fingerprint key
+		/// </summary>
+		public string Key { get; set; } = String.Empty;
+
+		/// <summary>
+		/// The phase keys associated with the error fingerprint
+		/// </summary>
+		public List<string> Phases { get; set; } = new List<string>();
 	}
 
 	/// <summary>
