Fix warning

#rnx
#code-analyzed

diff --git a/Engine/Source/Programs/Horde/Plugins/Build/HordeServer.Build/Jobs/TestData/TestDataCollectionV2.cs b/Engine/Source/Programs/Horde/Plugins/Build/HordeServer.Build/Jobs/TestData/TestDataCollectionV2.cs
--- a/Engine/Source/Programs/Horde/Plugins/Build/HordeServer.Build/Jobs/TestData/TestDataCollectionV2.cs
+++ b/Engine/Source/Programs/Horde/Plugins/Build/HordeServer.Build/Jobs/TestData/TestDataCollectionV2.cs
@@ -1081,7 +1081,7 @@ namespace HordeServer.Jobs.TestData
 
 				if (errorFingerprints.Count > 0)
 				{
-					sessionDoc.ErrorFingerprints = errorFingerprints.Select(E => new TestErrorfingerprintEntry(E.Key, E.Value)).ToList();
+					sessionDoc.ErrorFingerprints = errorFingerprints.Select(e => new TestErrorfingerprintEntry(e.Key, e.Value)).ToList();
 				}
 
 				await _testSessions.InsertOneAsync(sessionDoc, null, cancellationToken);
