[Horde][BuildHealth][Server] Log info from FindOrAddIssueForSpanAsync to issue audit as well
#DisabledValidations CheckCopyright, CheckUnacceptableWords

diff --git a/Engine/Source/Programs/Horde/Plugins/Build/HordeServer.Build/Issues/IssueService.cs b/Engine/Source/Programs/Horde/Plugins/Build/HordeServer.Build/Issues/IssueService.cs
--- a/Engine/Source/Programs/Horde/Plugins/Build/HordeServer.Build/Issues/IssueService.cs
+++ b/Engine/Source/Programs/Horde/Plugins/Build/HordeServer.Build/Issues/IssueService.cs
@@ -1029,14 +1029,14 @@ namespace HordeServer.Issues
 				// Get the span data
 				NewIssueSpanData spanData = new NewIssueSpanData(streamConfig.Id, streamConfig.Name, job.TemplateId, node.Name, eventGroup.Fingerprint, stepData);
 
+				// References to events, and thus their audit loggers, are lost when we create spans, so we take first audit logger in event group and use that downstream. Doesn't matter which as they're all emitted to same issue audit log.
+				ILogger? groupAuditLogger = eventGroup.Events.FirstOrDefault(e => e.AuditLogger != null)?.AuditLogger;
+
 				IJobStepRef? prevJob = await _jobStepRefs.GetPrevStepForNodeAsync(job.StreamId, job.TemplateId, node.Name, job.CommitId, JobStepOutcome.Success, true);
 				if (prevJob != null)
 				{
-					// Find first audit logger in event group. Doesn't matter which as they're all emitted to same issue audit log.
-					ILogger? suspectAuditLogger = eventGroup.Events.FirstOrDefault(e => e.AuditLogger != null)?.AuditLogger;
-
 					spanData.LastSuccess = new NewIssueStepData(prevJob);
-					spanData.Suspects = await FindSuspectsForSpanAsync(streamConfig, spanData.Fingerprint, spanData.LastSuccess.CommitId, spanData.FirstFailure.CommitId, suspectAuditLogger, cancellationToken);
+					spanData.Suspects = await FindSuspectsForSpanAsync(streamConfig, spanData.Fingerprint, spanData.LastSuccess.CommitId, spanData.FirstFailure.CommitId, groupAuditLogger, cancellationToken);
 				}
 
 				IJobStepRef? nextJob = await _jobStepRefs.GetNextStepForNodeAsync(job.StreamId, job.TemplateId, node.Name, job.CommitId, JobStepOutcome.Success, true);
@@ -1046,7 +1046,7 @@ namespace HordeServer.Issues
 				}
 
 				// Add all the new objects
-				IIssue newIssue = await FindOrAddIssueForSpanAsync(spanData, cancellationToken);
+				IIssue newIssue = await FindOrAddIssueForSpanAsync(spanData, groupAuditLogger, cancellationToken);
 				IIssueSpan newSpan = await _issueCollection.AddSpanAsync(newIssue.Id, spanData, cancellationToken);
 				IIssueStep newStep = await _issueCollection.AddStepAsync(newSpan.Id, stepData, cancellationToken);
 
@@ -1590,21 +1590,22 @@ namespace HordeServer.Issues
 		/// Find an existing issue that may match one of the given suspect changes
 		/// </summary>
 		/// <param name="span">The span to find an issue for</param>
+		/// <param name="issueLogger">The buffered logger used to write to the issue's eventual audit log</param>
 		/// <param name="cancellationToken">Cancellation token for the operation</param>
 		/// <returns>The matching issue</returns>
-		async Task<IIssue> FindOrAddIssueForSpanAsync(NewIssueSpanData span, CancellationToken cancellationToken)
+		async Task<IIssue> FindOrAddIssueForSpanAsync(NewIssueSpanData span, ILogger? issueLogger, CancellationToken cancellationToken)
 		{
 			using TelemetrySpan tracerSpan = _tracer.StartActiveSpan($"{nameof(IssueService)}.{nameof(FindOrAddIssueForSpanAsync)}.NewIssueSpan");
 			IReadOnlyList<IIssue> existingIssues;
 			if (span.LastSuccess == null)
 			{
 				existingIssues = await _issueCollection.FindIssuesAsync(streamId: span.StreamId, minCommitId: span.FirstFailure.CommitId, maxCommitId: span.FirstFailure.CommitId, resolved: false, cancellationToken: cancellationToken);
-				_logger.LogDebug("Found {NumIssues} open issues at {Change} in {StreamId}", existingIssues.Count, span.FirstFailure.CommitId, span.StreamId);
+				LoggerUtils.MultiLog([(_logger, LogLevel.Debug), (issueLogger, LogLevel.Debug)], "Found {NumIssues} open issues at {Change} in {StreamId}", existingIssues.Count, span.FirstFailure.CommitId, span.StreamId);
 			}
 			else
 			{
 				existingIssues = await _issueCollection.FindIssuesForChangesAsync(span.Suspects.ConvertAll(x => x.SourceCommitId ?? x.CommitId), cancellationToken);
-				_logger.LogDebug("Found {NumIssues} open issues in {StreamId} from [{ChangeList}]", existingIssues.Count, span.StreamId, String.Join(", ", span.Suspects.ConvertAll(x => (x.SourceCommitId ?? x.CommitId).ToString())));
+				LoggerUtils.MultiLog([(_logger, LogLevel.Debug), (issueLogger, LogLevel.Debug)], "Found {NumIssues} open issues in {StreamId} from [{ChangeList}]", existingIssues.Count, span.StreamId, String.Join(", ", span.Suspects.ConvertAll(x => (x.SourceCommitId ?? x.CommitId).ToString())));
 			}
 
 			IIssue? issue = existingIssues.FirstOrDefault(x => x.VerifiedAt == null && x.Fingerprints.Any(y => y.IsMatch(span.Fingerprint)));
@@ -1612,11 +1613,11 @@ namespace HordeServer.Issues
 			{
 				string summary = GetSummary(span.Fingerprint, span.FirstFailure.Severity);
 				issue = await _issueCollection.AddIssueAsync(summary, cancellationToken);
-				_logger.LogInformation("Created issue {IssueId}", issue.Id);
+				LoggerUtils.MultiLog([(_logger, LogLevel.Information), (issueLogger, LogLevel.Information)], "Created issue {IssueId}", issue.Id);
 			}
 			else
 			{
-				_logger.LogInformation("Matched to issue {IssueId}", issue.Id);
+				LoggerUtils.MultiLog([(_logger, LogLevel.Information), (issueLogger, LogLevel.Information)], "Matched to issue {IssueId}", issue.Id);
 			}
 			return issue;
 		}
