[Horde][Analytics][Server] Updated Label telemetry to emit individual events in addition to summary.

diff --git a/Engine/Source/Programs/Horde/Plugins/Analytics/HordeServer.Analytics.Integration.Tests/Telemetry/TelemetryIntegrationTests.cs b/Engine/Source/Programs/Horde/Plugins/Analytics/HordeServer.Analytics.Integration.Tests/Telemetry/TelemetryIntegrationTests.cs
--- a/Engine/Source/Programs/Horde/Plugins/Analytics/HordeServer.Analytics.Integration.Tests/Telemetry/TelemetryIntegrationTests.cs
+++ b/Engine/Source/Programs/Horde/Plugins/Analytics/HordeServer.Analytics.Integration.Tests/Telemetry/TelemetryIntegrationTests.cs
@@ -475,13 +475,20 @@ namespace HordeServer.Analytics.Integration.Tests.Telemetry
 				// Assert we actually have the expected telemetry type.
 				Hashtable actualStateJobSummaryPayload = actualTelemetryPayloads.Where(x => x.Item2[EventNameKeyStr] as string == JobLabelSummaryTelemetry.DefaultEventName).First().Item2;
 				JobLabelSummaryTelemetry? actualJobSummaryTelemetry = actualStateJobSummaryPayload.MapToObject<JobLabelSummaryTelemetry>();
-
 				Assert.IsNotNull(actualJobSummaryTelemetry);
 
 				Assert.IsTrue(actualJobSummaryTelemetry.SuccessfulLabels.Contains(ToolWin64Label));
 				Assert.IsTrue(actualJobSummaryTelemetry.SuccessfulLabels.Contains(AllWin64Label));
 				Assert.IsTrue(actualJobSummaryTelemetry.SuccessfulLabels.Contains(AllToolLabel));
 				Assert.IsFalse(actualJobSummaryTelemetry.FailedLabels.Any());
+
+				// Assert individualized label telemetry
+				IEnumerable<Hashtable> payloads = actualTelemetryPayloads.Where(x => x.Item2[EventNameKeyStr] as string == JobLabelTelemetry.DefaultEventName).Select(x => x.Item2);
+				Assert.AreEqual(3, payloads.Count());
+
+				Hashtable actualStateJobLabelPayload = payloads.First();
+				JobLabelTelemetry? actualJobLabelPayload = actualStateJobLabelPayload.MapToObject<JobLabelTelemetry>();
+				Assert.IsNotNull(actualJobLabelPayload);
 			}
 
 			// Label failures - cross batch
@@ -732,8 +739,8 @@ namespace HordeServer.Analytics.Integration.Tests.Telemetry
 				Assert.IsNotNull(actualJobSummaryTelemetry);
 
 				// Assert that the category prefixed label is present.
-				Assert.IsTrue(actualJobSummaryTelemetry.SuccessfulLabels.Contains(JobLabelSummaryTelemetry.CreateFullyQualifiedLabel(Win64Category, ToolWin64Label)));
-				Assert.IsTrue(actualJobSummaryTelemetry.SuccessfulLabels.Contains(JobLabelSummaryTelemetry.CreateFullyQualifiedLabel(Win64Category, AllWin64Label)));
+				Assert.IsTrue(actualJobSummaryTelemetry.SuccessfulLabels.Contains(JobLabelTelemetryHelpers.CreateFullyQualifiedLabel(Win64Category, ToolWin64Label)));
+				Assert.IsTrue(actualJobSummaryTelemetry.SuccessfulLabels.Contains(JobLabelTelemetryHelpers.CreateFullyQualifiedLabel(Win64Category, AllWin64Label)));
 				Assert.IsTrue(actualJobSummaryTelemetry.SuccessfulCategories.Contains(Win64Category));
 			}
 
@@ -796,7 +803,7 @@ namespace HordeServer.Analytics.Integration.Tests.Telemetry
 				Assert.IsNotNull(actualJobSummaryTelemetry);
 
 				// Assert that the category prefixed label is present.
-				Assert.IsTrue(actualJobSummaryTelemetry.FailedLabels.Contains(JobLabelSummaryTelemetry.CreateFullyQualifiedLabel(Game1Category, AllGameLabel)));
+				Assert.IsTrue(actualJobSummaryTelemetry.FailedLabels.Contains(JobLabelTelemetryHelpers.CreateFullyQualifiedLabel(Game1Category, AllGameLabel)));
 				Assert.IsTrue(actualJobSummaryTelemetry.FailedCategories.Contains(Game1Category));
 			}
 
@@ -916,6 +923,7 @@ namespace HordeServer.Analytics.Integration.Tests.Telemetry
 
 				// Assert we actually have the expected telemetry type.
 				Hashtable actualStateJobSummaryPayload = actualTelemetryPayloads.Where(x => x.Item2[EventNameKeyStr] as string == JobLabelSummaryTelemetry.DefaultEventName).First().Item2;
+
 				JobLabelSummaryTelemetry? actualJobSummaryTelemetry = actualStateJobSummaryPayload.MapToObject<JobLabelSummaryTelemetry>();
 
 				Assert.IsNotNull(actualJobSummaryTelemetry);
diff --git a/Engine/Source/Programs/Horde/Plugins/Build/HordeServer.Build/Jobs/JobService.cs b/Engine/Source/Programs/Horde/Plugins/Build/HordeServer.Build/Jobs/JobService.cs
--- a/Engine/Source/Programs/Horde/Plugins/Build/HordeServer.Build/Jobs/JobService.cs
+++ b/Engine/Source/Programs/Horde/Plugins/Build/HordeServer.Build/Jobs/JobService.cs
@@ -1059,7 +1059,7 @@ namespace HordeServer.Jobs
 									try
 									{
 										_telemetryWriter.WriteEvent(streamConfig.TelemetryStoreId, new JobSummaryTelemetry(job, completingJobStep.FinishTimeUtc.Value));
-										_telemetryWriter.WriteEvent(streamConfig.TelemetryStoreId, new JobLabelSummaryTelemetry(job, graph));
+										JobLabelTelemetryHelpers.EmitJobLabelTelemetry(_telemetryWriter, streamConfig.TelemetryStoreId, job, graph);
 									}
 									catch (Exception ex)
 									{
diff --git a/Engine/Source/Programs/Horde/Plugins/Build/HordeServer.Build/Jobs/JobTelemetry.cs b/Engine/Source/Programs/Horde/Plugins/Build/HordeServer.Build/Jobs/JobTelemetry.cs
--- a/Engine/Source/Programs/Horde/Plugins/Build/HordeServer.Build/Jobs/JobTelemetry.cs
+++ b/Engine/Source/Programs/Horde/Plugins/Build/HordeServer.Build/Jobs/JobTelemetry.cs
@@ -4,6 +4,8 @@ using EpicGames.Horde.Jobs;
 using EpicGames.Horde.Jobs.Graphs;
 using EpicGames.Horde.Jobs.Templates;
 using EpicGames.Horde.Streams;
+using EpicGames.Horde.Telemetry;
+using HordeServer.Telemetry;
 
 namespace HordeServer.Jobs
 {
@@ -149,6 +151,80 @@ namespace HordeServer.Jobs
 		}
 	}
 
+	/// <summary>
+	/// Helpers associated with producing <see cref="JobLabelTelemetry"/> and <see cref="JobLabelSummaryTelemetry"/>.
+	/// </summary>
+	public static class JobLabelTelemetryHelpers
+	{
+		/// <summary>
+		/// Delimiter used in separating the Label Category and the Label Name.
+		/// </summary>
+		public const string FullyQualifiedLabelDelimiter = "-";
+
+		/// <summary>
+		/// Creates a fully qualified label name.
+		/// </summary>
+		/// <param name="response">The label response to convert to a fully qualified label name.</param>
+		/// <returns>The fully qualified label name.</returns>
+		/// <remarks>If there is no valid category, the label name will simply be the <see cref="GetLabelStateResponse.DashboardName"/>. If no valid dashboard name, the function will return <see cref="String.Empty"/>.</remarks>
+		public static string CreateFullyQualifiedLabel(GetLabelStateResponse response)
+		{
+			return CreateFullyQualifiedLabel(response.DashboardCategory, response.DashboardName);
+		}
+
+		/// <summary>
+		/// Creates a fully qualified label name.
+		/// </summary>
+		/// <param name="dashboardCategory">The category of the label.</param>
+		/// <param name="dashboardName">The name of the label.</param>
+		/// <returns>The fully qualified label name.</returns>
+		/// <remarks>If there is no valid category, the label name will simply be the <see cref="GetLabelStateResponse.DashboardName"/>. If no valid dashboard name, the function will return <see cref="String.Empty"/>.</remarks>
+		public static string CreateFullyQualifiedLabel(string? dashboardCategory, string? dashboardName)
+		{
+			return (dashboardCategory != null && dashboardName != null) ? $"{dashboardCategory}{FullyQualifiedLabelDelimiter}{dashboardName}" : (dashboardName != null ? dashboardName : String.Empty);
+		}
+
+		/// <summary>
+		/// Creates a list of job label telemetry events for a given job and graph.
+		/// </summary>
+		/// <param name="job">The job to base the job label telemetry events on.</param>
+		/// <param name="graph">The graph to base the job label telemetry events on.</param>
+		/// <returns>A list of resulting job label telemetry events.</returns>
+		public static IList<JobLabelTelemetry> CreateJobLabelTelemetryEvents(IJob job, IGraph graph)
+		{
+			List<GetLabelStateResponse> responses = [];
+			job.GetLabelStateResponses(graph, responses);
+
+			IList<JobLabelTelemetry> telemetryEntries = new List<JobLabelTelemetry>(responses.Count);
+
+			for (int i = 0; i < responses.Count; ++i)
+			{
+				GetLabelStateResponse response = responses[i];
+				telemetryEntries.Add(new JobLabelTelemetry(job, graph, response.DashboardCategory, response.DashboardName, response.Outcome, response.State));
+			}
+
+			return telemetryEntries;
+		}
+
+		/// <summary>
+		/// Emits all job label telemetry for a given job and graph.
+		/// </summary>
+		/// <param name="telemetryWriter">The telemetry writer interface to use for emission.</param>
+		/// <param name="telemetryStoreId">The telemetry store to use for emission.</param>
+		/// <param name="job">The job to base the job label telemetry events on.</param>
+		/// <param name="graph">The graph to base the job label telemetry events on.</param>
+		public static void EmitJobLabelTelemetry(ITelemetryWriter telemetryWriter, TelemetryStoreId telemetryStoreId, IJob job, IGraph graph)
+		{
+			telemetryWriter.WriteEvent(telemetryStoreId, new JobLabelSummaryTelemetry(job, graph));
+			IList<JobLabelTelemetry> jobLabelTelemetries = CreateJobLabelTelemetryEvents(job, graph);
+
+			for (int i = 0; i < jobLabelTelemetries.Count; ++i)
+			{
+				telemetryWriter.WriteEvent(telemetryStoreId, jobLabelTelemetries[i]);
+			}
+		}
+	}
+
 	/// <summary>
 	/// Record used for Job Label Summary telemetry. Models a <see cref="IJob"/> and <see cref="IGraph"/>.
 	/// </summary>
@@ -251,9 +327,9 @@ namespace HordeServer.Jobs
 			List<GetLabelStateResponse> responses = [];
 			job.GetLabelStateResponses(graph, responses);
 
-			SuccessfulLabels = [.. responses.Where(r => r.State == LabelState.Complete && r.Outcome == LabelOutcome.Success && r.DashboardName != null).Select(r => CreateFullyQualifiedLabel(r))];
-			WarningLabels = [.. responses.Where(r => r.State == LabelState.Complete && r.Outcome == LabelOutcome.Warnings && r.DashboardName != null).Select(r => CreateFullyQualifiedLabel(r))];
-			FailedLabels = [.. responses.Where(r => r.State == LabelState.Complete && r.Outcome == LabelOutcome.Failure && r.DashboardName != null).Select(r => CreateFullyQualifiedLabel(r))];
+			SuccessfulLabels = [.. responses.Where(r => r.State == LabelState.Complete && r.Outcome == LabelOutcome.Success && r.DashboardName != null).Select(r => JobLabelTelemetryHelpers.CreateFullyQualifiedLabel(r))];
+			WarningLabels = [.. responses.Where(r => r.State == LabelState.Complete && r.Outcome == LabelOutcome.Warnings && r.DashboardName != null).Select(r => JobLabelTelemetryHelpers.CreateFullyQualifiedLabel(r))];
+			FailedLabels = [.. responses.Where(r => r.State == LabelState.Complete && r.Outcome == LabelOutcome.Failure && r.DashboardName != null).Select(r => JobLabelTelemetryHelpers.CreateFullyQualifiedLabel(r))];
 
 			Dictionary<string, LabelOutcome> categoryToOutcome = [];
 
@@ -284,33 +360,100 @@ namespace HordeServer.Jobs
 			WarningCategories = [.. flattenedCategoryList.Where(r => r.Value == LabelOutcome.Warnings).Select(r => r.Key)];
 			FailedCategories = [.. flattenedCategoryList.Where(r => r.Value == LabelOutcome.Failure).Select(r => r.Key)];
 		}
+	}
 
+	/// <summary>
+	/// Record used for Job Label telemetry. Models a <see cref="IJob"/> and <see cref="IGraph"/>.
+	/// </summary>
+	public record JobLabelTelemetry
+	{
 		/// <summary>
-		/// Delimiter used in separating the Label Category and the Label Name.
+		/// The telemetry event name.
 		/// </summary>
-		public const string FullyQualifiedLabelDelimiter = "-";
+		public string EventName { get; init; }
 
 		/// <summary>
-		/// Creates a fully qualified label name.
+		/// The schema version of the telemetry.
 		/// </summary>
-		/// <param name="response">The label response to convert to a fully qualified label name.</param>
-		/// <returns>The fully qualified label name.</returns>
-		/// <remarks>If there is no valid category, the label name will simply be the <see cref="GetLabelStateResponse.DashboardName"/>. If no valid dashboard name, the function will return <see cref="String.Empty"/>.</remarks>
-		public static string CreateFullyQualifiedLabel(GetLabelStateResponse response)
-		{
-			return CreateFullyQualifiedLabel(response.DashboardCategory, response.DashboardName);
-		}
+		public int SchemaVersion { get; init; }
 
 		/// <summary>
-		/// Creates a fully qualified label name.
+		/// The job Id.
 		/// </summary>
-		/// <param name="dashboardCategory">The category of the label.</param>
-		/// <param name="dashboardName">The name of the label.</param>
-		/// <returns>The fully qualified label name.</returns>
-		/// <remarks>If there is no valid category, the label name will simply be the <see cref="GetLabelStateResponse.DashboardName"/>. If no valid dashboard name, the function will return <see cref="String.Empty"/>.</remarks>
-		public static string CreateFullyQualifiedLabel(string? dashboardCategory, string? dashboardName)
+		public JobId JobId { get; init; }
+
+		/// <summary>
+		/// The graph content hash.
+		/// </summary>
+		public string GraphContentHash { get; init; }
+
+		/// <summary>
+		/// The template of the job.
+		/// </summary>
+		public TemplateId TemplateId { get; init; }
+
+		/// <summary>
+		/// The stream of the job.
+		/// </summary>
+		public StreamId StreamId { get; init; }
+
+		/// <summary>
+		/// The commit of the job. <see cref="IJob.CommitId"/>.
+		/// </summary>
+		public string CommitId { get; init; }
+
+		/// <summary>
+		/// The category of the label.
+		/// </summary>
+		public string? Category { get; init; }
+
+		/// <summary>
+		/// The label.
+		/// </summary>
+		public string? Label { get; init; }
+
+		/// <summary>
+		/// The outcome of the label.
+		/// </summary>
+		public LabelOutcome? Outcome { get; init; }
+
+		/// <summary>
+		/// The state of the label.
+		/// </summary>
+		public LabelState? State { get; init; }
+
+		/// <summary>
+		/// Default event name for the JobLabelTelemetry.
+		/// </summary>
+		public const string DefaultEventName = "State.JobLabel";
+
+		/// <summary>
+		/// Current schema version.
+		/// </summary>
+		public static readonly int CurrentSchemaVersion = 1;
+
+		/// <summary>
+		/// Constructor.
+		/// </summary>
+		/// <param name="job">The job to base the job label telemetry events on.</param>
+		/// <param name="graph">The graph to base the job label telemetry events on.</param>
+		/// <param name="category">The category of the label.</param>
+		/// <param name="name">The name of the label.</param>
+		/// <param name="outcome">The outcome of the label.</param>
+		/// <param name="state">The state of the label.</param>
+		public JobLabelTelemetry(IJob job, IGraph graph, string? category, string? name, LabelOutcome? outcome, LabelState? state)
 		{
-			return (dashboardCategory != null && dashboardName != null) ? $"{dashboardCategory}{FullyQualifiedLabelDelimiter}{dashboardName}" : (dashboardName != null ? dashboardName : String.Empty);
+			EventName = DefaultEventName;
+			SchemaVersion = CurrentSchemaVersion;
+			GraphContentHash = graph.Id.ToString();
+			StreamId = job.StreamId;
+			TemplateId = job.TemplateId;
+			JobId = job.Id;
+			CommitId = job.CommitId?.Name ?? "0";
+			Category = category;
+			Label = name;
+			Outcome = outcome;
+			State = state;
 		}
 	}
 }
\ No newline at end of file
