Horde: Fix warning in server build
#rnx

diff --git a/Engine/Source/Programs/Horde/Plugins/Build/HordeServer.Build/BuildConfig.cs b/Engine/Source/Programs/Horde/Plugins/Build/HordeServer.Build/BuildConfig.cs
--- a/Engine/Source/Programs/Horde/Plugins/Build/HordeServer.Build/BuildConfig.cs
+++ b/Engine/Source/Programs/Horde/Plugins/Build/HordeServer.Build/BuildConfig.cs
@@ -101,6 +101,7 @@ namespace HordeServer
 		[JsonIgnore]
 		public IReadOnlyDictionary<string, string> IssueChannelGroups { get; private set; } = null!;
 
+		/// <summary>
 		/// The maximum days since node success to scan for suspects
 		/// </summary>
 		public int IssueFindSuspectsMaxSuccessDays { get; set; } = 14;
