Horde: Refactor issue config to be a root object

diff --git a/Engine/Source/Programs/Horde/Plugins/Build/HordeServer.Build.Tests/Issues/IssueServiceTests.cs b/Engine/Source/Programs/Horde/Plugins/Build/HordeServer.Build.Tests/Issues/IssueServiceTests.cs
--- a/Engine/Source/Programs/Horde/Plugins/Build/HordeServer.Build.Tests/Issues/IssueServiceTests.cs
+++ b/Engine/Source/Programs/Horde/Plugins/Build/HordeServer.Build.Tests/Issues/IssueServiceTests.cs
@@ -163,10 +163,10 @@ namespace HordeServer.Tests.Issues
 
 			BuildConfig buildConfig = new BuildConfig();
 			buildConfig.Projects.Add(projectConfig);
-			buildConfig.Issues = new List<IssueConfig> { new IssueConfig { ChannelGroups = new List<IssueChannelGroupConfig> { 
-			new IssueChannelGroupConfig { Channels = new List<string> { "AssetCheck" } },
-			new IssueChannelGroupConfig { Channels = new List<string> { "LogLinker", "LogBlueprint" } },
-			} } };
+			buildConfig.Issues = new IssueConfig { ChannelGroups = new List<IssueChannelGroupConfig> { 
+				new IssueChannelGroupConfig { Channels = new List<string> { "AssetCheck" } },
+				new IssueChannelGroupConfig { Channels = new List<string> { "LogLinker", "LogBlueprint" } },
+			} };
 
 			GlobalConfig globalConfig = new GlobalConfig();
 			globalConfig.Plugins.AddBuildConfig(buildConfig);
diff --git a/Engine/Source/Programs/Horde/Plugins/Build/HordeServer.Build/BuildConfig.cs b/Engine/Source/Programs/Horde/Plugins/Build/HordeServer.Build/BuildConfig.cs
--- a/Engine/Source/Programs/Horde/Plugins/Build/HordeServer.Build/BuildConfig.cs
+++ b/Engine/Source/Programs/Horde/Plugins/Build/HordeServer.Build/BuildConfig.cs
@@ -91,9 +91,9 @@ namespace HordeServer
 		public string IssueFixedTag { get; set; } = "#horde";
 
 		/// <summary>
-		/// List of build health issue configurations
+		/// Build health issue configurations
 		/// </summary>
-		public List<IssueConfig> Issues { get; set; } = new List<IssueConfig>();
+		public IssueConfig? Issues { get; set; }
 
 		/// <summary>
 		/// Issue channel groupings
@@ -124,28 +124,25 @@ namespace HordeServer
 			IssueChannelGroups = issueChannelGroups;
 
 			// build up the issue channel groupings
-			Issues.ForEach(ic =>
+			Issues?.PostLoad();
+
+			Issues?.ChannelGroups?.ForEach(group =>
 			{
-				ic.PostLoad();
+				if (group.Channels == null || group.Channels.Count == 0)
+				{
+					return;
+				}
 
-				ic.ChannelGroups?.ForEach(group =>
+				string channelGroup = String.Join("_", group.Channels);
+
+				group.Channels.ForEach(channel =>
 				{
-					if (group.Channels == null || group.Channels.Count == 0)
+					if (String.IsNullOrEmpty(channel) || issueChannelGroups.ContainsKey(channel))
 					{
 						return;
 					}
 
-					string channelGroup = String.Join("_", group.Channels);
-
-					group.Channels.ForEach(channel =>
-					{
-						if (String.IsNullOrEmpty(channel) || issueChannelGroups.ContainsKey(channel))
-						{
-							return;
-						}
-
-						issueChannelGroups[channel] = channelGroup;
-					});
+					issueChannelGroups[channel] = channelGroup;
 				});
 			});
 
