Horde: Temporarily rename issues to work around redis cache
#rnx

diff --git a/Engine/Source/Programs/Horde/Plugins/Build/HordeServer.Build.Tests/Issues/IssueServiceTests.cs b/Engine/Source/Programs/Horde/Plugins/Build/HordeServer.Build.Tests/Issues/IssueServiceTests.cs
--- a/Engine/Source/Programs/Horde/Plugins/Build/HordeServer.Build.Tests/Issues/IssueServiceTests.cs
+++ b/Engine/Source/Programs/Horde/Plugins/Build/HordeServer.Build.Tests/Issues/IssueServiceTests.cs
@@ -163,7 +163,7 @@ namespace HordeServer.Tests.Issues
 
 			BuildConfig buildConfig = new BuildConfig();
 			buildConfig.Projects.Add(projectConfig);
-			buildConfig.Issues = new IssueConfig { ChannelGroups = new List<IssueChannelGroupConfig> { 
+			buildConfig.IssuesBug = new IssueConfig { ChannelGroups = new List<IssueChannelGroupConfig> { 
 				new IssueChannelGroupConfig { Channels = new List<string> { "AssetCheck" } },
 				new IssueChannelGroupConfig { Channels = new List<string> { "LogLinker", "LogBlueprint" } },
 			} };
diff --git a/Engine/Source/Programs/Horde/Plugins/Build/HordeServer.Build/BuildConfig.cs b/Engine/Source/Programs/Horde/Plugins/Build/HordeServer.Build/BuildConfig.cs
--- a/Engine/Source/Programs/Horde/Plugins/Build/HordeServer.Build/BuildConfig.cs
+++ b/Engine/Source/Programs/Horde/Plugins/Build/HordeServer.Build/BuildConfig.cs
@@ -93,7 +93,7 @@ namespace HordeServer
 		/// <summary>
 		/// Build health issue configurations
 		/// </summary>
-		public IssueConfig? Issues { get; set; }
+		public IssueConfig? IssuesBug { get; set; }
 
 		/// <summary>
 		/// Issue channel groupings
@@ -124,9 +124,9 @@ namespace HordeServer
 			IssueChannelGroups = issueChannelGroups;
 
 			// build up the issue channel groupings
-			Issues?.PostLoad();
+			IssuesBug?.PostLoad();
 
-			Issues?.ChannelGroups?.ForEach(group =>
+			IssuesBug?.ChannelGroups?.ForEach(group =>
 			{
 				if (group.Channels == null || group.Channels.Count == 0)
 				{
