Horde: Add include support to issue config
#rnx

diff --git a/Engine/Source/Programs/Horde/Plugins/Build/HordeServer.Build/Issues/IssueConfig.cs b/Engine/Source/Programs/Horde/Plugins/Build/HordeServer.Build/Issues/IssueConfig.cs
--- a/Engine/Source/Programs/Horde/Plugins/Build/HordeServer.Build/Issues/IssueConfig.cs
+++ b/Engine/Source/Programs/Horde/Plugins/Build/HordeServer.Build/Issues/IssueConfig.cs
@@ -28,6 +28,11 @@ namespace HordeServer.Issues
 	[ConfigMacroScope]
 	public class IssueConfig
 	{
+		/// <summary>
+		/// Includes for other configuration files
+		/// </summary>
+		public List<ConfigInclude> Include { get; set; } = new List<ConfigInclude>();
+
 		/// <summary>
 		/// Issue grouping channels
 		/// </summary>
