Horde: Updating schema documentation

#rnx

diff --git a/Engine/Source/Programs/Horde/Docs/Config/Schema/AclActions.md b/Engine/Source/Programs/Horde/Docs/Config/Schema/AclActions.md
--- a/Engine/Source/Programs/Horde/Docs/Config/Schema/AclActions.md
+++ b/Engine/Source/Programs/Horde/Docs/Config/Schema/AclActions.md
@@ -588,6 +588,12 @@
 | `UpdateReplicator` |  |
 | `ViewReplicator` |  |
 
+## Secrets
+
+| Name | Description |
+| ---- | ----------- |
+| `ViewSecret` | View a credential |
+
 ## ServiceAccounts
 
 | Name | Description |
@@ -750,7 +756,7 @@
 
 | Name | Description |
 | ---- | ----------- |
-| `ViewSecret` | View a credential |
+| `ViewSecret` |  |
 
 ## ServiceAccounts
 
@@ -940,19 +946,6 @@
 | `DownloadSoftware` |  |
 | `DeleteSoftware` |  |
 
-## Storage
-
-| Name | Description |
-| ---- | ----------- |
-| `ReadBlobs` | Ability to read blobs from the storage service |
-| `WriteBlobs` | Ability to write blobs to the storage service |
-| `ReadAliases` | Ability to read aliases from the storage service |
-| `WriteAliases` | Ability to write aliases to the storage service |
-| `DeleteAliases` | Ability to write aliases to the storage service |
-| `ReadRefs` | Ability to read refs from the storage service |
-| `WriteRefs` | Ability to write refs to the storage service |
-| `DeleteRefs` | Ability to delete refs |
-
 ## Streams
 
 | Name | Description |
@@ -964,6 +957,13 @@
 | `ViewChanges` |  |
 | `ViewTemplate` |  |
 
+## Tools
+
+| Name | Description |
+| ---- | ----------- |
+| `DownloadTool` | Ability to download a tool |
+| `UploadTool` | Ability to upload new tool versions |
+
 ## Accounts
 
 | Name | Description |
@@ -1121,14 +1121,14 @@
 
 | Name | Description |
 | ---- | ----------- |
-| `ReadBlobs` |  |
-| `WriteBlobs` |  |
-| `ReadAliases` |  |
-| `WriteAliases` |  |
-| `DeleteAliases` |  |
-| `ReadRefs` |  |
-| `WriteRefs` |  |
-| `DeleteRefs` |  |
+| `ReadBlobs` | Ability to read blobs from the storage service |
+| `WriteBlobs` | Ability to write blobs to the storage service |
+| `ReadAliases` | Ability to read aliases from the storage service |
+| `WriteAliases` | Ability to write aliases to the storage service |
+| `DeleteAliases` | Ability to write aliases to the storage service |
+| `ReadRefs` | Ability to read refs from the storage service |
+| `WriteRefs` | Ability to write refs to the storage service |
+| `DeleteRefs` | Ability to delete refs |
 
 ## Streams
 
@@ -1141,11 +1141,12 @@
 | `ViewChanges` |  |
 | `ViewTemplate` |  |
 
-## Symbols
+## Tools
 
 | Name | Description |
 | ---- | ----------- |
-| `ReadSymbols` | Ability to download symbols |
+| `DownloadTool` |  |
+| `UploadTool` |  |
 
 ## Accounts
 
@@ -1328,14 +1329,14 @@
 
 | Name | Description |
 | ---- | ----------- |
-| `ReadSymbols` |  |
+| `ReadSymbols` | Ability to download symbols |
 
 ## Tools
 
 | Name | Description |
 | ---- | ----------- |
-| `DownloadTool` | Ability to download a tool |
-| `UploadTool` | Ability to upload new tool versions |
+| `DownloadTool` |  |
+| `UploadTool` |  |
 
 ## Accounts
 
diff --git a/Engine/Source/Programs/Horde/Docs/Config/Schema/Globals.md b/Engine/Source/Programs/Horde/Docs/Config/Schema/Globals.md
--- a/Engine/Source/Programs/Horde/Docs/Config/Schema/Globals.md
+++ b/Engine/Source/Programs/Horde/Docs/Config/Schema/Globals.md
@@ -63,11 +63,11 @@
 `build` | [BuildConfig](#buildconfig)<br>Configuration for the build plugin
 `compute` | [ComputeConfig](#computeconfig)<br>Configuration for the compute plugin
 `exampleplugin` | [ExamplePluginConfig](#examplepluginconfig)<br>Configuration for the exampleplugin plugin
-`experimental` | [ExperimentalConfig](#experimentalconfig)<br>Configuration for the experimental plugin
 `secrets` | [SecretsConfig](#secretsconfig)<br>Configuration for the secrets plugin
+`experimental` | [ExperimentalConfig](#experimentalconfig)<br>Configuration for the experimental plugin
+`tools` | [ToolsConfig](#toolsconfig)<br>Configuration for the tools plugin
 `storage` | [StorageConfig](#storageconfig)<br>Configuration for the storage plugin
 `symbols` | [SymbolsConfig](#symbolsconfig)<br>Configuration for the symbols plugin
-`tools` | [ToolsConfig](#toolsconfig)<br>Configuration for the tools plugin
 
 ## AnalyticsConfig
 
@@ -91,6 +91,7 @@
 `projects` | [ProjectConfig](Projects.md)`[]`<br>List of projects
 `enableConformTasks` | `boolean`<br>Whether to allow conform tasks to run
 `issueFixedTag` | `string`<br>Commit tag to use for marking issues as fixed
+`issues` | [IssueConfig](Issues.md)<br>Build health issue configurations
 `issueFindSuspectsMaxSuccessDays` | `integer`<br>The maximum days since node success to scan for suspects
 `issueFindSuspectsMaxScoringSeconds` | `integer`<br>The maximum number of seconds to spend scording suspects
 
@@ -440,14 +441,6 @@
 ---- | -----------
 `greetingMessage` | `string`<br>The Greeting provided by the example plugin
 
-## ExperimentalConfig
-
-Configuration for the experimental plugin
-
-Name | Description
----- | -----------
-`notifications` | [NotificationConfig](Notification.md)`[]`<br>Notifications to be sent on the Horde server
-
 ## SecretsConfig
 
 Configuration for the secrets system
@@ -530,6 +523,42 @@
 ---- | -----------
 `vaultUri` | `string`<br>The Uri to the Azure Key Vault.
 
+## ExperimentalConfig
+
+Configuration for the experimental plugin
+
+Name | Description
+---- | -----------
+`notifications` | [NotificationConfig](Notification.md)`[]`<br>Notifications to be sent on the Horde server
+
+## ToolsConfig
+
+Configuration for the tools system
+
+Name | Description
+---- | -----------
+`tools` | [ToolConfig](#toolconfig)`[]`<br>Tool configurations
+
+## ToolConfig
+
+Options for configuring a tool
+
+Name | Description
+---- | -----------
+`id` | `string`<br>Unique identifier for the tool
+`name` | `string`<br>Name of the tool
+`description` | `string`<br>Description for the tool
+`category` | `string`<br>Category for the tool. Will cause the tool to be shown in a different tab in the dashboard.
+`group` | `string`<br>Grouping key for different variations of the same tool. The dashboard will show these together.
+`platforms` | `string[]`<br>Platforms for this tool. Takes the form of a NET RID (https://learn.microsoft.com/en-us/dotnet/core/rid-catalog).
+`public` | `boolean`<br>Whether this tool should be exposed for download on a public endpoint without authentication
+`showInUgs` | `boolean`<br>Whether to show this tool for download in the UGS tools menu
+`showInDashboard` | `boolean`<br>Whether to show this tool for download in the dashboard
+`showInToolbox` | `boolean`<br>Whether to show this tool for download in Unreal Toolbox
+`metadata` | `string` `->` `string`<br>Metadata for this tool
+`namespaceId` | `string`<br>Default namespace for new deployments of this tool
+`acl` | [AclConfig](#aclconfig)<br>Permissions for the tool
+
 ## StorageConfig
 
 Configuration for storage
@@ -620,31 +649,3 @@
 `namespaceId` | `string`<br>Configuration for the symbol store backend
 `public` | `boolean`<br>Whether to make this store available without auth
 `acl` | [AclConfig](#aclconfig)<br>Access to the symbol store
-
-## ToolsConfig
-
-Configuration for the tools system
-
-Name | Description
----- | -----------
-`tools` | [ToolConfig](#toolconfig)`[]`<br>Tool configurations
-
-## ToolConfig
-
-Options for configuring a tool
-
-Name | Description
----- | -----------
-`id` | `string`<br>Unique identifier for the tool
-`name` | `string`<br>Name of the tool
-`description` | `string`<br>Description for the tool
-`category` | `string`<br>Category for the tool. Will cause the tool to be shown in a different tab in the dashboard.
-`group` | `string`<br>Grouping key for different variations of the same tool. The dashboard will show these together.
-`platforms` | `string[]`<br>Platforms for this tool. Takes the form of a NET RID (https://learn.microsoft.com/en-us/dotnet/core/rid-catalog).
-`public` | `boolean`<br>Whether this tool should be exposed for download on a public endpoint without authentication
-`showInUgs` | `boolean`<br>Whether to show this tool for download in the UGS tools menu
-`showInDashboard` | `boolean`<br>Whether to show this tool for download in the dashboard
-`showInToolbox` | `boolean`<br>Whether to show this tool for download in Unreal Toolbox
-`metadata` | `string` `->` `string`<br>Metadata for this tool
-`namespaceId` | `string`<br>Default namespace for new deployments of this tool
-`acl` | [AclConfig](#aclconfig)<br>Permissions for the tool
diff --git a/Engine/Source/Programs/Horde/Docs/Config/Schema/Issues.md b/Engine/Source/Programs/Horde/Docs/Config/Schema/Issues.md
new file mode 100644
--- /dev/null
+++ b/Engine/Source/Programs/Horde/Docs/Config/Schema/Issues.md
@@ -0,0 +1,17 @@
+[Horde](../../../README.md) > [Configuration](../../Config.md) > *.issues.json
+
+# *.issues.json
+
+Stores configuration for build health issues
+
+Name | Description
+---- | -----------
+`channelGroups` | [IssueChannelGroupConfig](#issuechannelgroupconfig)`[]`<br>Issue grouping channels
+
+## IssueChannelGroupConfig
+
+Configuration for grouping to apply based on logging channels
+
+Name | Description
+---- | -----------
+`channels` | `string[]`<br>A list of channel logging boundaries to group issues on
diff --git a/Engine/Source/Programs/Horde/Docs/Deployment/ServerSettings.md b/Engine/Source/Programs/Horde/Docs/Deployment/ServerSettings.md
--- a/Engine/Source/Programs/Horde/Docs/Deployment/ServerSettings.md
+++ b/Engine/Source/Programs/Horde/Docs/Deployment/ServerSettings.md
@@ -131,11 +131,11 @@
 `build` | [BuildServerConfig](#buildserverconfig)<br>Configuration for the build plugin
 `compute` | [ComputeServerConfig](#computeserverconfig)<br>Configuration for the compute plugin
 `exampleplugin` | [ExamplePluginServerConfig](#examplepluginserverconfig)<br>Configuration for the exampleplugin plugin
-`experimental` | [ExperimentalServerConfig](#experimentalserverconfig)<br>Configuration for the experimental plugin
 `secrets` | [SecretsServerConfig](#secretsserverconfig)<br>Configuration for the secrets plugin
+`experimental` | [ExperimentalServerConfig](#experimentalserverconfig)<br>Configuration for the experimental plugin
+`tools` | [ToolsServerConfig](#toolsserverconfig)<br>Configuration for the tools plugin
 `storage` | [StorageServerConfig](#storageserverconfig)<br>Configuration for the storage plugin
 `symbols` | [PluginServerConfig](#pluginserverconfig)<br>Configuration for the symbols plugin
-`tools` | [ToolsServerConfig](#toolsserverconfig)<br>Configuration for the tools plugin
 
 ## AnalyticsServerConfig
 
@@ -369,14 +369,6 @@
 ---- | -----------
 `enabled` | `boolean`<br>Whether the plugin should be enabled or not
 
-## ExperimentalServerConfig
-
-Server configuration for the experimental plugin
-
-Name | Description
----- | -----------
-`enabled` | `boolean`<br>Whether the plugin should be enabled or not
-
 ## SecretsServerConfig
 
 Static configuration for the secrets plugin
@@ -386,68 +378,9 @@
 `withAws` | `boolean`<br>Whether to enable Amazon Web Services (AWS) specific features
 `enabled` | `boolean`<br>Whether the plugin should be enabled or not
 
-## StorageServerConfig
+## ExperimentalServerConfig
 
-Static settings for the storage system
+Server configuration for the experimental plugin
 
 Name | Description
 ---- | -----------
-`bundleCacheDir` | `string`<br>Directory to use for the coarse-grained backend cache. This caches full bundles downloaded from the upstream object store.
-`bundleCacheSize` | `string`<br>Maximum size of the storage cache on disk. Accepts standard binary suffixes (kb, mb, gb, tb, etc...)
-`bundleCacheSizeBytes` | `integer`<br>Accessor for the bundle cache size in bytes
-`backends` | [BackendConfig](#backendconfig)`[]`<br>Overridden settings for storage backends. Useful for running against a production server with custom backends.
-`enabled` | `boolean`<br>Whether the plugin should be enabled or not
-
-## BackendConfig
-
-Common settings object for different providers
-
-Name | Description
----- | -----------
-`id` | `string`<br>The storage backend ID
-`base` | `string`<br>Base backend to copy default settings from
-`secondary` | `string`<br>Specifies another backend to read from if an object is not found in this one. Can be used when migrating data from one backend to another.
-`type` | [StorageBackendType](#storagebackendtype-enum)<br>
-`baseDir` | `string`<br>
-`awsBucketName` | `string`<br>Name of the bucket to use
-`awsBucketPath` | `string`<br>Base path within the bucket
-`awsCredentials` | [AwsCredentialsType](#awscredentialstype-enum)<br>Type of credentials to use
-`awsRole` | `string`<br>ARN of a role to assume
-`awsProfile` | `string`<br>The AWS profile to read credentials form
-`awsRegion` | `string`<br>Region to connect to
-`azureConnectionString` | `string`<br>Connection string for Azure
-`azureContainerName` | `string`<br>Name of the container
-`relayServer` | `string`<br>
-`relayToken` | `string`<br>
-`gcsBucketName` | `string`<br>Name of the GCS bucket to use
-`gcsBucketPath` | `string`<br>Base path within the bucket
-
-## StorageBackendType (Enum)
-
-Types of storage backend to use
-
-Name | Description
----- | -----------
-`FileSystem` | Local filesystem
-`Aws` | AWS S3
-`Azure` | Azure blob store
-`Gcs` | Google Cloud Storage
-`Memory` | In-memory only (for testing)
-
-## AwsCredentialsType (Enum)
-
-Credentials to use for AWS
-
-Name | Description
----- | -----------
-`Default` | Use default credentials from the AWS SDK
-`Profile` | Read credentials from the  profile in the AWS config file
-`AssumeRole` | Assume a particular role. Should specify ARN in
-`AssumeRoleWebIdentity` | Assume a particular role using the current environment variables.
-
-## PluginServerConfig
-
-Base class for plugin server config objects
-
-Name | Description
----- | -----------
@@ -525,3 +458,70 @@
 `actions` | `string[]`<br>Actions to include
 `excludeActions` | `string[]`<br>Actions to exclude from the inherited actions
 `extends` | `string[]`<br>Other profiles to extend from
+
+## StorageServerConfig
+
+Static settings for the storage system
+
+Name | Description
+---- | -----------
+`bundleCacheDir` | `string`<br>Directory to use for the coarse-grained backend cache. This caches full bundles downloaded from the upstream object store.
+`bundleCacheSize` | `string`<br>Maximum size of the storage cache on disk. Accepts standard binary suffixes (kb, mb, gb, tb, etc...)
+`bundleCacheSizeBytes` | `integer`<br>Accessor for the bundle cache size in bytes
+`backends` | [BackendConfig](#backendconfig)`[]`<br>Overridden settings for storage backends. Useful for running against a production server with custom backends.
+`enabled` | `boolean`<br>Whether the plugin should be enabled or not
+
+## BackendConfig
+
+Common settings object for different providers
+
+Name | Description
+---- | -----------
+`id` | `string`<br>The storage backend ID
+`base` | `string`<br>Base backend to copy default settings from
+`secondary` | `string`<br>Specifies another backend to read from if an object is not found in this one. Can be used when migrating data from one backend to another.
+`type` | [StorageBackendType](#storagebackendtype-enum)<br>
+`baseDir` | `string`<br>
+`awsBucketName` | `string`<br>Name of the bucket to use
+`awsBucketPath` | `string`<br>Base path within the bucket
+`awsCredentials` | [AwsCredentialsType](#awscredentialstype-enum)<br>Type of credentials to use
+`awsRole` | `string`<br>ARN of a role to assume
+`awsProfile` | `string`<br>The AWS profile to read credentials form
+`awsRegion` | `string`<br>Region to connect to
+`azureConnectionString` | `string`<br>Connection string for Azure
+`azureContainerName` | `string`<br>Name of the container
+`relayServer` | `string`<br>
+`relayToken` | `string`<br>
+`gcsBucketName` | `string`<br>Name of the GCS bucket to use
+`gcsBucketPath` | `string`<br>Base path within the bucket
+
+## StorageBackendType (Enum)
+
+Types of storage backend to use
+
+Name | Description
+---- | -----------
+`FileSystem` | Local filesystem
+`Aws` | AWS S3
+`Azure` | Azure blob store
+`Gcs` | Google Cloud Storage
+`Memory` | In-memory only (for testing)
+
+## AwsCredentialsType (Enum)
+
+Credentials to use for AWS
+
+Name | Description
+---- | -----------
+`Default` | Use default credentials from the AWS SDK
+`Profile` | Read credentials from the  profile in the AWS config file
+`AssumeRole` | Assume a particular role. Should specify ARN in
+`AssumeRoleWebIdentity` | Assume a particular role using the current environment variables.
+
+## PluginServerConfig
+
+Base class for plugin server config objects
+
+Name | Description
+---- | -----------
+`enabled` | `boolean`<br>Whether the plugin should be enabled or not
