Horde: Tweak channel group fingerprint type
#rnx

diff --git a/Engine/Source/Programs/Horde/Plugins/Build/HordeServer.Build.Tests/Issues/IssueServiceTests.cs b/Engine/Source/Programs/Horde/Plugins/Build/HordeServer.Build.Tests/Issues/IssueServiceTests.cs
--- a/Engine/Source/Programs/Horde/Plugins/Build/HordeServer.Build.Tests/Issues/IssueServiceTests.cs
+++ b/Engine/Source/Programs/Horde/Plugins/Build/HordeServer.Build.Tests/Issues/IssueServiceTests.cs
@@ -2299,10 +2299,10 @@ namespace HordeServer.Tests.Issues
 				Assert.AreEqual(3, issues.Count);
 
 				Assert.AreEqual("AssetCheck errors in Update Version Files", issues[0].Summary);
-				Assert.AreEqual("Hashed_AssetCheck", issues[0].Fingerprints[0].Type);
+				Assert.AreEqual("Hashed:AssetCheck", issues[0].Fingerprints[0].Type);
 
 				Assert.AreEqual("LogBlueprint errors in Update Version Files", issues[1].Summary);
-				Assert.AreEqual("Hashed_LogBlueprint_LogLinker", issues[1].Fingerprints[0].Type);
+				Assert.AreEqual("Hashed:LogBlueprint_LogLinker", issues[1].Fingerprints[0].Type);
 
 				Assert.AreEqual("LogOther errors in Update Version Files", issues[2].Summary);
 				Assert.AreEqual("Hashed", issues[2].Fingerprints[0].Type);
diff --git a/Engine/Source/Programs/Shared/EpicGames.Horde/Issues/Handlers/HashedIssueHandler.cs b/Engine/Source/Programs/Shared/EpicGames.Horde/Issues/Handlers/HashedIssueHandler.cs
--- a/Engine/Source/Programs/Shared/EpicGames.Horde/Issues/Handlers/HashedIssueHandler.cs
+++ b/Engine/Source/Programs/Shared/EpicGames.Horde/Issues/Handlers/HashedIssueHandler.cs
@@ -118,7 +118,7 @@ namespace EpicGames.Horde.Issues.Handlers
 
 				if (!String.IsNullOrEmpty(channelGroup))
 				{
-					fingerprint += $"_{channelGroup}";
+					fingerprint += $":{channelGroup}";
 				}
 
 				List<IssueEvent>? events = null;
