Horde: Update release notes

#rnx

diff --git a/Engine/Source/Programs/Horde/Docs/ReleaseNotes.md b/Engine/Source/Programs/Horde/Docs/ReleaseNotes.md
--- a/Engine/Source/Programs/Horde/Docs/ReleaseNotes.md
+++ b/Engine/Source/Programs/Horde/Docs/ReleaseNotes.md
@@ -1,5 +1,26 @@
 # Release Notes
 
+## 2025-10-16
+
+* Refactor issue config to be a root object (46971377)
+* Add support for grouping issues on log channel boundaries (46967755)
+* Updated Label telemetry to emit individual events in addition to summary. (46958593)
+* Log info from FindOrAddIssueForSpanAsync to issue audit as well (46956201)
+* Add ability to clean up orphaned agent session IDs (46930732)
+* Add overload of MultiLog function to capture exception object and use it for suspect finding exceptions (46898218)
+* Update labels for select/clear all buttons + remove the clear-all filtered subset behaviour (46893975)
+* Added suspect finding info to issue audit log (46892792)
+* Label support within table (46880003)
+* Include job name in <title> for job details view (46840656)
+* Build health make issue suspect max days and scoring time configuration, also bump defaults (46803980)
+* Add Build Health node annotation to suppress warning issues (46803820)
+* Use full stream name for tab title if available (46766042)
+* Sending up schema version as well (46744094)
+* Fix issue with assignment/acknowledgement (46734592)
+* Increase per stream limit for historical views (46732352)
+* Replace PerformanceCounter with WINAPI GetSystemTimes
+* Updating API used for obtaining changelist results (46721751)
+
 ## 2025-10-09
 
 * [Build Health Dashboard] Updating API used for obtaining changelist result (46721751)
