Horde: Adding job/step messages/severity to job reports

diff --git a/Engine/Source/Programs/Horde/Drivers/JobDriver/Execution/JobExecutor.cs b/Engine/Source/Programs/Horde/Drivers/JobDriver/Execution/JobExecutor.cs
--- a/Engine/Source/Programs/Horde/Drivers/JobDriver/Execution/JobExecutor.cs
+++ b/Engine/Source/Programs/Horde/Drivers/JobDriver/Execution/JobExecutor.cs
@@ -194,6 +194,8 @@ namespace JobDriver.Execution
 			public string Name { get; set; } = String.Empty;
 			public string Content { get; set; } = String.Empty;
 			public string FileName { get; set; } = String.Empty;
+			public string? Message { get; set; }
+			public string? Severity { get; set; }
 		}
 
 		const string ScriptArgumentPrefix = "-Script=";
@@ -2218,6 +2220,8 @@ namespace JobDriver.Execution
 			request.Placement = report.Placement;
 			request.Name = report.Name;
 			request.Content = report.Content;
+			request.Message = report.Message ?? String.Empty;
+			request.Severity = report.Severity ?? String.Empty;
 			await JobRpc.CreateReportAsync(request);
 		}
 
diff --git a/Engine/Source/Programs/Horde/Plugins/Build/HordeServer.Build/Jobs/JobRpcService.cs b/Engine/Source/Programs/Horde/Plugins/Build/HordeServer.Build/Jobs/JobRpcService.cs
--- a/Engine/Source/Programs/Horde/Plugins/Build/HordeServer.Build/Jobs/JobRpcService.cs
+++ b/Engine/Source/Programs/Horde/Plugins/Build/HordeServer.Build/Jobs/JobRpcService.cs
@@ -962,7 +962,7 @@ namespace HordeServer.Jobs
 
 			IJobStepBatch batch = AuthorizeBatch(job, JobStepBatchId.Parse(request.BatchId), context);
 
-			JobReport newReport = new JobReport { Name = request.Name, Placement = (JobReportPlacement)request.Placement, Content = request.Content };
+			JobReport newReport = new JobReport { Name = request.Name, Placement = (JobReportPlacement)request.Placement, Content = request.Content, Message = request.Message, Severity = request.Severity };
 			if (request.Scope == RpcReportScope.Job)
 			{
 				_logger.LogDebug("Adding report to job {JobId}: {Name} -> {Content}", job.Id, request.Name, request.Content);
diff --git a/Engine/Source/Programs/Horde/Plugins/Build/HordeServer.Build/Jobs/JobsController.cs b/Engine/Source/Programs/Horde/Plugins/Build/HordeServer.Build/Jobs/JobsController.cs
--- a/Engine/Source/Programs/Horde/Plugins/Build/HordeServer.Build/Jobs/JobsController.cs
+++ b/Engine/Source/Programs/Horde/Plugins/Build/HordeServer.Build/Jobs/JobsController.cs
@@ -741,6 +741,8 @@ namespace HordeServer.Jobs
 			GetJobReportResponse response = new GetJobReportResponse(report.Name, report.Placement);
 			response.ArtifactId = report.ArtifactId?.ToString();
 			response.Content = report.Content;
+			response.Message = report.Message;
+			response.Severity = report.Severity;
 			return response;
 		}
 
diff --git a/Engine/Source/Programs/Shared/EpicGames.Horde/Jobs/IJob.cs b/Engine/Source/Programs/Shared/EpicGames.Horde/Jobs/IJob.cs
--- a/Engine/Source/Programs/Shared/EpicGames.Horde/Jobs/IJob.cs
+++ b/Engine/Source/Programs/Shared/EpicGames.Horde/Jobs/IJob.cs
@@ -1908,6 +1908,17 @@ namespace EpicGames.Horde.Jobs
 		/// Inline data for the report
 		/// </summary>
 		string? Content { get; }
+
+		/// <summary>
+		/// The job/step message
+		/// </summary>
+		string? Message { get; }
+
+		/// <summary>
+		/// The message severity 
+		/// </summary>
+		string? Severity { get; }
+
 	}
 
 	/// <summary>
@@ -1926,5 +1937,11 @@ namespace EpicGames.Horde.Jobs
 
 		/// <inheritdoc/>
 		public string? Content { get; set; }
+
+		/// <inheritdoc/>
+		public string? Message { get; set; }
+
+		/// <inheritdoc/>
+		public string? Severity { get; set; }
 	}
 }
diff --git a/Engine/Source/Programs/Shared/EpicGames.Horde/Jobs/JobMessages.cs b/Engine/Source/Programs/Shared/EpicGames.Horde/Jobs/JobMessages.cs
--- a/Engine/Source/Programs/Shared/EpicGames.Horde/Jobs/JobMessages.cs
+++ b/Engine/Source/Programs/Shared/EpicGames.Horde/Jobs/JobMessages.cs
@@ -487,6 +487,16 @@ namespace EpicGames.Horde.Jobs
 		/// </summary>
 		public string? Content { get; set; }
 
+		/// <summary>
+		/// The job/step message
+		/// </summary>
+		public string? Message { get; set; }
+
+		/// <summary>
+		/// The job/step message severity
+		/// </summary>
+		public string? Severity { get; set; }
+
 		/// <summary>
 		/// Constructor
 		/// </summary>
