Horde: Updating schema documentation

#rnx

diff --git a/Engine/Source/Programs/Horde/Docs/Config/Schema/AclActions.md b/Engine/Source/Programs/Horde/Docs/Config/Schema/AclActions.md
--- a/Engine/Source/Programs/Horde/Docs/Config/Schema/AclActions.md
+++ b/Engine/Source/Programs/Horde/Docs/Config/Schema/AclActions.md
@@ -588,12 +588,6 @@
 | `UpdateReplicator` |  |
 | `ViewReplicator` |  |
 
-## Secrets
-
-| Name | Description |
-| ---- | ----------- |
-| `ViewSecret` | View a credential |
-
 ## ServiceAccounts
 
 | Name | Description |
@@ -756,7 +750,7 @@
 
 | Name | Description |
 | ---- | ----------- |
-| `ViewSecret` |  |
+| `ViewSecret` | View a credential |
 
 ## ServiceAccounts
 
@@ -946,6 +940,19 @@
 | `DownloadSoftware` |  |
 | `DeleteSoftware` |  |
 
+## Storage
+
+| Name | Description |
+| ---- | ----------- |
+| `ReadBlobs` | Ability to read blobs from the storage service |
+| `WriteBlobs` | Ability to write blobs to the storage service |
+| `ReadAliases` | Ability to read aliases from the storage service |
+| `WriteAliases` | Ability to write aliases to the storage service |
+| `DeleteAliases` | Ability to write aliases to the storage service |
+| `ReadRefs` | Ability to read refs from the storage service |
+| `WriteRefs` | Ability to write refs to the storage service |
+| `DeleteRefs` | Ability to delete refs |
+
 ## Streams
 
 | Name | Description |
@@ -957,13 +964,6 @@
 | `ViewChanges` |  |
 | `ViewTemplate` |  |
 
-## Tools
-
-| Name | Description |
-| ---- | ----------- |
-| `DownloadTool` | Ability to download a tool |
-| `UploadTool` | Ability to upload new tool versions |
-
 ## Accounts
 
 | Name | Description |
@@ -1121,14 +1121,14 @@
 
 | Name | Description |
 | ---- | ----------- |
-| `ReadBlobs` | Ability to read blobs from the storage service |
-| `WriteBlobs` | Ability to write blobs to the storage service |
-| `ReadAliases` | Ability to read aliases from the storage service |
-| `WriteAliases` | Ability to write aliases to the storage service |
-| `DeleteAliases` | Ability to write aliases to the storage service |
-| `ReadRefs` | Ability to read refs from the storage service |
-| `WriteRefs` | Ability to write refs to the storage service |
-| `DeleteRefs` | Ability to delete refs |
+| `ReadBlobs` |  |
+| `WriteBlobs` |  |
+| `ReadAliases` |  |
+| `WriteAliases` |  |
+| `DeleteAliases` |  |
+| `ReadRefs` |  |
+| `WriteRefs` |  |
+| `DeleteRefs` |  |
 
 ## Streams
 
@@ -1141,12 +1141,11 @@
 | `ViewChanges` |  |
 | `ViewTemplate` |  |
 
-## Tools
+## Symbols
 
 | Name | Description |
 | ---- | ----------- |
-| `DownloadTool` |  |
-| `UploadTool` |  |
+| `ReadSymbols` | Ability to download symbols |
 
 ## Accounts
 
@@ -1329,14 +1328,14 @@
 
 | Name | Description |
 | ---- | ----------- |
-| `ReadSymbols` | Ability to download symbols |
+| `ReadSymbols` |  |
 
 ## Tools
 
 | Name | Description |
 | ---- | ----------- |
-| `DownloadTool` |  |
-| `UploadTool` |  |
+| `DownloadTool` | Ability to download a tool |
+| `UploadTool` | Ability to upload new tool versions |
 
 ## Accounts
 
diff --git a/Engine/Source/Programs/Horde/Docs/Config/Schema/Issues.md b/Engine/Source/Programs/Horde/Docs/Config/Schema/Issues.md
--- a/Engine/Source/Programs/Horde/Docs/Config/Schema/Issues.md
+++ b/Engine/Source/Programs/Horde/Docs/Config/Schema/Issues.md
@@ -6,8 +6,17 @@ Stores configuration for build health issues
 
 Name | Description
 ---- | -----------
+`include` | [ConfigInclude](#configinclude)`[]`<br>Includes for other configuration files
 `channelGroups` | [IssueChannelGroupConfig](#issuechannelgroupconfig)`[]`<br>Issue grouping channels
 
+## ConfigInclude
+
+Directive to merge config data from another source
+
+Name | Description
+---- | -----------
+`path` | `string`<br>Path to the config data to be included. May be relative to the including file's location.
+
 ## IssueChannelGroupConfig
 
 Configuration for grouping to apply based on logging channels
diff --git a/Engine/Source/Programs/Horde/Docs/Deployment/ServerSettings.md b/Engine/Source/Programs/Horde/Docs/Deployment/ServerSettings.md
--- a/Engine/Source/Programs/Horde/Docs/Deployment/ServerSettings.md
+++ b/Engine/Source/Programs/Horde/Docs/Deployment/ServerSettings.md
@@ -131,11 +131,11 @@
 `build` | [BuildServerConfig](#buildserverconfig)<br>Configuration for the build plugin
 `compute` | [ComputeServerConfig](#computeserverconfig)<br>Configuration for the compute plugin
 `exampleplugin` | [ExamplePluginServerConfig](#examplepluginserverconfig)<br>Configuration for the exampleplugin plugin
-`secrets` | [SecretsServerConfig](#secretsserverconfig)<br>Configuration for the secrets plugin
 `experimental` | [ExperimentalServerConfig](#experimentalserverconfig)<br>Configuration for the experimental plugin
-`tools` | [ToolsServerConfig](#toolsserverconfig)<br>Configuration for the tools plugin
+`secrets` | [SecretsServerConfig](#secretsserverconfig)<br>Configuration for the secrets plugin
 `storage` | [StorageServerConfig](#storageserverconfig)<br>Configuration for the storage plugin
 `symbols` | [PluginServerConfig](#pluginserverconfig)<br>Configuration for the symbols plugin
+`tools` | [ToolsServerConfig](#toolsserverconfig)<br>Configuration for the tools plugin
 
 ## AnalyticsServerConfig
 
@@ -369,6 +369,14 @@
 ---- | -----------
 `enabled` | `boolean`<br>Whether the plugin should be enabled or not
 
+## ExperimentalServerConfig
+
+Server configuration for the experimental plugin
+
+Name | Description
+---- | -----------
+`enabled` | `boolean`<br>Whether the plugin should be enabled or not
+
 ## SecretsServerConfig
 
 Static configuration for the secrets plugin
@@ -378,9 +386,68 @@
 `withAws` | `boolean`<br>Whether to enable Amazon Web Services (AWS) specific features
 `enabled` | `boolean`<br>Whether the plugin should be enabled or not
 
-## ExperimentalServerConfig
+## StorageServerConfig
 
-Server configuration for the experimental plugin
+Static settings for the storage system
 
 Name | Description
 ---- | -----------
+`bundleCacheDir` | `string`<br>Directory to use for the coarse-grained backend cache. This caches full bundles downloaded from the upstream object store.
+`bundleCacheSize` | `string`<br>Maximum size of the storage cache on disk. Accepts standard binary suffixes (kb, mb, gb, tb, etc...)
+`bundleCacheSizeBytes` | `integer`<br>Accessor for the bundle cache size in bytes
+`backends` | [BackendConfig](#backendconfig)`[]`<br>Overridden settings for storage backends. Useful for running against a production server with custom backends.
+`enabled` | `boolean`<br>Whether the plugin should be enabled or not
+
+## BackendConfig
+
+Common settings object for different providers
+
+Name | Description
+---- | -----------
+`id` | `string`<br>The storage backend ID
+`base` | `string`<br>Base backend to copy default settings from
+`secondary` | `string`<br>Specifies another backend to read from if an object is not found in this one. Can be used when migrating data from one backend to another.
+`type` | [StorageBackendType](#storagebackendtype-enum)<br>
+`baseDir` | `string`<br>
+`awsBucketName` | `string`<br>Name of the bucket to use
+`awsBucketPath` | `string`<br>Base path within the bucket
+`awsCredentials` | [AwsCredentialsType](#awscredentialstype-enum)<br>Type of credentials to use
+`awsRole` | `string`<br>ARN of a role to assume
+`awsProfile` | `string`<br>The AWS profile to read credentials form
+`awsRegion` | `string`<br>Region to connect to
+`azureConnectionString` | `string`<br>Connection string for Azure
+`azureContainerName` | `string`<br>Name of the container
+`relayServer` | `string`<br>
+`relayToken` | `string`<br>
+`gcsBucketName` | `string`<br>Name of the GCS bucket to use
+`gcsBucketPath` | `string`<br>Base path within the bucket
+
+## StorageBackendType (Enum)
+
+Types of storage backend to use
+
+Name | Description
+---- | -----------
+`FileSystem` | Local filesystem
+`Aws` | AWS S3
+`Azure` | Azure blob store
+`Gcs` | Google Cloud Storage
+`Memory` | In-memory only (for testing)
+
+## AwsCredentialsType (Enum)
+
+Credentials to use for AWS
+
+Name | Description
+---- | -----------
+`Default` | Use default credentials from the AWS SDK
+`Profile` | Read credentials from the  profile in the AWS config file
+`AssumeRole` | Assume a particular role. Should specify ARN in
+`AssumeRoleWebIdentity` | Assume a particular role using the current environment variables.
+
+## PluginServerConfig
+
+Base class for plugin server config objects
+
+Name | Description
+---- | -----------
@@ -458,70 +525,3 @@
 `actions` | `string[]`<br>Actions to include
 `excludeActions` | `string[]`<br>Actions to exclude from the inherited actions
 `extends` | `string[]`<br>Other profiles to extend from
-
-## StorageServerConfig
-
-Static settings for the storage system
-
-Name | Description
----- | -----------
-`bundleCacheDir` | `string`<br>Directory to use for the coarse-grained backend cache. This caches full bundles downloaded from the upstream object store.
-`bundleCacheSize` | `string`<br>Maximum size of the storage cache on disk. Accepts standard binary suffixes (kb, mb, gb, tb, etc...)
-`bundleCacheSizeBytes` | `integer`<br>Accessor for the bundle cache size in bytes
-`backends` | [BackendConfig](#backendconfig)`[]`<br>Overridden settings for storage backends. Useful for running against a production server with custom backends.
-`enabled` | `boolean`<br>Whether the plugin should be enabled or not
-
-## BackendConfig
-
-Common settings object for different providers
-
-Name | Description
----- | -----------
-`id` | `string`<br>The storage backend ID
-`base` | `string`<br>Base backend to copy default settings from
-`secondary` | `string`<br>Specifies another backend to read from if an object is not found in this one. Can be used when migrating data from one backend to another.
-`type` | [StorageBackendType](#storagebackendtype-enum)<br>
-`baseDir` | `string`<br>
-`awsBucketName` | `string`<br>Name of the bucket to use
-`awsBucketPath` | `string`<br>Base path within the bucket
-`awsCredentials` | [AwsCredentialsType](#awscredentialstype-enum)<br>Type of credentials to use
-`awsRole` | `string`<br>ARN of a role to assume
-`awsProfile` | `string`<br>The AWS profile to read credentials form
-`awsRegion` | `string`<br>Region to connect to
-`azureConnectionString` | `string`<br>Connection string for Azure
-`azureContainerName` | `string`<br>Name of the container
-`relayServer` | `string`<br>
-`relayToken` | `string`<br>
-`gcsBucketName` | `string`<br>Name of the GCS bucket to use
-`gcsBucketPath` | `string`<br>Base path within the bucket
-
-## StorageBackendType (Enum)
-
-Types of storage backend to use
-
-Name | Description
----- | -----------
-`FileSystem` | Local filesystem
-`Aws` | AWS S3
-`Azure` | Azure blob store
-`Gcs` | Google Cloud Storage
-`Memory` | In-memory only (for testing)
-
-## AwsCredentialsType (Enum)
-
-Credentials to use for AWS
-
-Name | Description
----- | -----------
-`Default` | Use default credentials from the AWS SDK
-`Profile` | Read credentials from the  profile in the AWS config file
-`AssumeRole` | Assume a particular role. Should specify ARN in
-`AssumeRoleWebIdentity` | Assume a particular role using the current environment variables.
-
-## PluginServerConfig
-
-Base class for plugin server config objects
-
-Name | Description
----- | -----------
-`enabled` | `boolean`<br>Whether the plugin should be enabled or not
