Fix missing recipeId property in test session

#rnx

diff --git a/Engine/Source/Programs/Horde/Plugins/Build/HordeServer.Build.Tests/Jobs/TestData/TestDataTests.cs b/Engine/Source/Programs/Horde/Plugins/Build/HordeServer.Build.Tests/Jobs/TestData/TestDataTests.cs
--- a/Engine/Source/Programs/Horde/Plugins/Build/HordeServer.Build.Tests/Jobs/TestData/TestDataTests.cs
+++ b/Engine/Source/Programs/Horde/Plugins/Build/HordeServer.Build.Tests/Jobs/TestData/TestDataTests.cs
@@ -445,6 +445,7 @@ namespace HordeServer.Tests.Jobs.TestData
 			Assert.AreEqual(6, sessions.Count);
 			Assert.AreEqual(TestOutcome.Failure, sessions[0].Outcome);
 			Assert.IsTrue(sessions[0].TagIds?.All(id => tags.Any(t => t.Id == id)));
+			Assert.IsNull(sessions[0].RecipeId);
 
 			Assert.IsNotNull(sessions[0].ErrorFingerprints);
 			Assert.AreEqual(1, sessions[0].ErrorFingerprints!.Count);
@@ -559,6 +560,14 @@ namespace HordeServer.Tests.Jobs.TestData
 			Assert.IsTrue(testRecipes[0].Params?.Any(p => p.Key == "project" && p.Value.Equals("Lyra")));
 
 			Assert.AreEqual(_releaseStreamId.ToString(), testRecipes[1].StreamId);
+
+			GetTestSessionsRequest sessionsRequest = new GetTestSessionsRequest() { StreamIds = streamIds.ToList(), TestIds = testnames.Select(t => t.Id.ToString()).ToList() };
+			ActionResult<List<GetTestSessionResponse>> testSessionResult = await TestDataController.GetTestSessionsV2Async(sessionsRequest);
+			Assert.IsNotNull(testSessionResult);
+			Assert.IsNotNull(testSessionResult.Value);
+			List<GetTestSessionResponse> sessions = testSessionResult.Value;
+			Assert.AreEqual(3, sessions.Count);
+			Assert.AreEqual(testRecipes[0].Id, sessions[0].RecipeId);
 		}
 
 		private readonly string[] _simpleTestDataLines =
diff --git a/Engine/Source/Programs/Horde/Plugins/Build/HordeServer.Build/Jobs/TestData/TestDataController.cs b/Engine/Source/Programs/Horde/Plugins/Build/HordeServer.Build/Jobs/TestData/TestDataController.cs
--- a/Engine/Source/Programs/Horde/Plugins/Build/HordeServer.Build/Jobs/TestData/TestDataController.cs
+++ b/Engine/Source/Programs/Horde/Plugins/Build/HordeServer.Build/Jobs/TestData/TestDataController.cs
@@ -820,6 +820,7 @@ namespace HordeServer.Jobs.TestData
 					CommitOrder = testData.BuildCommitId.Order,
 					MetadataId = testData.Metadata.ToString(),
 					NameRef = testData.NameRef.ToString(),
+					RecipeId = testData.Recipe?.ToString(),
 					TestDataId = testData.TestDataId.ToString(),
 					TagIds = testData.Tags?.Select(t => t.ToString()).ToArray(),
 					Outcome = testData.Outcome,
diff --git a/Engine/Source/Programs/Shared/EpicGames.Horde/Jobs/TestData/TestDataMessages.cs b/Engine/Source/Programs/Shared/EpicGames.Horde/Jobs/TestData/TestDataMessages.cs
--- a/Engine/Source/Programs/Shared/EpicGames.Horde/Jobs/TestData/TestDataMessages.cs
+++ b/Engine/Source/Programs/Shared/EpicGames.Horde/Jobs/TestData/TestDataMessages.cs
@@ -861,6 +861,11 @@ namespace EpicGames.Horde.Jobs.TestData
 		/// </summary>
 		public string NameRef { get; set; } = String.Empty;
 
+		/// <summary>
+		/// The recipe id associated with the test session
+		/// </summary>
+		public string? RecipeId { get; set; } = String.Empty;
+
 		/// <summary>
 		/// The outcome of the test session
 		/// </summary>
