Horde: Bake report step message/severtity into ref metadata

diff --git a/Engine/Source/Programs/Horde/Plugins/Build/HordeServer.Build/Jobs/JobRpcService.cs b/Engine/Source/Programs/Horde/Plugins/Build/HordeServer.Build/Jobs/JobRpcService.cs
--- a/Engine/Source/Programs/Horde/Plugins/Build/HordeServer.Build/Jobs/JobRpcService.cs
+++ b/Engine/Source/Programs/Horde/Plugins/Build/HordeServer.Build/Jobs/JobRpcService.cs
@@ -24,11 +24,11 @@ using HordeServer.Artifacts;
 using HordeServer.Jobs.Graphs;
 using HordeServer.Jobs.TestData;
 using HordeServer.Logs;
-using HordeServer.VersionControl.Perforce;
 using HordeServer.Server;
 using HordeServer.Streams;
 using HordeServer.Tasks;
 using HordeServer.Utilities;
+using HordeServer.VersionControl.Perforce;
 using Microsoft.AspNetCore.Authorization;
 using Microsoft.Extensions.Logging;
 using Microsoft.Extensions.Options;
@@ -970,8 +970,23 @@ namespace HordeServer.Jobs
 			}
 			else
 			{
+				JobStepId stepId = JobStepId.Parse(request.StepId);
 				_logger.LogDebug("Adding report to step {JobId}:{BatchId}:{StepId}: {Name} -> {Content}", job.Id, batch.Id, request.StepId, request.Name, request.Content);
-				await _jobService.UpdateStepAsync(job, batch.Id, JobStepId.Parse(request.StepId), streamConfig, newReports: new List<JobReport> { newReport });
+				await _jobService.UpdateStepAsync(job, batch.Id, stepId, streamConfig, newReports: new List<JobReport> { newReport });
+
+				if (!String.IsNullOrEmpty(newReport.Message))
+				{
+					List<string> reportMetaData = new List<string>() { $"ReportStepMessage={newReport.Message}"};
+
+					if (!String.IsNullOrEmpty(newReport.Severity))
+					{
+						reportMetaData.Add($"ReportStepSeverity={newReport.Severity}");
+					}
+
+					Dictionary<JobStepId, List<string>> stepMetaData = new Dictionary<JobStepId, List<string>>() { { stepId, reportMetaData } };					
+					
+					await _jobService.UpdateMetadataAsync(job.Id, null, stepMetaData, context.CancellationToken);
+				}
 			}
 
 			return new RpcCreateReportResponse();
