Horde: Adding step messages to dashboard step view, step details, log view, and history

diff --git a/Engine/Source/Programs/Horde/HordeDashboard/src/backend/Api.ts b/Engine/Source/Programs/Horde/HordeDashboard/src/backend/Api.ts
--- a/Engine/Source/Programs/Horde/HordeDashboard/src/backend/Api.ts
+++ b/Engine/Source/Programs/Horde/HordeDashboard/src/backend/Api.ts
@@ -1229,6 +1229,12 @@ export type GetReportResponse = {
 	/** The report markdown content */
 	content?: string;
 
+	/** The job or step message */
+	message?: string;
+
+	/** The job or step severity */
+	severity?: string;
+
 }
 
 export type GetJobArtifactResponse = {
@@ -1606,6 +1612,10 @@ export interface GetJobStepRefResponse {
 
 	/** Issue ids affecting this step ref */
 	issueIds?: number[];
+
+	/** Step metadata */
+	metadata?: string[];
+
 }
 
 export interface GetJobStepTraceResponse {
diff --git a/Engine/Source/Programs/Horde/HordeDashboard/src/components/jobDetailsV2/JobDetailStepHistory.tsx b/Engine/Source/Programs/Horde/HordeDashboard/src/components/jobDetailsV2/JobDetailStepHistory.tsx
--- a/Engine/Source/Programs/Horde/HordeDashboard/src/components/jobDetailsV2/JobDetailStepHistory.tsx
+++ b/Engine/Source/Programs/Horde/HordeDashboard/src/components/jobDetailsV2/JobDetailStepHistory.tsx
@@ -1,13 +1,13 @@
 // Copyright Epic Games, Inc. All Rights Reserved.
 
-import { DetailsList, DetailsListLayoutMode, DetailsRow, IColumn, IDetailsListProps, SelectionMode, Stack, Text } from "@fluentui/react";
+import { DetailsList, DetailsListLayoutMode, DetailsRow, FontIcon, IColumn, IDetailsListProps, SelectionMode, Stack, Text } from "@fluentui/react";
 import { observer } from "mobx-react-lite";
 import moment from "moment";
 import { useEffect } from "react";
 import { Link, useNavigate, useLocation } from "react-router-dom";
 import backend from "../../backend";
 import { GetJobStepRefResponse } from "../../backend/Api";
-import dashboard from "../../backend/Dashboard";
+import dashboard, { StatusColor } from "../../backend/Dashboard";
 import { ISideRailLink } from "../../base/components/SideRail";
 import { displayTimeZone, getElapsedString } from "../../base/utilities/timeUtils";
 import { ChangeButton } from "../ChangeButton";
@@ -56,7 +56,7 @@ class StepHistoryDataView extends JobDataView {
          return;
       }
       backend.getJobStepHistory(jobData.streamId, stepName, 1024, jobData.templateId!).then(response => {
-         this.history = response;         
+         this.history = response;
       }).finally(() => {
          this.initialize(this.history?.length ? [sideRail] : undefined);
          this.updateReady();
@@ -74,7 +74,7 @@ class StepHistoryDataView extends JobDataView {
          const stepName = this.details?.getStepName(this.stepId, false);
          if (stepName) {
             this.loadHistory(stepName);
-         }         
+         }
       }
    }
 
@@ -131,11 +131,12 @@ export const StepHistoryPanel: React.FC<{ jobDetails: JobDetailsV2; stepId: stri
    }
 
    const columns = [
-      { key: 'column1', name: 'Name', minWidth: 620, maxWidth: 620, isResizable: false },
+      { key: 'column1', name: 'Name', minWidth: 580, maxWidth: 580, isResizable: false },
       { key: 'column2', name: 'Change', minWidth: 80, maxWidth: 80, isResizable: false },
       { key: 'column3', name: 'Started', minWidth: 180, maxWidth: 180, isResizable: false },
       { key: 'column4', name: 'Agent', minWidth: 100, maxWidth: 100, isResizable: false },
-      { key: 'column5', name: 'Duration', minWidth: 110, maxWidth: 110, isResizable: false },
+      { key: 'column5', name: 'Meta', minWidth: 200, maxWidth: 200, isResizable: false },
+      { key: 'column6', name: 'Duration', minWidth: 90, maxWidth: 90, isResizable: false },
    ];
 
    const renderItem = (item: HistoryItem, index?: number, column?: IColumn) => {
@@ -175,6 +176,33 @@ export const StepHistoryPanel: React.FC<{ jobDetails: JobDetailsV2; stepId: stri
          return <a href={url} onClick={(ev) => { ev.preventDefault(); ev.stopPropagation(); navigate(url, { replace: true }); }}><Stack horizontal horizontalAlign={"end"} verticalFill={true} tokens={{ childrenGap: 0, padding: 0 }}><Text>{agentId}</Text></Stack></a>;
       }
 
+      if (column.name === "Meta") {
+
+         let stepMessage = ref.metadata?.find(m => m.startsWith("ReportStepMessage="));
+         let stepSeverity = ref.metadata?.find(m => m.startsWith("ReportStepSeverity="));
+
+         stepMessage = stepMessage?.split("=")[1]?.trim();
+         stepSeverity = stepSeverity?.split("=")[1]?.trim();
+
+         if (!stepMessage) {
+            return null;
+         }
+
+         let stepMessageColor = "";
+         if (stepSeverity === "Error" || stepSeverity === "Errors") {
+            stepMessageColor = dashboard.getStatusColors().get(StatusColor.Failure)!;
+         }
+         if (stepSeverity === "Warning" || stepSeverity === "Warnings") {
+            stepMessageColor = dashboard.getStatusColors().get(StatusColor.Warnings)!;
+         }
+
+         return <Stack verticalAlign="center" horizontal tokens={{ childrenGap: 6 }} horizontalAlign="center">
+               {!!stepMessageColor && <FontIcon style={{ fontSize: 13, color: stepMessageColor, paddingTop: 1 }} iconName="Square" />}
+               <Text style={{fontSize: 12, overflow: "hidden", textOverflow: "ellipsis"}}>{stepMessage}</Text>
+            </Stack>         
+      }
+
+
       if (column.name === "Started") {
 
          if (ref.startTime) {
@@ -184,7 +212,6 @@ export const StepHistoryPanel: React.FC<{ jobDetails: JobDetailsV2; stepId: stri
 
             let displayTimeStr = displayTime.format('MMM Do') + ` at ${displayTime.format(format)}`;
 
-
             return <Stack horizontal horizontalAlign={"start"}>{displayTimeStr}</Stack>;
 
          } else {
diff --git a/Engine/Source/Programs/Horde/HordeDashboard/src/components/jobDetailsV2/JobDetailViewSteps.tsx b/Engine/Source/Programs/Horde/HordeDashboard/src/components/jobDetailsV2/JobDetailViewSteps.tsx
--- a/Engine/Source/Programs/Horde/HordeDashboard/src/components/jobDetailsV2/JobDetailViewSteps.tsx
+++ b/Engine/Source/Programs/Horde/HordeDashboard/src/components/jobDetailsV2/JobDetailViewSteps.tsx
@@ -1,6 +1,6 @@
 // Copyright Epic Games, Inc. All Rights Reserved.
 
-import { CollapseAllVisibility, ConstrainMode, DetailsHeader, DetailsList, DetailsListLayoutMode, DetailsRow, DirectionalHint, IColumn, Icon, IDetailsListProps, ProgressIndicator, SelectionMode, Stack, Text, TooltipHost } from "@fluentui/react";
+import { CollapseAllVisibility, ConstrainMode, DetailsHeader, DetailsList, DetailsListLayoutMode, DetailsRow, DirectionalHint, FontIcon, IColumn, Icon, IDetailsListProps, ProgressIndicator, SelectionMode, Stack, Text, TooltipHost } from "@fluentui/react";
 import { observer } from "mobx-react-lite";
 import React, { useEffect, useId, useState } from "react";
 import { Link, useParams } from "react-router-dom";
@@ -61,7 +61,6 @@ JobDetailsV2.registerDataView("StepsDataView", (details: JobDetailsV2) => new St
 
 const RenderDynamic: React.FC<{ jobDetails: JobDetailsV2, dataView: StepsDataView, step?: StepData, batch?: GetBatchResponse, column: string }> = observer(({ jobDetails, dataView, column, step, batch }) => {
 
-
    dataView.subscribe();
 
    if (step && !step.finishTime && !step.abortRequested && (step.state !== JobStepState.Aborted && step.state !== JobStepState.Completed && step.state !== JobStepState.Skipped)) {
@@ -85,6 +84,29 @@ const RenderDynamic: React.FC<{ jobDetails: JobDetailsV2, dataView: StepsDataVie
          return <Stack horizontal horizontalAlign={"center"} tokens={{ childrenGap: 0, padding: 0 }} style={{ fontSize: "13px" }}><Link to={`/job/${jobDetails.jobId!}?step=${retry.id}`}>{`Retried - ${getStepStatusMessage(retry)}`}</Link></Stack>;
       }
 
+
+      // Handle step message
+      const report = step?.reports?.find(r => !!r.message);
+      if (report) {
+         const statusColors = dashboard.getStatusColors();
+         let color = "";
+         switch (report.severity) {
+            case "Error":
+            case "Errors":
+               color = statusColors.get(StatusColor.Failure)!;
+               break;
+            case "Warning":
+            case "Warnings":
+               color = statusColors.get(StatusColor.Warnings)!;
+               break;
+         }
+
+         return <Stack horizontal verticalFill verticalAlign="center" tokens={{ childrenGap: 6 }} horizontalAlign="center">
+            {!!color && <FontIcon style={{ fontSize: 13, color: color }} iconName="Square" />}
+            <Text variant="tiny" style={{ overflow: "hidden", textOverflow: "ellipsis", wordBreak: "normal" }}>{report.message}</Text>
+         </Stack>
+      }
+
       return <Stack horizontal horizontalAlign={"center"} tokens={{ childrenGap: 0, padding: 0 }}><Text>{getStepStatusMessage(step)}</Text></Stack>;
    }
 
diff --git a/Engine/Source/Programs/Horde/HordeDashboard/src/components/jobDetailsV2/StepDetailSummary.tsx b/Engine/Source/Programs/Horde/HordeDashboard/src/components/jobDetailsV2/StepDetailSummary.tsx
--- a/Engine/Source/Programs/Horde/HordeDashboard/src/components/jobDetailsV2/StepDetailSummary.tsx
+++ b/Engine/Source/Programs/Horde/HordeDashboard/src/components/jobDetailsV2/StepDetailSummary.tsx
@@ -1,6 +1,6 @@
 // Copyright Epic Games, Inc. All Rights Reserved.
 
-import { Label, Stack, Text } from '@fluentui/react';
+import { FontIcon, Label, Stack, Text } from '@fluentui/react';
 import { observer } from 'mobx-react-lite';
 import React, { useEffect } from 'react';
 import { JobStepError, JobStepOutcome, JobStepState, ReportPlacement } from '../../backend/Api';
@@ -10,6 +10,7 @@ import { getNiceTime, getStepElapsed, getStepETA, getStepFinishTime, getStepTimi
 import { getHordeStyling } from '../../styles/Styles';
 import { getBatchText } from '../JobDetailCommon';
 import { JobDataView, JobDetailsV2 } from './JobDetailsViewCommon';
+import dashboard, { StatusColor } from 'horde/backend/Dashboard';
 
 const sideRail: ISideRailLink = { text: "Summary", url: "rail_step_summary" };
 
@@ -240,6 +241,27 @@ export const StepSummaryPanel: React.FC<{ jobDetails: JobDetailsV2; stepId: stri
    const spawnedJobUrl = spawnedJobs.length > 0 ? `/job/${spawnedJobs[0]}` : undefined;
    const parentJobUrl = jobData.parentJobId && jobData.parentJobId ? `/job/${jobData.parentJobId}?step=${jobData.parentJobStepId}` : undefined;
 
+   let stepMessage: string | undefined;
+   let stepMessageColor: string | undefined;
+
+   const report = step?.reports?.find(r => !!r.message);
+   if (report) {
+
+      const statusColors = dashboard.getStatusColors();
+      stepMessage = report.message
+      stepMessageColor = "";
+      switch (report.severity) {
+         case "Error":
+         case "Errors":
+            stepMessageColor = statusColors.get(StatusColor.Failure)!;
+            break;
+         case "Warning":
+         case "Warnings":
+            stepMessageColor = statusColors.get(StatusColor.Warnings)!;
+            break;
+      }
+   }
+
    return (<Stack id={sideRail.url} styles={{ root: { paddingTop: 0, paddingRight: 12 } }}>
       <Stack className={hordeClasses.raised}>
          <Stack tokens={{ childrenGap: 18 }}>
@@ -257,8 +279,15 @@ export const StepSummaryPanel: React.FC<{ jobDetails: JobDetailsV2; stepId: stri
                   <Markdown>{step.cancellationReason}</Markdown>
                </Stack>
                }
-               {!!parentJobUrl && <Stack style={{ paddingTop: 8}}><Markdown>{`This step was spawned by a [parent job](${parentJobUrl}).`}</Markdown></Stack>}
+               {!!parentJobUrl && <Stack style={{ paddingTop: 8 }}><Markdown>{`This step was spawned by a [parent job](${parentJobUrl}).`}</Markdown></Stack>}
                {!!spawnedJobUrl && <Stack style={{ paddingTop: 8 }}><Markdown>{`This step spawned a [new job](${spawnedJobUrl}).`}</Markdown></Stack>}
+               {!!stepMessage && <Stack verticalFill verticalAlign="center" style={{ paddingTop: 12 }} tokens={{ childrenGap: 6 }}>
+                  <Text variant='medium' style={{ fontWeight: "bolder" }}>Step Message:</Text>
+                  <Stack horizontal verticalFill verticalAlign="center" style={{ paddingLeft: 24, paddingTop: 2 }} tokens={{ childrenGap: 6 }}>
+                     {!!stepMessageColor && <FontIcon style={{ fontSize: 13, color: stepMessageColor }} iconName="Square" />}
+                     <Text>{stepMessage}</Text>
+                  </Stack>
+               </Stack>}
                {!!reportData && <Stack style={{ paddingTop: 8 }}> <Markdown>{reportData}</Markdown> </Stack>}
                {!!priceText && <Stack style={{ paddingTop: 8 }}>
                   <Text>{priceText}</Text>
diff --git a/Engine/Source/Programs/Horde/HordeDashboard/src/components/logs/LogView.tsx b/Engine/Source/Programs/Horde/HordeDashboard/src/components/logs/LogView.tsx
--- a/Engine/Source/Programs/Horde/HordeDashboard/src/components/logs/LogView.tsx
+++ b/Engine/Source/Programs/Horde/HordeDashboard/src/components/logs/LogView.tsx
@@ -1,6 +1,6 @@
 // Copyright Epic Games, Inc. All Rights Reserved.
 
-import { Callout, DefaultButton, DetailsList, DetailsListLayoutMode, DetailsRow, DirectionalHint, Dropdown, FocusZone, FocusZoneDirection, FontIcon, IColumn, Icon, IconButton, IContextualMenuItem, IContextualMenuProps, IDetailsListProps, ITextField, List, mergeStyles, MessageBar, MessageBarType, Modal, PrimaryButton, ProgressIndicator, ScrollToMode, Selection, SelectionMode, SelectionZone, Separator, Spinner, SpinnerSize, Stack, Text, TextField, TooltipHost } from '@fluentui/react';
+import { Callout, DefaultButton, DetailsList, DetailsListLayoutMode, DetailsRow, DirectionalHint, Dropdown, FocusZone, FocusZoneDirection, FontIcon, IColumn, Icon, IconButton, IContextualMenuItem, IContextualMenuProps, IDetailsListProps, ITextField, List, mergeStyles, MessageBar, MessageBarType, Modal, PrimaryButton, ProgressIndicator, ScrollToMode, Selection, SelectionMode, SelectionZone, Separator, setMonth, Spinner, SpinnerSize, Stack, Text, TextField, TooltipHost } from '@fluentui/react';
 import { observer } from 'mobx-react-lite';
 import moment from 'moment-timezone';
 import React, { useEffect, useId, useState } from 'react';
@@ -8,7 +8,7 @@ import { Link, useLocation, useNavigate, useParams } from 'react-router-dom';
 import backend from '../../backend';
 import { ArtifactContextType, DashboardPreference, EventSeverity, GetChangeSummaryResponse, GetJobStepRefResponse, GetLogEventResponse, LogLevel } from '../../backend/Api';
 import { CommitCache } from '../../backend/CommitCache';
-import dashboard from '../../backend/Dashboard';
+import dashboard, { StatusColor } from '../../backend/Dashboard';
 import { Markdown } from '../../base/components/Markdown';
 import { useWindowSize } from '../../base/utilities/hooks';
 import { displayTimeZone, getElapsedString } from '../../base/utilities/timeUtils';
@@ -409,9 +409,9 @@ const LogList: React.FC<{ logId: string }> = observer(({ logId }) => {
    function searchSelectedText() {
       const selection = window.getSelection();
       const selected_text = selection ? selection.toString() : '';
-      
+
       if (selected_text) {
-         setSearchState({...searchState, text: selected_text})
+         setSearchState({ ...searchState, text: selected_text })
       }
    }
 
@@ -434,7 +434,7 @@ const LogList: React.FC<{ logId: string }> = observer(({ logId }) => {
             e.preventDefault();
 
             searchSelectedText();
-            
+
             if (searchState.text !== '') {
                searchBox.current?.setSelectionRange(0, searchBox.current.value!.length)
             }
@@ -480,7 +480,7 @@ const LogList: React.FC<{ logId: string }> = observer(({ logId }) => {
          <Stack horizontal style={{ paddingTop: 48 }}>
             <div key={`windowsize_logview1_${windowSize.width}_${windowSize.height}`} style={{ width: vw / 2 - 720, flexShrink: 0, backgroundColor: hordeTheme.horde.contentBackground }} />
             <Stack horizontalAlign="center" style={{ width: 1440 }}>
-               <div style={{whiteSpace:"pre"}}> {error} </div>
+               <div style={{ whiteSpace: "pre" }}> {error} </div>
             </Stack>
          </Stack>
       )
@@ -488,7 +488,7 @@ const LogList: React.FC<{ logId: string }> = observer(({ logId }) => {
    }
 
    // Important: This should come before handler init below
-   if(logError) {
+   if (logError) {
       return (
          <Stack horizontalAlign='center' style={{ paddingTop: 24 }} >
             <Text variant='mediumPlus'>{`Unable to load log data - ${logError.message}`}</Text>
@@ -496,7 +496,7 @@ const LogList: React.FC<{ logId: string }> = observer(({ logId }) => {
       )
    }
 
-   if(!handler){
+   if (!handler) {
       let lineIndexQuery: number | undefined;
       if (query.get("lineIndex")) {
          lineIndexQuery = parseInt(query.get("lineIndex")!);
@@ -517,7 +517,7 @@ const LogList: React.FC<{ logId: string }> = observer(({ logId }) => {
 
    // subscribe
    if (handler.updated) { }
-   
+
    const logSource = handler.logSource!;
    const logSourceIsJobLogSource = logSource instanceof JobLogSource;
    const warnings = logSource.warnings.sort((a, b) => a.lineIndex - b.lineIndex);
@@ -618,10 +618,10 @@ const LogList: React.FC<{ logId: string }> = observer(({ logId }) => {
             const href = `${location.pathname}?${newQuery.toString()}`;
 
             const fontStyle = {
-               fontFamily: "Horde Open Sans SemiBold, sans-serif, sans-serif", color: error ? errorButtonTextColor: warningButtonTextColor, fontSize: handler.fontSize, textDecoration: !item!.issue?.resolvedAt ? undefined : "line-through"
+               fontFamily: "Horde Open Sans SemiBold, sans-serif, sans-serif", color: error ? errorButtonTextColor : warningButtonTextColor, fontSize: handler.fontSize, textDecoration: !item!.issue?.resolvedAt ? undefined : "line-through"
             }
 
-            return <Stack style={{width: 80}}>
+            return <Stack style={{ width: 80 }}>
                <TooltipHost
                   content={timestamp}
                   id={tooltipId}
@@ -739,7 +739,7 @@ const LogList: React.FC<{ logId: string }> = observer(({ logId }) => {
    let summaryText = logSource.summary;
 
    let largeLogWarning = logSource.logItems?.length > 1500000;
-   let largeLogWarningText= `Warning: large log file (${logSource.logItems?.length} lines) may not display entirely.  Please download to view full log and consider reducing verbosity.`;
+   let largeLogWarningText = `Warning: large log file (${logSource.logItems?.length} lines) may not display entirely.  Please download to view full log and consider reducing verbosity.`;
 
    let baseUrl = location.pathname;
 
@@ -760,10 +760,10 @@ const LogList: React.FC<{ logId: string }> = observer(({ logId }) => {
          logListKey++;
 
          setSearchState({
-               search: '',
-               curRequest: undefined,
-               results: undefined,
-               text: ''
+            search: '',
+            curRequest: undefined,
+            results: undefined,
+            text: ''
          });
 
          // Reset current line and scroll to value before search was cleared
@@ -776,7 +776,8 @@ const LogList: React.FC<{ logId: string }> = observer(({ logId }) => {
          updateLineQuery(currentLine, navigate);
 
          inTimeout = true;
-         setTimeout(() => {inTimeout = false; listRef?.scrollToIndex(lineIdx, () => { return handler.lineHeight; }, ScrollToMode.top);
+         setTimeout(() => {
+            inTimeout = false; listRef?.scrollToIndex(lineIdx, () => { return handler.lineHeight; }, ScrollToMode.top);
          }, 10)
 
       } else if (newValue !== searchState.search) {
@@ -1143,41 +1144,55 @@ const LogList: React.FC<{ logId: string }> = observer(({ logId }) => {
    const prevWarningBlockEnabled: boolean = !!handler?.getPrevLogEventBlock(true);
 
    let jobDetails: JobDetailsV2 | undefined;
-   if(logSourceIsJobLogSource) {
+   if (logSourceIsJobLogSource) {
       jobDetails = (logSource as JobLogSource).jobDetails;
    }
 
+   let stepMessage = "";
+   let stepMessageColor = "";
+
+   const report = jobDetails?.stepByLogId(logId)?.reports?.find(r => !!r.message);
+   if (report) {
+      stepMessage = report.message!;
+      if (report.severity === "Error" || report.severity === "Errors") {
+         stepMessageColor = dashboard.getStatusColors().get(StatusColor.Failure)!
+      }
+      else if (report.severity === "Warning" || report.severity === "Warnings") {
+         stepMessageColor = dashboard.getStatusColors().get(StatusColor.Warnings)!
+      }
+   }
+
    return <Stack>
       {logSourceIsJobLogSource && <IssueModalV2 issueId={query.get("issue")} popHistoryOnClose={issueHistory} />}
-      
+
       {
-         logSourceIsJobLogSource 
-         && logHistory 
-         && 
-            <StepHistoryModal 
-               jobDetails={jobDetails!} 
-               stepId={jobDetails!.stepByLogId(logId)?.id} 
-               onClose={() => setLogHistory(false)} 
-            />
+         logSourceIsJobLogSource
+         && logHistory
+         &&
+         <StepHistoryModal
+            jobDetails={jobDetails!}
+            stepId={jobDetails!.stepByLogId(logId)?.id}
+            onClose={() => setLogHistory(false)}
+         />
       }
-      
+
       {
-         logSourceIsJobLogSource 
+         logSourceIsJobLogSource
          && !!artifactContext
-         && <JobArtifactsModal 
-               jobId={jobDetails!.jobData!.id} 
-               stepId={jobDetails!.stepByLogId(logId)?.id!} 
-               artifacts={(logSource as JobLogSource).artifactsV2} 
-               contextType={artifactContext} artifactPath={artifactPath} 
-               artifactId={artifactId} 
-               onClose={() => { navigate(window.location.pathname, { replace: true }) }} 
-            /> 
+         && <JobArtifactsModal
+            jobId={jobDetails!.jobData!.id}
+            stepId={jobDetails!.stepByLogId(logId)?.id!}
+            artifacts={(logSource as JobLogSource).artifactsV2}
+            contextType={artifactContext} artifactPath={artifactPath}
+            artifactId={artifactId}
+            onClose={() => { navigate(window.location.pathname, { replace: true }) }}
+         />
       }
 
       {!!historyAgentId && <HistoryModal agentId={historyAgentId} onDismiss={() => { navigate(baseUrl, { replace: true }); setHistoryAgentId(undefined) }} />}
-      
+
       <Breadcrumbs items={logSource?.crumbs ?? []} title={logSource?.crumbTitle} />
-      
+
       <Stack tokens={{ childrenGap: 0 }} style={{ backgroundColor: hordeTheme.horde.contentBackground, paddingTop: 12 }}>
          <Stack horizontal >
             <div key={`windowsize_logview1_${windowSize.width}_${windowSize.height}`} style={{ width: vw / 2 - (1440 / 2), flexShrink: 0 }} />
@@ -1187,8 +1202,8 @@ const LogList: React.FC<{ logId: string }> = observer(({ logId }) => {
                      <Stack horizontal tokens={{ childrenGap: 2 }}>
 
                         <Stack horizontal verticalAlign='center'>
-                           <Stack horizontal verticalAlign='center' style={{position: 'relative'}}>
-                              <TooltipHost content={nextErrorEnabled ? 'Navigate to next error': ''}>
+                           <Stack horizontal verticalAlign='center' style={{ position: 'relative' }}>
+                              <TooltipHost content={nextErrorEnabled ? 'Navigate to next error' : ''}>
                                  <DefaultButton disabled={!nextErrorEnabled} className={nextErrorEnabled ? handler.style.errorButton : handler.style.errorButtonDisabled}
                                     text={`${errorText} ${errors.length === 1 ? "Error" : "Errors"}`}
                                     onClick={() => {
@@ -1198,7 +1213,7 @@ const LogList: React.FC<{ logId: string }> = observer(({ logId }) => {
                                     {!!errors.length && <Icon style={{ fontSize: 19, paddingLeft: 4, paddingRight: 0 }} iconName='ChevronDown' />}
                                  </DefaultButton>
                               </TooltipHost>
-                              {!!errors.length && <div style={{backgroundColor: errorButtonTextColor, height: 18, width: 1, position: 'absolute', right: -0.5, zIndex: 1 }} />}
+                              {!!errors.length && <div style={{ backgroundColor: errorButtonTextColor, height: 18, width: 1, position: 'absolute', right: -0.5, zIndex: 1 }} />}
                            </Stack>
                            {!!errors.length && <TooltipHost content={'Navigate to previous error'}>
                               <IconButton disabled={!prevErrorEnabled} className={prevErrorEnabled ? handler.style.errorButton : handler.style.errorButtonDisabled} style={{ color: errorButtonTextColor, height: 30, fontSize: 19, padding: 2, borderRadius: '0 4px 4px 0' }} iconProps={{ iconName: 'ChevronUp' }} onClick={(event: any) => {
@@ -1208,22 +1223,22 @@ const LogList: React.FC<{ logId: string }> = observer(({ logId }) => {
                            </TooltipHost>}
                         </Stack>
 
-                        {!!errors.length && <Stack horizontal verticalAlign='center' style={{paddingLeft: 1}} >
-                           <Stack horizontal verticalAlign='center' style={{position: 'relative'}}>
+                        {!!errors.length && <Stack horizontal verticalAlign='center' style={{ paddingLeft: 1 }} >
+                           <Stack horizontal verticalAlign='center' style={{ position: 'relative' }}>
                               <TooltipHost content={'Skip to next block of errors or issue'}>
                                  <IconButton disabled={!nextErrorBlockEnabled} className={nextErrorBlockEnabled ? handler.style.errorButton : handler.style.errorButtonDisabled}
                                     style={{ color: errorButtonTextColor, height: 30, fontSize: 19, padding: 9, borderRadius: '4px 0 0 4px' }}
-                                    styles={{rootPressed: {border: 'none', boxShadow: 'none'}}} 
+                                    styles={{ rootPressed: { border: 'none', boxShadow: 'none' } }}
                                     iconProps={{ iconName: 'DoubleChevronDown' }}
                                     onClick={() => {
                                        nextErrorBlock();
                                     }} />
                               </TooltipHost>
-                                 <div style={{backgroundColor: errorButtonTextColor, height: 18, width: 1, position: 'absolute', right: -0.5, zIndex: 1 }}></div>
+                              <div style={{ backgroundColor: errorButtonTextColor, height: 18, width: 1, position: 'absolute', right: -0.5, zIndex: 1 }}></div>
                            </Stack>
                            <TooltipHost content={'Skip to previous block of errors or issue'}>
                               <IconButton disabled={!prevErrorBlockEnabled} className={prevErrorBlockEnabled ? handler.style.errorButton : handler.style.errorButtonDisabled}
-                                 style={{ color: errorButtonTextColor, height: 30, fontSize: 19, padding: 9,  borderRadius: '0 4px 4px 0'}} 
+                                 style={{ color: errorButtonTextColor, height: 30, fontSize: 19, padding: 9, borderRadius: '0 4px 4px 0' }}
                                  iconProps={{ iconName: 'DoubleChevronUp' }}
                                  onClick={(event: any) => {
                                     event?.stopPropagation();
@@ -1234,22 +1249,22 @@ const LogList: React.FC<{ logId: string }> = observer(({ logId }) => {
 
 
                         <Stack horizontal verticalAlign='center' style={{ paddingLeft: 12 }} >
-                           <Stack horizontal verticalAlign='center' style={{position: 'relative'}}>
-                              <TooltipHost content={nextWarningEnabled ? 'Navigate to next warning': ''}>
+                           <Stack horizontal verticalAlign='center' style={{ position: 'relative' }}>
+                              <TooltipHost content={nextWarningEnabled ? 'Navigate to next warning' : ''}>
                                  <DefaultButton disabled={!nextWarningEnabled} className={nextWarningEnabled ? handler.style.warningButton : handler.style.warningButtonDisabled}
                                     text={`${warningsText} ${warnings.length === 1 ? "Warning" : "Warnings"}`}
                                     onClick={() => {
                                        nextWarning();
                                     }}
-                                    style={{ color: warnings.length ? warningButtonTextColor : buttonNoneTextColor, height: 30, fontSize: 19, paddingLeft:  warnings.length ? 12 : 8, borderRadius: warnings.length ? '4px 0 0 4px': '4px' }} >
+                                    style={{ color: warnings.length ? warningButtonTextColor : buttonNoneTextColor, height: 30, fontSize: 19, paddingLeft: warnings.length ? 12 : 8, borderRadius: warnings.length ? '4px 0 0 4px' : '4px' }} >
                                     {!!warnings.length && <Icon style={{ fontSize: 19, paddingLeft: 4, paddingRight: 0 }} iconName='ChevronDown' />}
                                  </DefaultButton>
-                              </TooltipHost>     
-                              {!!warnings.length && <div style={{backgroundColor: warningButtonTextColor, height: 18, width: 1, position: 'absolute', right: -0.5, zIndex: 1 }} />}
+                              </TooltipHost>
+                              {!!warnings.length && <div style={{ backgroundColor: warningButtonTextColor, height: 18, width: 1, position: 'absolute', right: -0.5, zIndex: 1 }} />}
                            </Stack>
                            {!!warnings.length && <Stack>
                               <TooltipHost content={'Navigate to previous warning'}>
-                                 <IconButton disabled={!prevWarningEnabled} className={prevWarningEnabled ? handler.style.warningButton : handler.style.warningButtonDisabled} style={{ height: 30, fontSize: 19, padding: 2, color: warningButtonTextColor, borderRadius: '0 4px 4px 0'}} iconProps={{ iconName: 'ChevronUp' }} onClick={(event: any) => {
+                                 <IconButton disabled={!prevWarningEnabled} className={prevWarningEnabled ? handler.style.warningButton : handler.style.warningButtonDisabled} style={{ height: 30, fontSize: 19, padding: 2, color: warningButtonTextColor, borderRadius: '0 4px 4px 0' }} iconProps={{ iconName: 'ChevronUp' }} onClick={(event: any) => {
                                     event?.stopPropagation();
                                     prevWarning();
                                  }} />
@@ -1257,22 +1272,22 @@ const LogList: React.FC<{ logId: string }> = observer(({ logId }) => {
                            </Stack>}
                         </Stack>
 
-                        {!!warnings.length && <Stack horizontal verticalAlign='center' style={{paddingLeft: 1}} >
-                           <Stack horizontal verticalAlign='center' style={{position: 'relative'}}>
+                        {!!warnings.length && <Stack horizontal verticalAlign='center' style={{ paddingLeft: 1 }} >
+                           <Stack horizontal verticalAlign='center' style={{ position: 'relative' }}>
                               <TooltipHost content={'Skip to next block of warnings or issue'}>
                                  <IconButton disabled={!nextWarningBlockEnabled} className={nextWarningBlockEnabled ? handler.style.warningButton : handler.style.warningButtonDisabled}
                                     style={{ color: warningButtonTextColor, height: 30, fontSize: 19, padding: 9, borderRadius: '4px 0 0 4px' }}
-                                    styles={{rootPressed: {border: 'none', boxShadow: 'none'}}} 
+                                    styles={{ rootPressed: { border: 'none', boxShadow: 'none' } }}
                                     iconProps={{ iconName: 'DoubleChevronDown' }}
                                     onClick={() => {
                                        nextWarningBlock();
                                     }} />
                               </TooltipHost>
-                                 <div style={{backgroundColor: warningButtonTextColor, height: 18, width: 1, position: 'absolute', right: -0.5, zIndex: 1 }}></div>
+                              <div style={{ backgroundColor: warningButtonTextColor, height: 18, width: 1, position: 'absolute', right: -0.5, zIndex: 1 }}></div>
                            </Stack>
                            <TooltipHost content={'Skip to previous block of warnings or issue'}>
                               <IconButton disabled={!prevWarningBlockEnabled} className={prevWarningBlockEnabled ? handler.style.warningButton : handler.style.warningButtonDisabled}
-                                 style={{ color: warningButtonTextColor, height: 30, fontSize: 19, padding: 9,  borderRadius: '0 4px 4px 0'}} 
+                                 style={{ color: warningButtonTextColor, height: 30, fontSize: 19, padding: 9, borderRadius: '0 4px 4px 0' }}
                                  iconProps={{ iconName: 'DoubleChevronUp' }}
                                  onClick={(event: any) => {
                                     event?.stopPropagation();
@@ -1287,6 +1302,16 @@ const LogList: React.FC<{ logId: string }> = observer(({ logId }) => {
                      </Stack>
                      <Stack>
                         {logSource.active && <LogProgressIndicator logSource={logSource} />}
+                        {!!stepMessage && <Stack horizontalAlign='center'>
+                           <Stack style={{ paddingLeft: 4 }} horizontal verticalAlign="center" tokens={{ childrenGap: 4 }}>
+                              <Text style={{ fontSize: 11, fontFamily: "Horde Open Sans Semibold" }}>Step:</Text>
+                              <Stack horizontal verticalFill verticalAlign="center" style={{ paddingTop: 2 }} tokens={{ childrenGap: 4 }}>
+                                 {!!stepMessageColor && <FontIcon style={{ fontSize: 13, color: stepMessageColor, paddingTop: 1 }} iconName="Square" />}
+                                 <Text style={{ maxWidth: 190, fontSize: 11, fontFamily: "Horde Open Sans Semibold", overflow: "hidden", textOverflow: "ellipsis" }}>{stepMessage}</Text>
+                              </Stack>
+                           </Stack>
+                        </Stack>}
+
                      </Stack>
                   </Stack>
 
@@ -1305,23 +1330,23 @@ const LogList: React.FC<{ logId: string }> = observer(({ logId }) => {
                                     validateOnLoad={false}
                                     componentRef={searchBox}
                                     value={searchState.text ? searchState.text : ''}
-                                    onChange={(ev, value) => { setSearchState({...searchState, text: value}); }}
-                                    
+                                    onChange={(ev, value) => { setSearchState({ ...searchState, text: value }); }}
+
                                     onRenderSuffix={
                                        !!searchState.text ? () => (
                                           <Stack verticalAlign="center" horizontalAlign="center" horizontal tokens={{ childrenGap: 2 }} styles={{ root: { padding: 0 } }}>
-                                             <IconButton iconProps={{ iconName: 'Cancel'}} styles={{ root: { padding: 12, height: 20, width: 20 }, icon: { fontSize: 15 } }}
-                                             onClick={() => { doQuery('') }}
+                                             <IconButton iconProps={{ iconName: 'Cancel' }} styles={{ root: { padding: 12, height: 20, width: 20 }, icon: { fontSize: 15 } }}
+                                                onClick={() => { doQuery('') }}
                                              />
-                                             <span className={mergeStyles({ paddingLeft: 1, paddingRight: 7, paddingBottom: 1})}>
+                                             <span className={mergeStyles({ paddingLeft: 1, paddingRight: 7, paddingBottom: 1 })}>
                                                 {searchState?.results ? searchState?.results?.length ? `${curSearchIdx + 1}/${searchState?.results?.length}` : '0/0' : undefined}
-                                             </span>               
+                                             </span>
                                           </Stack>
                                        ) : undefined
                                     }
 
                                     styles={{
-                                       root: { width: 280, fontSize: 12 }, 
+                                       root: { width: 280, fontSize: 12 },
                                        fieldGroup: { borderWidth: 1 },
                                        suffix: { padding: '0 2px', height: 30 }
                                     }}
@@ -1393,7 +1418,7 @@ const LogList: React.FC<{ logId: string }> = observer(({ logId }) => {
                      </Stack>
                   </Stack>
                   {logTelemetry && !!logSource.agentTelemetry && <AgentTelemetrySparkline handler={logSource.agentTelemetry} />}
-                  {largeLogWarning && <Stack style={{ paddingBottom: 12, paddingTop: 12, margin: 'auto'}}>
+                  {largeLogWarning && <Stack style={{ paddingBottom: 12, paddingTop: 12, margin: 'auto' }}>
                      <MessageBar
                         messageBarType={MessageBarType.warning} isMultiline={false} >
                         <Markdown styles={{ root: { fontSize: 12, fontFamily: "Horde Open Sans Bold" } }}>{largeLogWarningText}</Markdown>
