[Horde][BuildHealth][Dashboard] State filter & minor fix for batch sync failed.

diff --git a/Engine/Source/Programs/Horde/HordeDashboard/src/components/buildhealth/BuildHealthDataHandler.ts b/Engine/Source/Programs/Horde/HordeDashboard/src/components/buildhealth/BuildHealthDataHandler.ts
--- a/Engine/Source/Programs/Horde/HordeDashboard/src/components/buildhealth/BuildHealthDataHandler.ts
+++ b/Engine/Source/Programs/Horde/HordeDashboard/src/components/buildhealth/BuildHealthDataHandler.ts
@@ -20,6 +20,7 @@ export type DataHandlerRefreshRequest = {
     steps?: boolean;
 }
 
+export const EMPTY_REFRESH_REQUEST: DataHandlerRefreshRequest = { projects: false, streams: false, jobs: false, steps: false };
 export const DEFAULT_REFRESH_REQUEST: DataHandlerRefreshRequest = { projects: true, streams: true, jobs: true, steps: true };
 
 /**
diff --git a/Engine/Source/Programs/Horde/HordeDashboard/src/components/buildhealth/BuildHealthDataTypes.ts b/Engine/Source/Programs/Horde/HordeDashboard/src/components/buildhealth/BuildHealthDataTypes.ts
--- a/Engine/Source/Programs/Horde/HordeDashboard/src/components/buildhealth/BuildHealthDataTypes.ts
+++ b/Engine/Source/Programs/Horde/HordeDashboard/src/components/buildhealth/BuildHealthDataTypes.ts
@@ -7,12 +7,14 @@ export const PROJECTS_URL_PARAM_NAME: string = "projects";
 export const STREAMS_URL_PARAM_NAME: string = "streams";
 export const TEMPLATES_URL_PARAM_NAME: string = "templates";
 export const STEPS_URL_PARAM_NAME: string = "stepNames";
+export const STATES_URL_PARAM_NAME: string = "states";
 export const PARAMETER_KEY_PREFIX: string = "parameter_key";
 export const TIME_SPAN_URL_PARAM: string = "lastTimeRange";
 export const INCLUDE_PREFLIGHT_PARAM: string = "includePreflight";
 export const KEY_SEPARATOR: string = "::";
 
 export const STEPS_URL_PARAM_ALL_STEPS: string = "stepNames_AllKnownValues";
+export const STATES_URL_PARAM_ALL_STATES: string = "stateNames_AllKnownValues";
 
 /**
  * Convenience type to fold in extra metadata relevant for the Template. 
diff --git a/Engine/Source/Programs/Horde/HordeDashboard/src/components/buildhealth/BuildHealthOptions.ts b/Engine/Source/Programs/Horde/HordeDashboard/src/components/buildhealth/BuildHealthOptions.ts
--- a/Engine/Source/Programs/Horde/HordeDashboard/src/components/buildhealth/BuildHealthOptions.ts
+++ b/Engine/Source/Programs/Horde/HordeDashboard/src/components/buildhealth/BuildHealthOptions.ts
@@ -1,12 +1,13 @@
 // Copyright Epic Games, Inc. All Rights Reserved.
 
 import { makeObservable, observable, runInAction } from "mobx";
-import { INCLUDE_PREFLIGHT_PARAM, PROJECTS_URL_PARAM_NAME, QUERY_SCHEMA_PARAM_NAME, STEPS_URL_PARAM_ALL_STEPS, STEPS_URL_PARAM_NAME, STREAMS_URL_PARAM_NAME, TEMPLATES_URL_PARAM_NAME, TIME_SPAN_URL_PARAM } from "../buildhealth/BuildHealthDataTypes";
+import { INCLUDE_PREFLIGHT_PARAM, PROJECTS_URL_PARAM_NAME, QUERY_SCHEMA_PARAM_NAME, STATES_URL_PARAM_ALL_STATES, STATES_URL_PARAM_NAME, STEPS_URL_PARAM_ALL_STEPS, STEPS_URL_PARAM_NAME, STREAMS_URL_PARAM_NAME, TEMPLATES_URL_PARAM_NAME, TIME_SPAN_URL_PARAM } from "../buildhealth/BuildHealthDataTypes";
 import { BuildHealthDataHandler } from "./BuildHealthDataHandler";
-import { decodeStepKey, decodeTemplateKey, encodeStepKey, encodeStepNameFromStringsPreserveCase, encodeTemplateKey, encodeTemplateKeyFromStrings } from "./BuildHealthUtilities";
+import { decodeStateKey, decodeStepKey, decodeTemplateKey, encodeStateKeyFromString, encodeStepKey, encodeStepNameFromStringsPreserveCase, encodeTemplateKey, encodeTemplateKeyFromStrings } from "./BuildHealthUtilities";
 import { Location } from 'react-router-dom';
 import LZString from "lz-string"
 import { BuildHealthOptionsController, IBuildHealthOptionWriter } from "./BuildHealthOptionsController";
+import { JobStepState } from "horde/backend/Api";
 
 // #region -- Options Utilities & Data Types --
 
@@ -33,6 +34,7 @@ type BuildHealthQueryParams = {
     streamIds: string[];
     templateIds: string[];
     stepNames: string[];
+    stateNames: string[];
     timeSpanKey: string;
     includePreflight: boolean;
     querySchemaVersion: number;
@@ -113,17 +115,19 @@ export function toNavigationQuery(optionsState: BuildHealthOptionsState, querySc
     let streams: string = "";
     let jobs: string = "";
     let steps: string = "";
+    let states: string = "";
 
     projects = Object.keys(optionsState.enabledProjects).join(",");
     streams = optionsState.stepOutcomeEnabledStreamKeys.join(",");
     const lastTimeRange = optionsState.jobHistoryTimeSpan.key.toString();
     const includePreflight = optionsState.includePreflight;
 
+    // Job & Steps 
     if (querySchemaVersion === 1) {
         jobs = optionsState.stepOutcomeEnabledJobKeys.join(",");
         steps = optionsState.stepOutcomeEnabledStepKeys.join(",");
     }
-    else if (querySchemaVersion === 2) {
+    else if (querySchemaVersion >= 2) {
 
         jobs = compressKeys(optionsState.stepOutcomeEnabledJobKeys, (key: string) => {
             let { streamId, templateId } = decodeTemplateKey(key);
@@ -145,10 +149,16 @@ export function toNavigationQuery(optionsState: BuildHealthOptionsState, querySc
             });
     }
 
+    // States
+    if (querySchemaVersion === 3) {
+        states = optionsState.stepOutcomeEnabledsStates.length === Object.keys(JobStepState).length ? STATES_URL_PARAM_ALL_STATES : optionsState.stepOutcomeEnabledsStates.join(",");
+    }
+
     params.set(PROJECTS_URL_PARAM_NAME, LZString.compressToEncodedURIComponent(projects));
     params.set(STREAMS_URL_PARAM_NAME, LZString.compressToEncodedURIComponent(streams));
     params.set(TEMPLATES_URL_PARAM_NAME, LZString.compressToEncodedURIComponent(jobs));
     params.set(STEPS_URL_PARAM_NAME, LZString.compressToEncodedURIComponent(steps));
+    params.set(STATES_URL_PARAM_NAME, LZString.compressToEncodedURIComponent(states));
     params.set(TIME_SPAN_URL_PARAM, LZString.compressToEncodedURIComponent(lastTimeRange));
     params.set(INCLUDE_PREFLIGHT_PARAM, LZString.compressToEncodedURIComponent(String(includePreflight)));
     params.set(QUERY_SCHEMA_PARAM_NAME, querySchemaVersion.toString());
@@ -193,6 +203,7 @@ export function parseBuildHealthQueryParams(location: Location): BuildHealthQuer
     let streamIds: string[] = [];
     let templateIds: string[] = [];
     let stepNames: string[] = [];
+    let stateNames: string[] = [];
 
     const getList = (key: string) => {
         const raw = searchParams.get(key);
@@ -208,10 +219,11 @@ export function parseBuildHealthQueryParams(location: Location): BuildHealthQuer
     projectIds = getList(PROJECTS_URL_PARAM_NAME);
     streamIds = getList(STREAMS_URL_PARAM_NAME);
 
+    // Job & steps
     if (querySchemaVersion === 1) {
         templateIds = getList(TEMPLATES_URL_PARAM_NAME);
         stepNames = getList(STEPS_URL_PARAM_NAME);
-    } else if (querySchemaVersion === 2) {
+    } else if (querySchemaVersion >= 2) {
         const rawTemplates = searchParams.get(TEMPLATES_URL_PARAM_NAME);
         if (!rawTemplates) {
             templateIds = [];
@@ -237,6 +249,11 @@ export function parseBuildHealthQueryParams(location: Location): BuildHealthQuer
         }
     }
 
+    // States
+    if (querySchemaVersion === 3) {
+        stateNames = getList(STATES_URL_PARAM_NAME);
+    }
+
     let timeSpanKey: string = (() => {
         const raw = searchParams.get(TIME_SPAN_URL_PARAM);
         return raw ? LZString.decompressFromEncodedURIComponent(raw) ?? "" : "";
@@ -249,7 +266,7 @@ export function parseBuildHealthQueryParams(location: Location): BuildHealthQuer
     // Read from URLSearchParams
     const includePreflight = (searchParams.get(INCLUDE_PREFLIGHT_PARAM) ? LZString.decompressFromEncodedURIComponent(searchParams.get(INCLUDE_PREFLIGHT_PARAM)!) : "false") === "true";
 
-    return { projectIds, streamIds, templateIds: templateIds, stepNames, timeSpanKey, includePreflight, querySchemaVersion };
+    return { projectIds, streamIds, templateIds: templateIds, stepNames, stateNames, timeSpanKey, includePreflight, querySchemaVersion };
 }
 
 /**
@@ -288,6 +305,11 @@ export function upgradeBuildHealthQueryParams(params: BuildHealthQueryParams, cu
         clonedParams.stepNames.push(...upgradedSteps);
     }
 
+    // We must add in the state information, as no such filter existed before
+    if (clonedParams.querySchemaVersion < 3 && currentSchemaVersion >= 3) {
+        clonedParams.stateNames.push(encodeStateKeyFromString(-1, STATES_URL_PARAM_ALL_STATES));
+    }
+
     return clonedParams;
 }
 
@@ -313,6 +335,20 @@ export async function loadBuildHealthOptionsFromParams(
         });
     }
 
+    // Set state
+    if (params.stateNames) {
+        optionsWriter.clearStates();
+
+        if (params.stateNames.length === 1 && params.stateNames[0].includes(STATES_URL_PARAM_ALL_STATES)) {
+            optionsWriter.toggleAllStates(ALL_STATES_TUPLE_LIST.map(entry => entry.key), true);
+        } else if (params.stateNames.length > 0) {
+            for (let idx: number = 0; idx < params.stateNames.length; ++idx) {
+                let { index, stateName } = decodeStateKey(params.stateNames[idx]);
+                optionsWriter.toggleState(params.stateNames[idx], stateName, true);
+            }
+        }
+    }
+
     // Get baseline project data first
     await handler.requestHierarchicalRefresh();
 
@@ -393,6 +429,7 @@ export class BuildHealthUIOptionsState {
     @observable includeDateAnchors: boolean = false;
     @observable warningsAsSummaryFailure: boolean = false;
     @observable showLabels: boolean = false;
+    @observable validStates: Set<string> = new Set<string>();
 
     constructor() {
         makeObservable(this);
@@ -425,6 +462,7 @@ export class BuildHealthOptionStateReceipt {
     enabledStreams: ReadonlySet<string>;
     enabledTemplates: ReadonlySet<string>;
     enabledSteps: ReadonlySet<string>;
+    enabledStates: ReadonlySet<string>;
     enabledPreflight: boolean;
     timespan: TimeSpan;
 
@@ -433,6 +471,7 @@ export class BuildHealthOptionStateReceipt {
         streams: Iterable<string> = [],
         templates: Iterable<string> = [],
         steps: Iterable<string> = [],
+        states: Iterable<string> = [],
         enabledPreflight: boolean,
         timespan: TimeSpan
     ) {
@@ -440,6 +479,7 @@ export class BuildHealthOptionStateReceipt {
         this.enabledStreams = new Set(streams);
         this.enabledTemplates = new Set(templates);
         this.enabledSteps = new Set(steps);
+        this.enabledStates = new Set(states);
         this.enabledPreflight = enabledPreflight;
         this.timespan = timespan;
     }
@@ -460,6 +500,7 @@ export class BuildHealthOptionStateReceipt {
             BuildHealthOptionStateReceipt.difference(next.enabledStreams, prev.enabledStreams),
             BuildHealthOptionStateReceipt.difference(next.enabledTemplates, prev.enabledTemplates),
             BuildHealthOptionStateReceipt.difference(next.enabledSteps, prev.enabledSteps),
+            BuildHealthOptionStateReceipt.difference(next.enabledStates, prev.enabledStates),
             next.enabledPreflight,
             next.timespan
         );
@@ -469,6 +510,7 @@ export class BuildHealthOptionStateReceipt {
             BuildHealthOptionStateReceipt.difference(prev.enabledStreams, next.enabledStreams),
             BuildHealthOptionStateReceipt.difference(prev.enabledTemplates, next.enabledTemplates),
             BuildHealthOptionStateReceipt.difference(prev.enabledSteps, next.enabledSteps),
+            BuildHealthOptionStateReceipt.difference(prev.enabledStates, next.enabledStates),
             prev.enabledPreflight,
             prev.timespan
         );
@@ -488,6 +530,13 @@ export class BuildHealthOptionStateReceipt {
     }
 }
 
+export const ALL_STATES_TUPLE_LIST = Object.values(JobStepState).map(
+    (stateName, idx): { key: string; stateName: string } => ({
+        key: encodeStateKeyFromString(idx, stateName),
+        stateName: stateName
+    })
+);
+
 /**
  * Build Health Filter Options.
  */
@@ -498,6 +547,7 @@ export class BuildHealthOptionsState {
     @observable stepOutcomeEnabledStreamKeys: string[] = [];
     @observable stepOutcomeEnabledJobKeys: string[] = [];
     @observable stepOutcomeEnabledStepKeys: string[] = [];
+    @observable stepOutcomeEnabledsStates: string[] = [];
     @observable startDate!: Date;
     @observable endDate?: Date;
 
@@ -509,6 +559,7 @@ export class BuildHealthOptionsState {
     @observable enabledStreams: Record<string, string> = {};
     @observable enabledTemplates: Record<string, string> = {};
     @observable enabledSteps: Record<string, string> = {};
+    @observable enabledStates: Record<string, string> = {};
     @observable jobHistoryTimeSpan: TimeSpan = DEFAULT_JOB_HISTORY_TIME_SPAN;
     @observable includePreflight: boolean = false;
     @observable linkedSelectionModeEnabled: boolean = false;
@@ -540,6 +591,7 @@ export class BuildHealthOptionsState {
             Object.keys(this.enabledStreams),
             Object.keys(this.enabledTemplates),
             Object.keys(this.enabledSteps),
+            Object.values(this.enabledStates),
             this.includePreflight,
             this.jobHistoryTimeSpan
         );
diff --git a/Engine/Source/Programs/Horde/HordeDashboard/src/components/buildhealth/BuildHealthOptionsController.ts b/Engine/Source/Programs/Horde/HordeDashboard/src/components/buildhealth/BuildHealthOptionsController.ts
--- a/Engine/Source/Programs/Horde/HordeDashboard/src/components/buildhealth/BuildHealthOptionsController.ts
+++ b/Engine/Source/Programs/Horde/HordeDashboard/src/components/buildhealth/BuildHealthOptionsController.ts
@@ -1,7 +1,8 @@
 import { action, makeObservable } from "mobx";
 import { BuildHealthOptionsState, BuildHealthOptionStateDiff, BuildHealthOptionStateReceipt, BuildHealthUIOptionsState, TimeSpan } from "./BuildHealthOptions";
-import { decodeStepKey, decodeTemplateKey, encodeStepNameFromStringsPreserveCase, encodeTemplateKeyFromStrings } from "./BuildHealthUtilities";
+import { decodeStateKey, decodeStepKey, decodeTemplateKey, encodeStepNameFromStringsPreserveCase, encodeTemplateKeyFromStrings } from "./BuildHealthUtilities";
 import { BuildHealthDataHandler } from "./BuildHealthDataHandler";
+import { JobStepState } from "horde/backend/Api";
 
 /**
  * An interface that can provide hierarchy checks on stream, template, and step relationships.
@@ -123,6 +124,20 @@ export interface IBuildHealthOptionWriter {
      */
     toggleStep(id: string, name: string, enabled: boolean): void;
 
+    /**
+     * Toggles all provided states.
+     * @param enabled True to enable, false otherwise.
+     */
+    toggleAllStates(states: string[], enabled: boolean): void;
+
+    /**
+     * Toggles the state of the provided JobStepState.
+     * @param id The id of the state to toggle.
+     * @param name The name of the state.
+     * @param enabled Whether it's enabled or not.
+     */
+    toggleState(id: string, state: string, enabled: boolean): void;
+
     /**
      * Clears all currently selected options.
      */
@@ -152,6 +167,11 @@ export interface IBuildHealthOptionWriter {
      * @param owningHierarchy Whether to clear an item if it belongs to a stream & template. If not provided, will remove any step.
      */
     clearSteps(owningHierarchy?: { streamId: string; templateId: string }): void;
+
+    /**
+     * Clear all the states tracked by the options.
+     */
+    clearStates(): void;
 }
 
 export class BuildHealthOptionsWriter implements IBuildHealthOptionWriter {
@@ -182,6 +202,7 @@ export class BuildHealthOptionsWriter implements IBuildHealthOptionWriter {
             delete this.state.enabledProjects[key];
         });
         this.clearStreams(true);
+        this.clearStates();
     }
 
     /**
@@ -276,6 +297,18 @@ export class BuildHealthOptionsWriter implements IBuildHealthOptionWriter {
         }
     }
 
+    /**
+     * @inheritdoc
+     */
+    @action
+    toggleState(id: string, state: string, enabled: boolean) {
+        if (enabled) {
+            this.state.enabledStates[id] = state;
+        } else {
+            delete this.state.enabledStates[id];
+        }
+    }
+
     /**
      * @inheritdoc
      */
@@ -292,6 +325,17 @@ export class BuildHealthOptionsWriter implements IBuildHealthOptionWriter {
         }
     }
 
+    /**
+     * @inheritdoc
+     */
+    @action
+    toggleAllStates(states: string[], enabled: boolean) {
+        for (let idx: number = 0; idx < states.length; ++idx) {
+            let { index, stateName } = decodeStateKey(states[idx]);
+            this.toggleState(states[idx], stateName, enabled);
+        }
+    }
+
     /**
      * @inheritdoc
      */
@@ -370,6 +414,15 @@ export class BuildHealthOptionsWriter implements IBuildHealthOptionWriter {
         });
     }
 
+    /**
+     * @inheritdoc
+     */
+    clearStates() {
+        Object.keys(this.state.enabledStates).forEach(key => {
+            delete this.state.enabledStates[key];
+        });
+    }
+
     // #region -- Private API --
 
     private toggleTemplateInternal(id: string, name: string, enabled: boolean) {
@@ -438,7 +491,7 @@ export class BuildHealthOptionsController {
     // #region -- Public Members --
 
     optionsChangeVersion = 0;
-    querySchemaVersion: number = 2;
+    querySchemaVersion: number = 3;
 
     // #endregion -- Public Members --
 
@@ -545,9 +598,13 @@ export class BuildHealthOptionsController {
             this.state.stepOutcomeEnabledStreamKeys = Object.keys(this.state.enabledStreams);
             this.state.stepOutcomeEnabledJobKeys = Object.keys(this.state.enabledTemplates);
             this.state.stepOutcomeEnabledStepKeys = Object.keys(this.state.enabledSteps);
+            this.state.stepOutcomeEnabledsStates = Object.keys(this.state.enabledStates);
             this.state.startDate = new Date(new Date().valueOf() - (this.state.jobHistoryTimeSpan.minutes * 60000));
             this.state.endDate = new Date();
 
+            this.uiState.validStates.clear();
+            Object.values(this.state.enabledStates).forEach(x => this.uiState.validStates.add(x));
+
             this.lastSynchronize = this.optionsChangeVersion;
         }
     }
diff --git a/Engine/Source/Programs/Horde/HordeDashboard/src/components/buildhealth/BuildHealthUtilities.ts b/Engine/Source/Programs/Horde/HordeDashboard/src/components/buildhealth/BuildHealthUtilities.ts
--- a/Engine/Source/Programs/Horde/HordeDashboard/src/components/buildhealth/BuildHealthUtilities.ts
+++ b/Engine/Source/Programs/Horde/HordeDashboard/src/components/buildhealth/BuildHealthUtilities.ts
@@ -2,6 +2,23 @@
 
 import { KEY_SEPARATOR, StepRefData, TemplateRefData } from "./BuildHealthDataTypes";
 
+/**
+ * Produces an encoded state key for unique identification.
+ * @returns The key representing the state.
+ */
+export function encodeStateKeyFromString(index: number, stateName: string): string {
+    return `${index}${KEY_SEPARATOR}${stateName}`;
+}
+
+/**
+ * Decodes an encoed state key, providing it's constituent parts.
+ * @returns The constituent parts.
+ */
+export function decodeStateKey(encodedStateKey: string): { index: number; stateName: string } {
+    const [indexStr, stateName] = encodedStateKey.split(KEY_SEPARATOR);
+    return { index: Number(indexStr) || 0, stateName: stateName ?? '' };
+}
+
 /**
  * Produces an encoded template key for unique identification.
  * @param template The template to produce the key for.
diff --git a/Engine/Source/Programs/Horde/HordeDashboard/src/components/buildhealth/BuildhealthView.tsx b/Engine/Source/Programs/Horde/HordeDashboard/src/components/buildhealth/BuildhealthView.tsx
--- a/Engine/Source/Programs/Horde/HordeDashboard/src/components/buildhealth/BuildhealthView.tsx
+++ b/Engine/Source/Programs/Horde/HordeDashboard/src/components/buildhealth/BuildhealthView.tsx
@@ -1,7 +1,7 @@
 // Copyright Epic Games, Inc. All Rights Reserved.
 
 import { Checkbox, ComboBox, DefaultButton, DirectionalHint, Dropdown, DropdownMenuItemType, IconButton, IDropdownOption, IDropdownProps, ISelectableOption, Label, mergeStyleSets, PrimaryButton, SearchBox, Stack, TooltipHost } from "@fluentui/react";
-import { GetProjectResponse, GetStreamResponse } from "horde/backend/Api";
+import { GetProjectResponse, GetStreamResponse, JobStepState } from "horde/backend/Api";
 import { observer } from "mobx-react-lite";
 import { useCallback, useEffect, useMemo, useRef, useState } from "react";
 import { useLocation, useNavigate } from "react-router";
@@ -10,10 +10,10 @@ import { getHordeStyling } from "../../styles/Styles";
 import { Breadcrumbs } from "../Breadcrumbs";
 import { TopNav } from "../TopNav";
 import { StepOutcomeView } from "../buildhealth/stepoutcome/StepOutcome";
-import { BuildHealthDataHandler, DataHandlerRefreshRequest, expandRefreshRequest } from "./BuildHealthDataHandler";
+import { BuildHealthDataHandler, DataHandlerRefreshRequest, EMPTY_REFRESH_REQUEST, expandRefreshRequest } from "./BuildHealthDataHandler";
 import { KEY_SEPARATOR, PARAMETER_KEY_PREFIX, StepRefData, TemplateRefData } from "./BuildHealthDataTypes";
-import { BuildHealthOptionsState, BuildHealthOptionStateDiff, BuildHealthUIOptionsState, JobHistoryTimeSpans, loadBuildHealthOptionsFromParams, parseBuildHealthQueryParams, TimeSpan, toNavigationQuery, upgradeBuildHealthQueryParams } from "./BuildHealthOptions";
-import { decodeStepKey, decodeTemplateKey, encodeStepKey, encodeTemplateKey } from "./BuildHealthUtilities";
+import { ALL_STATES_TUPLE_LIST, BuildHealthOptionsState, BuildHealthOptionStateDiff, BuildHealthUIOptionsState, JobHistoryTimeSpans, loadBuildHealthOptionsFromParams, parseBuildHealthQueryParams, TimeSpan, toNavigationQuery, upgradeBuildHealthQueryParams } from "./BuildHealthOptions";
+import { decodeStepKey, decodeTemplateKey, encodeStateKeyFromString, encodeStepKey, encodeTemplateKey } from "./BuildHealthUtilities";
 import { stepOutcomeTableClasses } from "./stepoutcome/StepOutcomeSharedUIComponents";
 import { runInAction } from "mobx";
 import { BuildHealthOptionsController } from "./BuildHealthOptionsController";
@@ -699,6 +699,94 @@ const StepDropdown: React.FC<{ handler: BuildHealthDataHandler, buildHealthOptio
     );
 });
 
+const StateDropdown: React.FC<{ handler: BuildHealthDataHandler, buildHealthOptions: BuildHealthOptionsController }> = observer(function ConstructProjectDropdown({ handler, buildHealthOptions }) {
+    let selectedStateKeys: string[] = Object.keys(buildHealthOptions.state.enabledStates);
+
+    // #region -- Data Memos --
+
+    const items = useMemo<BuildHealthOptionData[]>(() =>
+        ALL_STATES_TUPLE_LIST.map((entry) => ({
+            id: entry.key,
+            text: entry.stateName,
+            group: "States"
+        })),
+        []
+    );
+
+    const dropdownOptions = useMemo<IDropdownOption<BuildHealthOptionData>[]>(() =>
+        items.map(item => ({
+            key: item.id,
+            text: item.text,
+            data: item,
+            group: item.group
+        })),
+        [items]
+    );
+
+    // #endregion
+
+    // #region -- Callbacks --
+
+    const handleChange = useCallback(
+        (option?: IDropdownOption<BuildHealthOptionData>) => {
+            if (!option) return;
+
+            const stateKey = option.key as string;
+            const stateName = option.text;
+            const isSelected = !!option.selected;
+
+            buildHealthOptions
+                .getTransactionSession()
+                .toggleState(stateKey, stateName, isSelected);
+        },
+        [buildHealthOptions]
+    );
+
+    const handleSelectAll = useCallback(
+        (
+            option?: IDropdownOption<BuildHealthOptionData>,
+            optionSelectAllContext?: IDropdownOption<BuildHealthOptionData>[]
+        ) => {
+            if (!option) return;
+
+            const isSelected = !!option.selected;
+            const stateKeys = optionSelectAllContext === undefined ? dropdownOptions.map(x => x.key as string) : optionSelectAllContext.map(x => x.key as string);
+
+            buildHealthOptions
+                .getTransactionSession()
+                .toggleAllStates(stateKeys, isSelected);
+        },
+        [buildHealthOptions, dropdownOptions]
+    );
+
+    // #endregion
+
+    // Dropdown params
+    const stateDropdownParams: BuildHealthOptionList = {
+        label: "Select State(s)",
+        items: items,
+        toolTip: "Select an option from the dropdown."
+    };
+
+    return stateDropdownParams.items.length > 0 && (
+        <Stack id={`state_dropdown`}>
+            <Label>State(s)</Label>
+            <Stack horizontal tokens={{ childrenGap: 8 }} verticalAlign="end">
+                <MultiListParameter
+                    handler={handler}
+                    buildHealthOptions={buildHealthOptions}
+                    param={stateDropdownParams}
+                    selectedKeys={selectedStateKeys}
+                    request={EMPTY_REFRESH_REQUEST}
+                    onChange={handleChange}
+                    enabledSelectAll={true}
+                    onSelectAll={handleSelectAll}
+                />
+            </Stack>
+        </Stack>
+    );
+});
+
 // #endregion -- Option Specific Dropdown Components --
 
 // #endregion -- Dropdown Components --
@@ -833,10 +921,11 @@ const BuildHealthSidebar: React.FC<{ handler: BuildHealthDataHandler, buildHealt
                         <StreamDropdown handler={handler} buildHealthOptions={buildHealthOptions} />
                         <TemplateDropdown handler={handler} buildHealthOptions={buildHealthOptions} />
                         <StepDropdown handler={handler} buildHealthOptions={buildHealthOptions} />
+                        <StateDropdown handler={handler} buildHealthOptions={buildHealthOptions} />
                         <TooltipHost content={
                             <Stack style={{ maxWidth: 300 }}>
                                 <Stack>
-                                    Clears the current Project/Stream/Template/Step selection.
+                                    Clears the current Project/Stream/Template/Step/State selection.
                                 </Stack>
                             </Stack>
                         }
@@ -858,7 +947,7 @@ const BuildHealthSidebar: React.FC<{ handler: BuildHealthDataHandler, buildHealt
                                 directionalHint: DirectionalHint.rightCenter
                             }}
                         >
-                            <Label>Jobs Since:</Label>
+                            <Label style={{ paddingTop: 10 }}>Jobs Since:</Label>
                         </TooltipHost>
                         <ComboBox
                             styles={{ root: { width: timeComboWidth } }}
@@ -870,42 +959,6 @@ const BuildHealthSidebar: React.FC<{ handler: BuildHealthDataHandler, buildHealt
                             }}
                         />
                     </Stack>
-                    <Stack tokens={{ childrenGap: 5 }}>
-                        <TooltipHost content={
-                            <div style={{ maxWidth: 300 }}>
-                                <div>
-                                    Include presubmits in the job history results, or not.
-                                </div>
-                            </div>
-                        }
-                            calloutProps={{
-                                directionalHint: DirectionalHint.rightCenter
-                            }}
-                        >
-                            <Checkbox
-                                label="Include Presubmits"
-                                checked={includePreflight}
-                                onChange={(_, isChecked) => updateIncludePreflight(!!isChecked)}
-                            />
-                        </TooltipHost>
-                        <TooltipHost content={
-                            <div style={{ maxWidth: 300 }}>
-                                <div>
-                                    Linked editing mode: Links filter selects across streams & templates where applicable.
-                                </div>
-                            </div>
-                        }
-                            calloutProps={{
-                                directionalHint: DirectionalHint.rightCenter
-                            }}
-                        >
-                            <Checkbox
-                                label="Linked Selection Mode"
-                                checked={linkedSelectionMode}
-                                onChange={(_, isChecked) => updateLinkedSelectionMode(!!isChecked)}
-                            />
-                        </TooltipHost>
-                    </Stack>
                 </Stack>
                 <Stack
                     styles={{
@@ -916,8 +969,25 @@ const BuildHealthSidebar: React.FC<{ handler: BuildHealthDataHandler, buildHealt
                         },
                     }}>
                     <Stack style={{ paddingTop: 0, paddingBottom: 4 }}>
-                        <Label>Step Outcome Visual Options</Label>
+                        <Label>Options</Label>
                     </Stack>
+                    <TooltipHost content={
+                        <div style={{ maxWidth: 300 }}>
+                            <div>
+                                Linked editing mode: Links filter selects across streams & templates where applicable.
+                            </div>
+                        </div>
+                    }
+                        calloutProps={{
+                            directionalHint: DirectionalHint.rightCenter
+                        }}
+                    >
+                        <Checkbox
+                            label="Linked Selection Mode"
+                            checked={linkedSelectionMode}
+                            onChange={(_, isChecked) => updateLinkedSelectionMode(!!isChecked)}
+                        />
+                    </TooltipHost>
                     <TooltipHost content={
                         <div style={{ maxWidth: 300 }}>
                             <div>
@@ -955,7 +1025,7 @@ const BuildHealthSidebar: React.FC<{ handler: BuildHealthDataHandler, buildHealt
                     <TooltipHost content={
                         <div style={{ maxWidth: 300 }}>
                             <div>
-                                Warnings are considered success in succes/failure ratio for summary.
+                                Shows labels (e.g. Win64) as entries in the table.
                             </div>
                         </div>
                     }
@@ -969,6 +1039,23 @@ const BuildHealthSidebar: React.FC<{ handler: BuildHealthDataHandler, buildHealt
                             onChange={(_, isChecked) => updateShowLabels(!!isChecked)}
                         />
                     </TooltipHost>
+                    <TooltipHost content={
+                        <div style={{ maxWidth: 300 }}>
+                            <div>
+                                Include presubmits in the job history results, or not.
+                            </div>
+                        </div>
+                    }
+                        calloutProps={{
+                            directionalHint: DirectionalHint.rightCenter
+                        }}
+                    >
+                        <Checkbox
+                            label="Include Presubmits"
+                            checked={includePreflight}
+                            onChange={(_, isChecked) => updateIncludePreflight(!!isChecked)}
+                        />
+                    </TooltipHost>
                 </Stack>
                 <Stack
                     styles={{
@@ -1019,13 +1106,15 @@ export const BuildHealthView: React.FC = observer(function ConstructBuildHealthV
     const uiOptions = useMemo(() => ({
         includeDateAnchors: buildHealthUIState.includeDateAnchors,
         warningsAsSummaryFailure: buildHealthUIState.warningsAsSummaryFailure,
-        showLabels: buildHealthUIState.showLabels
+        showLabels: buildHealthUIState.showLabels,
+        validStates: buildHealthUIState.validStates
     }),
         // eslint-disable-next-line react-hooks/exhaustive-deps -- All necessary dependencies included for MobX observables
         [
             buildHealthUIState.showLabels,
             buildHealthUIState.includeDateAnchors,
-            buildHealthUIState.warningsAsSummaryFailure
+            buildHealthUIState.warningsAsSummaryFailure,
+            buildHealthUIState.validStates
         ]);
 
     return (
diff --git a/Engine/Source/Programs/Horde/HordeDashboard/src/components/buildhealth/stepoutcome/StepOutcomeDataHandler.ts b/Engine/Source/Programs/Horde/HordeDashboard/src/components/buildhealth/stepoutcome/StepOutcomeDataHandler.ts
--- a/Engine/Source/Programs/Horde/HordeDashboard/src/components/buildhealth/stepoutcome/StepOutcomeDataHandler.ts
+++ b/Engine/Source/Programs/Horde/HordeDashboard/src/components/buildhealth/stepoutcome/StepOutcomeDataHandler.ts
@@ -85,6 +85,7 @@ export interface StepOutcomeViewOptions {
     includeDateAnchors: boolean;
     warningsAsSummaryFailure: boolean;
     showLabels: boolean;
+    validStates : Set<string>;
 }
 
 /**
@@ -108,7 +109,7 @@ export interface StepOutcomeFilters {
  */
 export class StepOutcomeDataHandler extends PollBase {
     private readonly STREAM_JOBS_FILTER = "id,streamId,name,change,templateId,state,createTime,updateTime,abortedByUserInfo,preflightCommitId,preflightDescription,cancellationReason,labels,defaultLabel";
-    private readonly STREAM_JOBS_BATCH_STEP_FILTER = "batches.steps.id,batches.steps.name,batches.steps.state,batches.steps.outcome,batches.steps.startTime,batches.steps.finishTime,batches.steps.warning,batches.steps.error"
+    private readonly STREAM_JOBS_BATCH_STEP_FILTER = "batches.error,batches.steps.id,batches.steps.name,batches.steps.state,batches.steps.outcome,batches.steps.startTime,batches.steps.finishTime,batches.steps.warning,batches.steps.error"
     private readonly STREAM_CHANGES_FILTER = "id,dateUtc";
     static readonly MAX_JOB_COUNT_DEFAULT = 400;
 
diff --git a/Engine/Source/Programs/Horde/HordeDashboard/src/components/buildhealth/stepoutcome/StepOutcomeDataTypes.ts b/Engine/Source/Programs/Horde/HordeDashboard/src/components/buildhealth/stepoutcome/StepOutcomeDataTypes.ts
--- a/Engine/Source/Programs/Horde/HordeDashboard/src/components/buildhealth/stepoutcome/StepOutcomeDataTypes.ts
+++ b/Engine/Source/Programs/Horde/HordeDashboard/src/components/buildhealth/stepoutcome/StepOutcomeDataTypes.ts
@@ -300,14 +300,23 @@ export class LabelOutcomeTableEntry extends OutcomeTableEntry {
 
     /**
      * @inheritdoc
-     * Note: Labels have a smaller stable state domain than Steps. Labels don't have a concept of skip / cancelled (due to complexities & edge cases), so we attempt to use a synthesized step state in such a case.
+     * 
+     * Note: 
+     * Labels have a smaller stable state domain than Steps. Labels don't have a concept of skip / cancelled (due to complexities & edge cases), so we attempt to use a synthesized step state in such a case.
      * This is inferred from other metadata associated with the job, and is provided at construction of the lable outcome entry.
+     * 
+     * Note:
+     * An explicit remapping occurs from @see LabelState.Complete to @see JobStepState.Completed to homogenize frontend experience.
      */
     get state(): LabelState | JobStepState {
         if (this.labelResponse.state === LabelState.Unspecified) {
             return this.synthesizedStepState ?? LabelState.Unspecified;
         }
         if (this.labelResponse.state !== undefined && this.labelResponse.state !== null) {
+            if (this.labelResponse.state === LabelState.Complete) {
+                return JobStepState.Completed;
+            }
+
             return this.labelResponse.state;
         }
         return LabelState.Unspecified;
@@ -757,6 +766,10 @@ export class StepOutcomeTable {
                     continue;
                 }
 
+                if (!this.viewModelOptions.validStates.has(cell.state)) {
+                    continue;
+                }
+
                 summaryTableEntry.steps++;
 
                 if (cell.state === WaitingStr || cell.state === RunningStr) {
@@ -835,8 +848,9 @@ export class StepOutcomeTable {
             }
             else {
                 const colIdx = this.changeColLookup.get(entry.change);
-                yield colIdx !== undefined ? rowEntries[colIdx] : undefined;
-
+                const tableEntry = colIdx !== undefined ? rowEntries[colIdx] : null;
+                
+                yield tableEntry !== null && this.viewModelOptions.validStates.has(tableEntry.state) ? tableEntry : undefined;
             }
         }
     }
@@ -860,7 +874,7 @@ export class StepOutcomeTable {
             }
 
             let tableEntry: (OutcomeTableEntry | null) = this.tableEntries[rowIdx][colIdx];
-            if (tableEntry !== null) {
+            if (tableEntry !== null && this.viewModelOptions.validStates.has(tableEntry.state)) {
                 yield tableEntry;
             }
         }
@@ -937,6 +951,10 @@ export class StepOutcomeTable {
                     continue;
                 }
 
+                if (!this.viewModelOptions.validStates.has(castedEntry.state)) {
+                    continue;
+                }
+
                 if (resultString == undefined) {
                     resultString = castedEntry.getCSVHeader();
                 }
@@ -976,6 +994,10 @@ export class StepOutcomeTable {
                     continue;
                 }
 
+                if (!this.viewModelOptions.validStates.has(castedEntry.state)) {
+                    continue;
+                }
+
                 if (resultString == undefined) {
                     resultString = castedEntry.getCSVHeader();
                 }
diff --git a/Engine/Source/Programs/Horde/HordeDashboard/src/components/buildhealth/stepoutcome/StepOutcomeUtilities.ts b/Engine/Source/Programs/Horde/HordeDashboard/src/components/buildhealth/stepoutcome/StepOutcomeUtilities.ts
--- a/Engine/Source/Programs/Horde/HordeDashboard/src/components/buildhealth/stepoutcome/StepOutcomeUtilities.ts
+++ b/Engine/Source/Programs/Horde/HordeDashboard/src/components/buildhealth/stepoutcome/StepOutcomeUtilities.ts
@@ -1,6 +1,6 @@
 // Copyright Epic Games, Inc. All Rights Reserved.
 
-import { GetJobResponse, JobStepOutcome, JobStepState } from "horde/backend/Api";
+import { GetJobResponse, JobStepBatchError, JobStepOutcome, JobStepState } from "horde/backend/Api";
 import { OutcomeTableEntry, ReadyStr, RunningStr, WaitingStr } from "./StepOutcomeDataTypes";
 
 export const DURATION_NOT_SET: number = -1;
@@ -41,6 +41,13 @@ export function summarizeJob(job: GetJobResponse): JobStepOutcome {
                 accumulatorOutcome = JobStepOutcome.Warnings;
             }
         }
+
+        // Check to see if an entire batch has failed, in spite of the steps declaring a success (due to the absence of steps).
+        // This is a common case when a more terminal failure occurs in batch resolution.
+        if (accumulatorOutcome === JobStepOutcome.Success && response.error !== JobStepBatchError.None) {
+            accumulatorOutcome = JobStepOutcome.Failure;
+            return accumulatorOutcome;
+        }
     }
 
     return accumulatorOutcome;
