Horde: Check if a conform is required before asssigning a Job Task lease to an agent

diff --git a/Engine/Source/Programs/Horde/Plugins/Build/HordeServer.Build.Tests/BuildTestSetup.cs b/Engine/Source/Programs/Horde/Plugins/Build/HordeServer.Build.Tests/BuildTestSetup.cs
--- a/Engine/Source/Programs/Horde/Plugins/Build/HordeServer.Build.Tests/BuildTestSetup.cs
+++ b/Engine/Source/Programs/Horde/Plugins/Build/HordeServer.Build.Tests/BuildTestSetup.cs
@@ -123,6 +123,7 @@ namespace HordeServer.Tests
 
 			IConfiguration config = new ConfigurationBuilder().Build();
 			services.Configure<ServerSettings>(ConfigureSettings);
+			services.Configure<BuildServerConfig>(x => x.EnableJobTaskConformCheck = true);
 			services.AddSingleton<IConfiguration>(config);
 
 			services.AddHttpClient<RpcService>();
diff --git a/Engine/Source/Programs/Horde/Plugins/Build/HordeServer.Build.Tests/Fixture.cs b/Engine/Source/Programs/Horde/Plugins/Build/HordeServer.Build.Tests/Fixture.cs
--- a/Engine/Source/Programs/Horde/Plugins/Build/HordeServer.Build.Tests/Fixture.cs
+++ b/Engine/Source/Programs/Horde/Plugins/Build/HordeServer.Build.Tests/Fixture.cs
@@ -65,11 +65,16 @@ namespace HordeServer.Tests
 
 			Dictionary<string, AgentConfig> agentTypes = new()
 			{
-				{ "Win64", new() { Pool = new PoolId(PoolName) } }
+				{ "Win64", new() { Pool = new PoolId(PoolName), Workspace = "Workspace1"} }
+			};
+
+			Dictionary<string, WorkspaceConfig> workspaceTypes = new()
+			{
+				{ "Workspace1", new () { Identifier = "Workspace1", ConformDiskFreeSpace = 500 } }
 			};
 
 			StreamId streamId = new StreamId("ue5-main");
-			StreamConfig streamConfig = new StreamConfig { Id = streamId, Name = "//UE5/Main", Tabs = tabs, Templates = templates, AgentTypes = agentTypes };
+			StreamConfig streamConfig = new StreamConfig { Id = streamId, Name = "//UE5/Main", Tabs = tabs, Templates = templates, AgentTypes = agentTypes, WorkspaceTypes = workspaceTypes };
 
 			ProjectId projectId = new ProjectId("ue5");
 			ProjectConfig projectConfig = new ProjectConfig { Id = projectId, Name = "UE5", Streams = new List<StreamConfig> { streamConfig } };
diff --git a/Engine/Source/Programs/Horde/Plugins/Build/HordeServer.Build.Tests/Jobs/JobTaskSourceTests.cs b/Engine/Source/Programs/Horde/Plugins/Build/HordeServer.Build.Tests/Jobs/JobTaskSourceTests.cs
--- a/Engine/Source/Programs/Horde/Plugins/Build/HordeServer.Build.Tests/Jobs/JobTaskSourceTests.cs
+++ b/Engine/Source/Programs/Horde/Plugins/Build/HordeServer.Build.Tests/Jobs/JobTaskSourceTests.cs
@@ -176,5 +176,41 @@ namespace HordeServer.Tests.Jobs
 
 			return fixture;
 		}
+
+		[TestMethod]
+		[DataRow(0, true, DisplayName = "Conform is requested when assigning JobTask and agent has low disk free space")]
+		[DataRow(1000, false, DisplayName = "Conform is not requested when assigning JobTask and agent has enough disk free space")]
+		public async Task AgentAssignLeaseWithConformCheckAsync(long diskFreeSpaceInMb, bool expectingConform)
+		{
+			// TODO Make workspace conform size an input parameter rather than be hardcoded in Fixture.PopulateAsync
+			// TODO Make setting EnableJobTaskConformCheck part of the test rather than set feature flag globally for all in BuildTestSetup.cs
+
+			Fixture fixture = await CreateFixtureAsync();
+			IPool pool = await CreatePoolAsync(new PoolConfig { Name = Fixture.PoolName, EnableAutoscaling = false, MinAgents = 0, NumReserveAgents = 0 });
+
+			IAgent agent = fixture.Agent1;
+			agent = Deref(await agent.TryUpdateAsync(new UpdateAgentOptions { Enabled = true, ExplicitPools = [pool.Id] }));
+
+			List<string> properties = [$"{KnownPropertyNames.DiskFreeSpace}={diskFreeSpaceInMb * 1024 * 1024}"];
+			agent = Deref(await AgentService.CreateSessionAsync(agent, new RpcAgentCapabilities(properties), null));
+
+			await JobTaskSource.TickAsync(CancellationToken.None);
+
+			Task<CreateLeaseOptions?> lease = Deref(await JobTaskSource.AssignLeaseAsync(agent, CancellationToken.None));
+			CreateLeaseOptions? leaseResult = await lease.WaitAsync(CancellationToken.None);
+
+			agent = Deref(await AgentService.GetAgentAsync(agent.Id));
+
+			if (expectingConform)
+			{
+				Assert.IsTrue(agent.RequestFullConform);
+				Assert.IsNull(leaseResult);
+			}
+			else
+			{
+				Assert.IsFalse(agent.RequestFullConform);
+				Assert.IsNotNull(leaseResult);
+			}
+		}
 	}
 }
diff --git a/Engine/Source/Programs/Horde/Plugins/Build/HordeServer.Build/Agents/Pools/PoolUpdateService.cs b/Engine/Source/Programs/Horde/Plugins/Build/HordeServer.Build/Agents/Pools/PoolUpdateService.cs
--- a/Engine/Source/Programs/Horde/Plugins/Build/HordeServer.Build/Agents/Pools/PoolUpdateService.cs
+++ b/Engine/Source/Programs/Horde/Plugins/Build/HordeServer.Build/Agents/Pools/PoolUpdateService.cs
@@ -5,6 +5,7 @@ using EpicGames.Horde.Agents;
 using EpicGames.Horde.Agents.Pools;
 using HordeServer.VersionControl.Perforce;
 using HordeServer.Streams;
+using HordeServer.Utilities;
 using Microsoft.Extensions.Hosting;
 using Microsoft.Extensions.Logging;
 using Microsoft.Extensions.Options;
@@ -137,29 +138,9 @@ namespace HordeServer.Agents.Pools
 			List<IAgent> agents = (await _agentService.GetCachedAgentsAsync(cancellationToken))
 				.Where(x => x is { Status: AgentStatus.Stopped, Deleted: false, Enabled: true, Ephemeral: false }).ToList();
 
-			long MegabytesToBytes(long v) => v * 1024 * 1024;
 			foreach (IAgent agent in agents)
 			{
-				long? freeDiskSpace = agent.GetDiskFreeSpace();
-				long conformDiskSpaceNeeded = 0;
-				HashSet<AgentWorkspaceInfo> workspaces = await _poolService.GetWorkspacesAsync(agent, DateTime.UtcNow - TimeSpan.FromHours(1), _buildConfig.CurrentValue, cancellationToken);
-				
-				// Find the largest conform disk space amount needed, if any
-				foreach (AgentWorkspaceInfo workspace in workspaces)
-				{
-					if (workspace.ConformDiskFreeSpace is > 0)
-					{
-						conformDiskSpaceNeeded = Math.Max(conformDiskSpaceNeeded, MegabytesToBytes(workspace.ConformDiskFreeSpace.Value));
-					}
-				}
-
-				if (freeDiskSpace != null && conformDiskSpaceNeeded > 0 && freeDiskSpace < conformDiskSpaceNeeded &&
-				    agent.ConformAttemptCount is null or 0 && !agent.RequestFullConform)
-				{
-					await agent.TryUpdateAsync(new UpdateAgentOptions { RequestFullConform = true }, cancellationToken: cancellationToken);
-					_logger.LogInformation("Auto-conforming {AgentId} as workspace conform disk space needed ({ConformDiskSpace:F1} MB) is less than free disk space ({FreeDiskSpace:F1} MB)",
-						agent.Id.ToString(), conformDiskSpaceNeeded / 1024.0 / 1024.0, freeDiskSpace.Value / 1024.0 / 1024.0);
-				}
+				await ConformUtilities.AutoConformAgentAsync(agent, _poolService, _buildConfig.CurrentValue, _logger, cancellationToken);
 			}
 		}
 
diff --git a/Engine/Source/Programs/Horde/Plugins/Build/HordeServer.Build/BuildServerConfig.cs b/Engine/Source/Programs/Horde/Plugins/Build/HordeServer.Build/BuildServerConfig.cs
--- a/Engine/Source/Programs/Horde/Plugins/Build/HordeServer.Build/BuildServerConfig.cs
+++ b/Engine/Source/Programs/Horde/Plugins/Build/HordeServer.Build/BuildServerConfig.cs
@@ -35,6 +35,11 @@ namespace HordeServer
 		/// </summary>
 		public bool EnableConformTasks { get; set; } = true;
 
+		/// <summary>
+		/// Whether to check if a conform is required on an agent before assigning a lease for a job task source.
+		/// </summary>
+		public bool EnableJobTaskConformCheck { get; set; }
+
 		/// <summary>
 		/// Url of P4 Swarm installation
 		/// </summary>
diff --git a/Engine/Source/Programs/Horde/Plugins/Build/HordeServer.Build/Jobs/JobTaskSource.cs b/Engine/Source/Programs/Horde/Plugins/Build/HordeServer.Build/Jobs/JobTaskSource.cs
--- a/Engine/Source/Programs/Horde/Plugins/Build/HordeServer.Build/Jobs/JobTaskSource.cs
+++ b/Engine/Source/Programs/Horde/Plugins/Build/HordeServer.Build/Jobs/JobTaskSource.cs
@@ -169,6 +169,11 @@ namespace HordeServer.Jobs
 			/// </summary>
 			public TaskCompletionSource<CreateLeaseOptions?> LeaseSource { get; } = new TaskCompletionSource<CreateLeaseOptions?>(TaskCreationOptions.RunContinuationsAsynchronously);
 
+			/// <summary>
+			/// Ignore the waiting agent when trying to assign a new queue item
+			/// </summary>
+			public bool Skip { get; set; }
+
 			/// <summary>
 			/// Constructor
 			/// </summary>
@@ -610,7 +615,7 @@ namespace HordeServer.Jobs
 		/// <returns></returns>
 		bool TryAssignItemToWaiter(QueueItem item, QueueWaiter waiter)
 		{
-			if (item._assignTask == null && item.Batch.SessionId == null && waiter.Agent.Enabled && waiter.Agent.Leases.Count == 0 && waiter.Agent.IsInPool(item._poolId))
+			if (item._assignTask == null && item.Batch.SessionId == null && waiter.Agent.Enabled && waiter.Agent.Leases.Count == 0 && waiter.Agent.IsInPool(item._poolId) && !waiter.Skip)
 			{
 				Task startTask = new Task<Task>(() => TryCreateLeaseAsync(item, waiter, CancellationToken.None));
 				Task executeTask = startTask.ContinueWith(task => task, TaskScheduler.Default);
@@ -760,6 +765,29 @@ namespace HordeServer.Jobs
 			return CancelLeaseAsync(agent, JobId.Parse(task.JobId), JobStepBatchId.Parse(task.BatchId), cancellationToken);
 		}
 
+		private async Task<bool> AutoConformAgentAsync(IAgent agent, CancellationToken cancellationToken)
+		{
+			if (!_settings.CurrentValue.EnableJobTaskConformCheck)
+			{
+				return false;
+			}
+
+			if (!(_settings.CurrentValue.EnableConformTasks && _buildConfig.CurrentValue.EnableConformTasks))
+			{
+				return false;
+			}
+
+			// When a call to JobTaskSource.AssignLeaseAsync by the AgentService requests conform it will not be actioned
+			// by ConformTaskSource until the next iteration of AgentService.WaitForLeaseAsync. Make sure a job/batch is
+			// not assigned on that next iteration because it could take priority over the conform task.
+			if (agent.ConformAttemptCount is null or 0 && agent.RequestFullConform)
+			{
+				return true;
+			}
+
+			return await ConformUtilities.AutoConformAgentAsync(agent, _poolService, _buildConfig.CurrentValue, _logger, cancellationToken);
+		}
+
 		/// <summary>
 		/// Assign a new batch to be executed by the given agent
 		/// </summary>
@@ -781,6 +809,27 @@ namespace HordeServer.Jobs
 				.WithProperty("LeaseId", leaseId.ToString())
 				.BeginScope();
 
+			bool conformRequested = await AutoConformAgentAsync(waiter.Agent, cancellationToken);
+			if (conformRequested)
+			{
+				// Won't be assigning a lease to this waiter
+				if (waiter.LeaseSource.TrySetResult(null))
+				{
+					_logger.LogInformation("Failed to assign job/batch to agent. Conform requested");
+
+					// Don't try to assign another item to this waiter to avoid duplicating a conform request
+					waiter.Skip = true;
+
+					// Clear out the assignment for this item, and try to reassign it
+					item._assignTask = null;
+					if (!item._removed)
+					{
+						AssignQueueItemToAnyWaiter(item);
+					}
+					return null;
+				}
+			}
+
 			_logger.LogInformation("Assigning job to waiter");
 
 			// Allocate a log ID but hold off creating the actual log file until the lease has been accepted
diff --git a/Engine/Source/Programs/Horde/Plugins/Build/HordeServer.Build/Utilities/ConformUtilities.cs b/Engine/Source/Programs/Horde/Plugins/Build/HordeServer.Build/Utilities/ConformUtilities.cs
new file mode 100644
--- /dev/null
+++ b/Engine/Source/Programs/Horde/Plugins/Build/HordeServer.Build/Utilities/ConformUtilities.cs
@@ -0,0 +1,53 @@
+// Copyright Epic Games, Inc. All Rights Reserved.
+
+using HordeServer.Agents;
+using HordeServer.Agents.Pools;
+using HordeServer.VersionControl.Perforce;
+using Microsoft.Extensions.Logging;
+
+namespace HordeServer.Utilities;
+
+/// <summary>
+/// Utilities for conforming agents
+/// </summary>
+public static class ConformUtilities
+{
+	/// <summary>
+	/// Conform an agent that is below the disk free conform threshold for a workspace
+	/// </summary>
+	/// <param name="agent">Agent to query</param>
+	/// <param name="poolService">Pool service instance</param>
+	/// <param name="config">Current configuration</param>
+	/// <param name="logger">Logger instance</param>
+	/// <param name="cancellationToken">Cancellation token for the async task</param>
+	/// <returns>True if conform was requested for the agent; otherwise false</returns>
+	public static async Task<bool> AutoConformAgentAsync(IAgent agent, PoolService poolService, BuildConfig config, ILogger logger, CancellationToken cancellationToken)
+	{
+		long MegabytesToBytes(long v) => v * 1024 * 1024;
+		double BytesToMegabytes(long v) => v / 1024.0 / 1024.0;
+		long? freeDiskSpace = agent.GetDiskFreeSpace();
+		long conformDiskSpaceNeeded = 0;
+		HashSet<AgentWorkspaceInfo> workspaces = await poolService.GetWorkspacesAsync(agent, DateTime.UtcNow - TimeSpan.FromHours(1), config, cancellationToken);
+
+		// Find the largest conform disk space amount needed, if any
+		foreach (AgentWorkspaceInfo workspace in workspaces)
+		{
+			if (workspace.ConformDiskFreeSpace is > 0)
+			{
+				conformDiskSpaceNeeded = Math.Max(conformDiskSpaceNeeded, MegabytesToBytes(workspace.ConformDiskFreeSpace.Value));
+			}
+		}
+
+		IAgent? newAgent = null;
+
+		if (freeDiskSpace != null && conformDiskSpaceNeeded > 0 && freeDiskSpace < conformDiskSpaceNeeded &&
+		    agent.ConformAttemptCount is null or 0 && !agent.RequestFullConform)
+		{
+			newAgent = await agent.TryUpdateAsync(new UpdateAgentOptions { RequestFullConform = true }, cancellationToken: cancellationToken);
+			logger.LogInformation("Auto-conforming {AgentId} as workspace conform disk space needed ({ConformDiskSpace:F1} MB) is greater than free disk space ({FreeDiskSpace:F1} MB)",
+				agent.Id.ToString(), BytesToMegabytes(conformDiskSpaceNeeded), BytesToMegabytes(freeDiskSpace.Value));
+		}
+
+		return newAgent?.RequestFullConform ?? false;
+	}
+}
\ No newline at end of file
