Horde: Add support for parallel test execution for database integration tests (MongoDB and Redis)

diff --git a/Engine/Source/Programs/Horde/HordeServer.Tests.Shared/ControllerIntegrationTest.cs b/Engine/Source/Programs/Horde/HordeServer.Tests.Shared/ControllerIntegrationTest.cs
--- a/Engine/Source/Programs/Horde/HordeServer.Tests.Shared/ControllerIntegrationTest.cs
+++ b/Engine/Source/Programs/Horde/HordeServer.Tests.Shared/ControllerIntegrationTest.cs
@@ -87,13 +87,24 @@ public class FakeHordeWebApp : IAsyncDisposable
 	public RedisInstance RedisInstance { get; }
 	public IServiceProvider ServiceProvider => Factory.Services;
 	private TestWebApplicationFactory<Startup> Factory { get; }
+	private readonly DatabaseIntegrationTest _dbIntegrationTest;
 
 	public FakeHordeWebApp(Action<IServiceCollection>? configureServices = null, Dictionary<string, string>? settings = null)
 	{
 		ServerApp.InitializePluginsForTests();
 		MongoInstance = new MongoInstance();
 		RedisInstance = new RedisInstance();
-		Factory = new TestWebApplicationFactory<Startup>(MongoInstance, RedisInstance, configureServices, settings);
+		_dbIntegrationTest = new DatabaseIntegrationTest();
+
+		void ConfigureServices(IServiceCollection services)
+		{
+			configureServices?.Invoke(services);
+			services.AddSingleton<IMongoService>(_dbIntegrationTest.GetMongoServiceSingleton());
+			services.AddSingleton<IRedisService>(_dbIntegrationTest.GetRedisServiceSingleton());
+		}
+		
+		Factory = new TestWebApplicationFactory<Startup>(MongoInstance, RedisInstance, ConfigureServices, settings);
+		
 	}
 
 	public HttpClient CreateHttpClient(bool allowAutoRedirect = true)
@@ -107,6 +118,7 @@ public class FakeHordeWebApp : IAsyncDisposable
 		try
 		{
 			await Factory.DisposeAsync();
+			await _dbIntegrationTest.DisposeAsync();
 			MongoInstance.Dispose();
 			RedisInstance.Dispose();
 		}
diff --git a/Engine/Source/Programs/Horde/HordeServer.Tests.Shared/DatabaseIntegrationTest.cs b/Engine/Source/Programs/Horde/HordeServer.Tests.Shared/DatabaseIntegrationTest.cs
--- a/Engine/Source/Programs/Horde/HordeServer.Tests.Shared/DatabaseIntegrationTest.cs
+++ b/Engine/Source/Programs/Horde/HordeServer.Tests.Shared/DatabaseIntegrationTest.cs
@@ -1,5 +1,6 @@
 // Copyright Epic Games, Inc. All Rights Reserved.
 
+using System.Collections.Concurrent;
 using System.Net;
 using System.Runtime.InteropServices;
 using HordeServer.Server;
@@ -126,17 +127,25 @@ namespace HordeServer.Tests
 	{
 		static readonly object s_lockObject = new object();
 
-		const bool UseExistingRedisInstance = true;
-		const int RedisPort = 6379;
-		const int RedisDbNum = 15;
-
+		private static (int min, int max) DatabaseNumbers { get; } = (1000, 2000);
+		private static readonly BlockingCollection<int> s_availableRedisDbs = new ();
 		private static string? s_redisConnectionString;
-		private static int s_redisDbNum;
 		private static RedisProcess? s_redisProcess;
 
 		public string ConnectionString { get; }
 		public int DatabaseNumber { get; }
 
+		private const bool UseExistingRedisInstance = true;
+		private const int RedisPort = 6379;
+		
+		static RedisInstance()
+		{
+			for (int i = DatabaseNumbers.min; i < DatabaseNumbers.max; i++)
+			{
+				s_availableRedisDbs.Add(i);
+			}
+		}
+		
 		public RedisInstance()
 		{
 			lock (s_lockObject)
@@ -145,16 +154,32 @@ namespace HordeServer.Tests
 				{
 					int port = GetRedisPortInternal();
 					s_redisConnectionString = $"127.0.0.1:{port},allowAdmin=true";
-					s_redisDbNum = RedisDbNum;
 				}
 
 				ConnectionString = s_redisConnectionString;
-				DatabaseNumber = s_redisDbNum;
+				DatabaseNumber = CheckoutDatabase();
 			}
 		}
 
 		public void Dispose()
 		{
+			ReturnDatabase(DatabaseNumber);
+		}
+		
+		private static int CheckoutDatabase()
+		{
+			TimeSpan timeoutToUse = TimeSpan.FromMinutes(1);
+			if (s_availableRedisDbs.TryTake(out int dbNumber, timeoutToUse))
+			{
+				return dbNumber;
+			}
+        
+			throw new TimeoutException($"Failed to checkout database within {timeoutToUse}");
+		}
+
+		private static void ReturnDatabase(int dbNumber)
+		{
+			s_availableRedisDbs.Add(dbNumber);
 		}
 
 		static int GetRedisPortInternal()
@@ -180,9 +205,10 @@ namespace HordeServer.Tests
 	public class DatabaseIntegrationTest : ServiceTest
 	{
 		private static readonly object s_lockObject = new object();
+		private static MongoService? s_globalMongoService;
 
 		private MongoInstance? _mongoInstance;
-		private MongoService? _mongoService;
+		private IMongoService? _mongoService;
 		private readonly LoggerFactory _loggerFactory = new LoggerFactory();
 
 		private RedisInstance? _redisInstance;
@@ -215,9 +241,9 @@ namespace HordeServer.Tests
 
 			GC.SuppressFinalize(this);
 
-			if (_mongoService != null)
+			if (_mongoService is IAsyncDisposable ad)
 			{
-				await _mongoService.DisposeAsync();
+				await ad.DisposeAsync();
 			}
 			_mongoInstance?.Dispose();
 
@@ -234,17 +260,24 @@ namespace HordeServer.Tests
 		{
 			lock (s_lockObject)
 			{
-				if (_mongoService == null)
+				if (s_globalMongoService == null)
 				{
 					IRedisService redisService = GetRedisServiceSingleton();
 					_mongoInstance = new MongoInstance();
-
-					ServerSettings ss = new ServerSettings();
-					ss.MongoDatabaseName = _mongoInstance.DatabaseName;
-					ss.MongoConnectionString = _mongoInstance.ConnectionString;
-
+					
+					ServerSettings ss = new()
+					{
+						MongoDatabaseName = _mongoInstance.DatabaseName,
+						MongoConnectionString = _mongoInstance.ConnectionString
+					};
+					
 					MongoCommandTracer mongoTracer = new(OpenTelemetryTracers.MongoDb, _loggerFactory.CreateLogger<MongoCommandTracer>());
-					_mongoService = new MongoService(Options.Create(ss), redisService, mongoTracer, OpenTelemetryTracers.Horde, _loggerFactory.CreateLogger<MongoService>(), _loggerFactory);
+					s_globalMongoService = new MongoService(Options.Create(ss), redisService, mongoTracer, OpenTelemetryTracers.Horde, _loggerFactory.CreateLogger<MongoService>(), _loggerFactory);
+				}
+
+				if (_mongoService == null)
+				{
+					_mongoService = new TestScopedMongoService(s_globalMongoService.GetClient());
 				}
 			}
 			return _mongoService;
diff --git a/Engine/Source/Programs/Horde/HordeServer.Tests.Shared/TestScopedMongoService.cs b/Engine/Source/Programs/Horde/HordeServer.Tests.Shared/TestScopedMongoService.cs
new file mode 100644
--- /dev/null
+++ b/Engine/Source/Programs/Horde/HordeServer.Tests.Shared/TestScopedMongoService.cs
@@ -0,0 +1,143 @@
+﻿// Copyright Epic Games, Inc. All Rights Reserved.
+
+using HordeServer.Server;
+using HordeServer.Utilities;
+using Microsoft.Extensions.Diagnostics.HealthChecks;
+using MongoDB.Bson;
+using MongoDB.Driver;
+
+namespace HordeServer.Tests;
+
+/// <summary>
+/// Implementation instantiated per test with a unique database name to sandbox and isolate reads/writes during a test
+/// </summary>
+public sealed class TestScopedMongoService : IMongoService, IDisposable
+{
+	public IMongoDatabase Database { get; }
+	public bool ReadOnlyMode { get; } = false;
+
+	private readonly string _dbName = $"HordeTest-{Guid.NewGuid()}";
+	private readonly MongoClient _mongoClient;
+	private readonly IMongoCollection<BsonDocument> _singletonsV1;
+	private readonly IMongoCollection<BsonDocument> _singletonsV2;
+	private readonly ISet<string> _indexesCreated = new HashSet<string>();
+	private readonly object _indexLock = new();
+	
+	public TestScopedMongoService(MongoClient mongoClient)
+	{
+		_mongoClient = mongoClient;
+		Database = _mongoClient.GetDatabase(_dbName);
+		_singletonsV1 = GetCollection<BsonDocument>("Singletons");
+		_singletonsV2 = GetCollection<BsonDocument>("SingletonsV2");
+	}
+	
+	/// <inheritdoc/>
+	public void Dispose()
+	{
+		_mongoClient.DropDatabase(_dbName);
+	}
+
+	/// <inheritdoc/>
+	public MongoClient GetClient()
+	{
+		return _mongoClient;
+	}
+
+	/// <inheritdoc/>
+	public IMongoCollection<T> GetCollection<T>(string name)
+	{
+		return Database.GetCollection<T>(name);
+	}
+
+	/// <inheritdoc/>
+	public IMongoCollection<T> GetCollection<T>(string name, Func<IndexKeysDefinitionBuilder<T>, IndexKeysDefinition<T>> keysFunc, bool unique = false)
+	{
+		return GetCollection(name, [MongoIndex.Create<T>(keysFunc, unique)]);
+	}
+
+	/// <inheritdoc/>
+	public IMongoCollection<T> GetCollection<T>(string name, IEnumerable<MongoIndex<T>> indexes)
+	{
+		IMongoCollection<T> collection = GetCollection<T>(name);
+		AddIndexes(name, collection, indexes.ToArray());
+		return collection;
+	}
+
+	private void AddIndexes<T>(string collectionName, IMongoCollection<T> collection, MongoIndex<T>[] indexes)
+	{
+		lock (_indexLock)
+		{
+			if (_indexesCreated.Contains(collectionName))
+			{
+				return;
+			}
+		
+			foreach (MongoIndex<T> index in indexes)
+			{
+				CreateIndexOptions<T> options = new CreateIndexOptions<T>
+				{
+					Name = index.Name,
+					Unique = index.Unique,
+					Sparse = index.Sparse,
+					Background = true
+				};
+				collection.Indexes.CreateOne(new CreateIndexModel<T>(index.Keys, options));
+			}
+
+			_indexesCreated.Add(collectionName);
+		}
+	}
+
+	/// <inheritdoc/>
+	public Task<T> GetSingletonAsync<T>(CancellationToken cancellationToken) where T : SingletonBase, new()
+	{
+		return GetSingletonAsync(() => new T(), cancellationToken);
+	}
+
+	/// <inheritdoc/>
+	public Task<T> GetSingletonAsync<T>(Func<T> constructor, CancellationToken cancellationToken) where T : SingletonBase, new()
+	{
+		return MongoService.GetSingletonStaticAsync(_singletonsV1, _singletonsV2, constructor, cancellationToken);
+	}
+
+	/// <inheritdoc/>
+	public Task UpdateSingletonAsync<T>(Action<T> updater, CancellationToken cancellationToken) where T : SingletonBase, new()
+	{
+		bool Update(T instance)
+		{
+			updater(instance);
+			return true;
+		}
+		return UpdateSingletonAsync<T>(Update, cancellationToken);
+	}
+
+	/// <inheritdoc/>
+    public async Task UpdateSingletonAsync<T>(Func<T, bool> updater, CancellationToken cancellationToken) where T : SingletonBase, new()
+	{
+		for (; ; )
+		{
+			T document = await GetSingletonAsync(() => new T(), cancellationToken);
+			if (!updater(document))
+			{
+				break;
+			}
+			if (await TryUpdateSingletonAsync(document, cancellationToken))
+			{
+				break;
+			}
+		}
+	}
+
+	/// <inheritdoc/>
+    public Task<bool> TryUpdateSingletonAsync<T>(T singletonObject, CancellationToken cancellationToken) where T : SingletonBase
+    {
+	    return MongoService.TryUpdateSingletonStaticAsync(_singletonsV2, singletonObject, cancellationToken);
+	}
+	
+	/// <inheritdoc/>
+	public Task<HealthCheckResult> CheckHealthAsync(HealthCheckContext context, CancellationToken cancellationToken = new CancellationToken())
+	{
+		return Task.FromResult(HealthCheckResult.Healthy());
+	}
+
+}
\ No newline at end of file
diff --git a/Engine/Source/Programs/Horde/HordeServer/Server/MongoService.cs b/Engine/Source/Programs/Horde/HordeServer/Server/MongoService.cs
--- a/Engine/Source/Programs/Horde/HordeServer/Server/MongoService.cs
+++ b/Engine/Source/Programs/Horde/HordeServer/Server/MongoService.cs
@@ -561,25 +561,13 @@ namespace HordeServer.Server
 			return _client;
 		}
 
-		/// <summary>
-		/// Get a MongoDB collection from database
-		/// </summary>
-		/// <param name="name">Name of collection</param>
-		/// <typeparam name="T">A MongoDB document</typeparam>
-		/// <returns></returns>
+		/// <inheritdoc/>
 		public IMongoCollection<T> GetCollection<T>(string name)
 		{
 			return GetCollection<T>(name, Enumerable.Empty<MongoIndex<T>>());
 		}
 
-		/// <summary>
-		/// Get a MongoDB collection from database with a single index
-		/// </summary>
-		/// <param name="name">Name of collection</param>
-		/// <param name="keysFunc">Method to configure keys for the collection</param>
-		/// <param name="unique">Whether a unique index is required</param>
-		/// <typeparam name="T">A MongoDB document</typeparam>
-		/// <returns></returns>
+		/// <inheritdoc/>
 		public IMongoCollection<T> GetCollection<T>(string name, Func<IndexKeysDefinitionBuilder<T>, IndexKeysDefinition<T>> keysFunc, bool unique = false)
 		{
 			List<MongoIndex<T>> indexes = new List<MongoIndex<T>>();
@@ -587,13 +575,7 @@ namespace HordeServer.Server
 			return GetCollection(name, indexes);
 		}
 
-		/// <summary>
-		/// Get a MongoDB collection from database
-		/// </summary>
-		/// <param name="name">Name of collection</param>
-		/// <param name="indexes">Indexes for the collection</param>
-		/// <typeparam name="T">A MongoDB document</typeparam>
-		/// <returns></returns>
+		/// <inheritdoc/>
 		public IMongoCollection<T> GetCollection<T>(string name, IEnumerable<MongoIndex<T>> indexes)
 		{
 			IMongoCollection<T> collection = Database.GetCollection<T>(name);
@@ -767,7 +749,7 @@ namespace HordeServer.Server
 				{
 					return false;
 				}
-				if (await TryUpdateSingletonAsync(currentSchema, cancellationToken))
+				if (await TryUpdateSingletonStaticAsync(SingletonsV2, currentSchema, cancellationToken))
 				{
 					return true;
 				}
@@ -789,29 +771,68 @@ namespace HordeServer.Server
 			}
 		}
 
-		/// <summary>
-		/// Gets a singleton document by id
-		/// </summary>
-		/// <returns>The document</returns>
+		/// <inheritdoc/>
 		public Task<T> GetSingletonAsync<T>(CancellationToken cancellationToken) where T : SingletonBase, new()
 		{
-			return GetSingletonAsync(() => new T(), cancellationToken);
+			return GetSingletonStaticAsync(SingletonsV1, SingletonsV2, () => new T(), cancellationToken);
+		}
+
+		/// <inheritdoc/>
+		public Task<T> GetSingletonAsync<T>(Func<T> constructor, CancellationToken cancellationToken) where T : SingletonBase, new()
+		{
+			return GetSingletonStaticAsync(SingletonsV1, SingletonsV2, constructor, cancellationToken);
+		}
+
+		/// <inheritdoc/>
+		public Task UpdateSingletonAsync<T>(Action<T> updater, CancellationToken cancellationToken) where T : SingletonBase, new()
+		{
+			bool Update(T instance)
+			{
+				updater(instance);
+				return true;
+			}
+			return UpdateSingletonAsync<T>(Update, cancellationToken);
+		}
+
+		/// <inheritdoc/>
+		public async Task UpdateSingletonAsync<T>(Func<T, bool> updater, CancellationToken cancellationToken) where T : SingletonBase, new()
+		{
+			for (; ; )
+			{
+				T document = await GetSingletonStaticAsync(SingletonsV1, SingletonsV2, () => new T(), cancellationToken);
+				if (!updater(document))
+				{
+					break;
+				}
+				if (await TryUpdateSingletonStaticAsync(SingletonsV2, document, cancellationToken))
+				{
+					break;
+				}
+			}
+		}
+
+		/// <inheritdoc/>
+		public Task<bool> TryUpdateSingletonAsync<T>(T singletonObject, CancellationToken cancellationToken) where T : SingletonBase
+		{
+			return TryUpdateSingletonStaticAsync(SingletonsV2, singletonObject, cancellationToken);
 		}
 
 		/// <summary>
 		/// Gets a singleton document by id
 		/// </summary>
+		/// <param name="singletonsV1">Collection of Singleton documents V1</param>
+		/// <param name="singletonsV2">Collection of Singleton documents V2</param>
 		/// <param name="constructor">Method to use to construct a new object</param>
 		/// <param name="cancellationToken">Cancellation token for the operation</param>
 		/// <returns>The document</returns>
-		public async Task<T> GetSingletonAsync<T>(Func<T> constructor, CancellationToken cancellationToken) where T : SingletonBase, new()
+		internal static async Task<T> GetSingletonStaticAsync<T>(IMongoCollection<BsonDocument> singletonsV1, IMongoCollection<BsonDocument> singletonsV2, Func<T> constructor, CancellationToken cancellationToken) where T : SingletonBase, new()
 		{
 			SingletonDocumentAttribute attribute = SingletonInfo<T>.Attribute;
 
 			FilterDefinition<BsonDocument> filter = new BsonDocument(new BsonElement("_id", attribute.Id));
 			for (; ; )
 			{
-				BsonDocument? document = await SingletonsV2.Find(filter).FirstOrDefaultAsync(cancellationToken);
+				BsonDocument? document = await singletonsV2.Find(filter).FirstOrDefaultAsync(cancellationToken);
 				if (document != null)
 				{
 					T item = BsonSerializer.Deserialize<T>(document);
@@ -822,7 +843,7 @@ namespace HordeServer.Server
 				T? newItem = null;
 				if (attribute.LegacyId != null)
 				{
-					BsonDocument? legacyDocument = await SingletonsV1.Find(new BsonDocument(new BsonElement("_id", ObjectId.Parse(attribute.LegacyId)))).FirstOrDefaultAsync(cancellationToken);
+					BsonDocument? legacyDocument = await singletonsV1.Find(new BsonDocument(new BsonElement("_id", ObjectId.Parse(attribute.LegacyId)))).FirstOrDefaultAsync(cancellationToken);
 					if (legacyDocument != null)
 					{
 						legacyDocument.Remove("_id");
@@ -833,64 +854,25 @@ namespace HordeServer.Server
 				newItem ??= constructor();
 
 				newItem.Id = new SingletonId(attribute.Id);
-				await SingletonsV2.InsertOneIgnoreDuplicatesAsync(newItem.ToBsonDocument(), cancellationToken);
-			}
-		}
-
-		/// <summary>
-		/// Updates a singleton
-		/// </summary>
-		/// <typeparam name="T"></typeparam>
-		/// <param name="updater"></param>
-		/// <param name="cancellationToken">Cancellation token for the operation</param>
-		/// <returns></returns>
-		public Task UpdateSingletonAsync<T>(Action<T> updater, CancellationToken cancellationToken) where T : SingletonBase, new()
-		{
-			bool Update(T instance)
-			{
-				updater(instance);
-				return true;
-			}
-			return UpdateSingletonAsync<T>(Update, cancellationToken);
-		}
-
-		/// <summary>
-		/// Updates a singleton
-		/// </summary>
-		/// <typeparam name="T"></typeparam>
-		/// <param name="updater"></param>
-		/// <param name="cancellationToken">Cancellation token for the operation</param>
-		/// <returns></returns>
-		public async Task UpdateSingletonAsync<T>(Func<T, bool> updater, CancellationToken cancellationToken) where T : SingletonBase, new()
-		{
-			for (; ; )
-			{
-				T document = await GetSingletonAsync(() => new T(), cancellationToken);
-				if (!updater(document))
-				{
-					break;
-				}
-				if (await TryUpdateSingletonAsync(document, cancellationToken))
-				{
-					break;
-				}
+				await singletonsV2.InsertOneIgnoreDuplicatesAsync(newItem.ToBsonDocument(), cancellationToken);
 			}
 		}
 
 		/// <summary>
 		/// Attempts to update a singleton object
 		/// </summary>
+		/// <param name="singletonsV2">Collection of Singleton documents V2</param>
 		/// <param name="singletonObject">The singleton object</param>
 		/// <param name="cancellationToken">Cancellation token for the operation</param>
 		/// <returns>True if the singleton document was updated</returns>
-		public async Task<bool> TryUpdateSingletonAsync<T>(T singletonObject, CancellationToken cancellationToken) where T : SingletonBase
+		internal static async Task<bool> TryUpdateSingletonStaticAsync<T>(IMongoCollection<BsonDocument> singletonsV2, T singletonObject, CancellationToken cancellationToken) where T : SingletonBase
 		{
 			int prevRevision = singletonObject.Revision++;
 
 			BsonDocument filter = new BsonDocument { new BsonElement("_id", singletonObject.Id.ToString()), new BsonElement(nameof(SingletonBase.Revision), prevRevision) };
 			try
 			{
-				ReplaceOneResult result = await SingletonsV2.ReplaceOneAsync(filter, singletonObject.ToBsonDocument(), new ReplaceOptions { IsUpsert = true }, cancellationToken);
+				ReplaceOneResult result = await singletonsV2.ReplaceOneAsync(filter, singletonObject.ToBsonDocument(), new ReplaceOptions { IsUpsert = true }, cancellationToken);
 				return result.MatchedCount > 0;
 			}
 			catch (MongoWriteException ex)
diff --git a/Engine/Source/Programs/Horde/HordeServer/Server/RedisProcess.cs b/Engine/Source/Programs/Horde/HordeServer/Server/RedisProcess.cs
--- a/Engine/Source/Programs/Horde/HordeServer/Server/RedisProcess.cs
+++ b/Engine/Source/Programs/Horde/HordeServer/Server/RedisProcess.cs
@@ -81,7 +81,7 @@ namespace HordeServer.Server
 
 			// Launch the child process
 			_processGroup = new ManagedProcessGroup();
-			_process = new ManagedProcess(_processGroup, tempRedisExe.FullName, $"{arguments} --bind 127.0.0.1 --port {_port} --save \"\" --appendonly no", redisExe.Directory.FullName, null, ProcessPriorityClass.Normal);
+			_process = new ManagedProcess(_processGroup, tempRedisExe.FullName, $"{arguments} --bind 127.0.0.1 --port {_port} --save \"\" --appendonly no --databases 3000", redisExe.Directory.FullName, null, ProcessPriorityClass.Normal);
 			_process.StdIn.Close();
 
 			_relayTask = BackgroundTask.StartNew(RelayOutputAsync);
diff --git a/Engine/Source/Programs/Horde/Scripts/test.sh b/Engine/Source/Programs/Horde/Scripts/test.sh
--- a/Engine/Source/Programs/Horde/Scripts/test.sh
+++ b/Engine/Source/Programs/Horde/Scripts/test.sh
@@ -4,7 +4,7 @@
 export UE_DOTNET_VERSION=net8.0
 
 # Start Redis and MongoDB in the background for tests to use
-redis-server --save "" --appendonly no --daemonize yes || exit 1
+redis-server --save "" --appendonly no --daemonize yes --databases 3000 || exit 1
 mongod --noauth --quiet --fork --dbpath /tmp/mongodb --logpath /tmp/mongod.log || exit 1
 
 testProjects=(
