Horde: Enable parallel test execution for HordeServer.Build.Tests and HordeServer.Compute.Tests

diff --git a/Engine/Source/Programs/Horde/Plugins/Build/HordeServer.Build.Tests/AssemblyInfo.cs b/Engine/Source/Programs/Horde/Plugins/Build/HordeServer.Build.Tests/AssemblyInfo.cs
new file mode 100644
--- /dev/null
+++ b/Engine/Source/Programs/Horde/Plugins/Build/HordeServer.Build.Tests/AssemblyInfo.cs
@@ -0,0 +1,3 @@
+// Copyright Epic Games, Inc. All Rights Reserved.
+
+[assembly: Parallelize(Workers = 0, Scope = ExecutionScope.MethodLevel)]
\ No newline at end of file
diff --git a/Engine/Source/Programs/Horde/Plugins/Build/HordeServer.Build.Tests/Fleet/JobQueueStrategyTest.cs b/Engine/Source/Programs/Horde/Plugins/Build/HordeServer.Build.Tests/Fleet/JobQueueStrategyTest.cs
--- a/Engine/Source/Programs/Horde/Plugins/Build/HordeServer.Build.Tests/Fleet/JobQueueStrategyTest.cs
+++ b/Engine/Source/Programs/Horde/Plugins/Build/HordeServer.Build.Tests/Fleet/JobQueueStrategyTest.cs
@@ -21,6 +21,7 @@ using Microsoft.Extensions.Options;
 namespace HordeServer.Tests.Fleet
 {
 	[TestClass]
+	[DoNotParallelize] // Possibly due to ConfigService being global
 	public class JobQueueStrategyTest : BuildTestSetup
 	{
 		[TestMethod]
diff --git a/Engine/Source/Programs/Horde/Plugins/Compute/HordeServer.Compute.Tests/AssemblyInfo.cs b/Engine/Source/Programs/Horde/Plugins/Compute/HordeServer.Compute.Tests/AssemblyInfo.cs
new file mode 100644
--- /dev/null
+++ b/Engine/Source/Programs/Horde/Plugins/Compute/HordeServer.Compute.Tests/AssemblyInfo.cs
@@ -0,0 +1,3 @@
+// Copyright Epic Games, Inc. All Rights Reserved.
+
+[assembly: Parallelize(Workers = 0, Scope = ExecutionScope.MethodLevel)]
\ No newline at end of file
