Horde: Enable parallel test running for all Horde server related projects

diff --git a/Engine/Source/Programs/Horde/HordeServer.Tests/AssemblyInfo.cs b/Engine/Source/Programs/Horde/HordeServer.Tests/AssemblyInfo.cs
new file mode 100644
--- /dev/null
+++ b/Engine/Source/Programs/Horde/HordeServer.Tests/AssemblyInfo.cs
@@ -0,0 +1,5 @@
+// Copyright Epic Games, Inc. All Rights Reserved.
+
+using Microsoft.VisualStudio.TestTools.UnitTesting;
+
+[assembly: Parallelize(Workers = 0, Scope = ExecutionScope.MethodLevel)]
\ No newline at end of file
diff --git a/Engine/Source/Programs/Horde/HordeServer.Tests/DatabaseRunnerTest.cs b/Engine/Source/Programs/Horde/HordeServer.Tests/DatabaseRunnerTest.cs
--- a/Engine/Source/Programs/Horde/HordeServer.Tests/DatabaseRunnerTest.cs
+++ b/Engine/Source/Programs/Horde/HordeServer.Tests/DatabaseRunnerTest.cs
@@ -22,6 +22,7 @@ namespace HordeServer.Tests
 		}
 
 		[TestMethod]
+		[Ignore("Very slow to stop on Windows. Likely has to do with the specific Windows-version of redis-server.exe")]
 		public async Task RunRedisTestAsync()
 		{
 			if (RuntimeInformation.IsOSPlatform(OSPlatform.Windows))
diff --git a/Engine/Source/Programs/Horde/Plugins/Analytics/HordeServer.Analytics.Tests/AssemblyInfo.cs b/Engine/Source/Programs/Horde/Plugins/Analytics/HordeServer.Analytics.Tests/AssemblyInfo.cs
new file mode 100644
--- /dev/null
+++ b/Engine/Source/Programs/Horde/Plugins/Analytics/HordeServer.Analytics.Tests/AssemblyInfo.cs
@@ -0,0 +1,3 @@
+// Copyright Epic Games, Inc. All Rights Reserved.
+
+[assembly: Parallelize(Workers = 0, Scope = ExecutionScope.MethodLevel)]
\ No newline at end of file
diff --git a/Engine/Source/Programs/Horde/Plugins/Experimental/HordeServer.Experimental.Tests/AssemblyInfo.cs b/Engine/Source/Programs/Horde/Plugins/Experimental/HordeServer.Experimental.Tests/AssemblyInfo.cs
new file mode 100644
--- /dev/null
+++ b/Engine/Source/Programs/Horde/Plugins/Experimental/HordeServer.Experimental.Tests/AssemblyInfo.cs
@@ -0,0 +1,3 @@
+// Copyright Epic Games, Inc. All Rights Reserved.
+
+[assembly: Parallelize(Workers = 0, Scope = ExecutionScope.MethodLevel)]
\ No newline at end of file
diff --git a/Engine/Source/Programs/Horde/Plugins/Secrets/HordeServer.Secrets.Tests/AssemblyInfo.cs b/Engine/Source/Programs/Horde/Plugins/Secrets/HordeServer.Secrets.Tests/AssemblyInfo.cs
new file mode 100644
--- /dev/null
+++ b/Engine/Source/Programs/Horde/Plugins/Secrets/HordeServer.Secrets.Tests/AssemblyInfo.cs
@@ -0,0 +1,3 @@
+// Copyright Epic Games, Inc. All Rights Reserved.
+
+[assembly: Parallelize(Workers = 0, Scope = ExecutionScope.MethodLevel)]
\ No newline at end of file
diff --git a/Engine/Source/Programs/Horde/Plugins/Storage/HordeServer.Storage.Tests/AssemblyInfo.cs b/Engine/Source/Programs/Horde/Plugins/Storage/HordeServer.Storage.Tests/AssemblyInfo.cs
new file mode 100644
--- /dev/null
+++ b/Engine/Source/Programs/Horde/Plugins/Storage/HordeServer.Storage.Tests/AssemblyInfo.cs
@@ -0,0 +1,3 @@
+// Copyright Epic Games, Inc. All Rights Reserved.
+
+[assembly: Parallelize(Workers = 0, Scope = ExecutionScope.MethodLevel)]
\ No newline at end of file
diff --git a/Engine/Source/Programs/Horde/Plugins/Tools/HordeServer.Tools.Tests/AssemblyInfo.cs b/Engine/Source/Programs/Horde/Plugins/Tools/HordeServer.Tools.Tests/AssemblyInfo.cs
new file mode 100644
--- /dev/null
+++ b/Engine/Source/Programs/Horde/Plugins/Tools/HordeServer.Tools.Tests/AssemblyInfo.cs
@@ -0,0 +1,3 @@
+// Copyright Epic Games, Inc. All Rights Reserved.
+
+[assembly: Parallelize(Workers = 0, Scope = ExecutionScope.MethodLevel)]
\ No newline at end of file
