[Horde][BuildHealth][Dashboard] D3 container fixes for label/no-label toggle

diff --git a/Engine/Source/Programs/Horde/HordeDashboard/src/components/buildhealth/stepoutcome/StepOutcomeD3Component.tsx b/Engine/Source/Programs/Horde/HordeDashboard/src/components/buildhealth/stepoutcome/StepOutcomeD3Component.tsx
--- a/Engine/Source/Programs/Horde/HordeDashboard/src/components/buildhealth/stepoutcome/StepOutcomeD3Component.tsx
+++ b/Engine/Source/Programs/Horde/HordeDashboard/src/components/buildhealth/stepoutcome/StepOutcomeD3Component.tsx
@@ -360,19 +360,21 @@ function GenerateD3Table(stepOutcomeTable: StepOutcomeTable, handler: StepOutcom
         rows.each(function (rowData: StepNameHeader) {
             const tr = d3.select(this);
             const stream = rowData.streamId;
+            let metadata = stepOutcomeTable.tableStreamMetadata.get(stream) ?? { stepCount: 1, labelCount: 0 };
+            let rowSpan = stepOutcomeTable.viewModelOptions.showLabels ? metadata?.labelCount! + metadata?.stepCount! : metadata?.stepCount;
 
-            // Only add a TH for the first row of each stream
+            // Only add a TH for the first row of each stream   
             const streamData = !streamTracker.has(stream)
                 ? [{
                     streamId: stream,
-                    rowspan: stepOutcomeTable.tableStreamMetadata.get(stream) ?? 1
+                    rowspan: rowSpan
                 }]
                 : [];
 
             tr.selectAll("th.stream")
                 .data(streamData)
                 .join(
-                    enter => enter.append("th")
+                    enter => enter.insert("th", ":first-child")
                         .attr("class", "stream")
                         .attr("rowspan", d => d.rowspan)
                         .style("position", "sticky")
@@ -393,7 +395,9 @@ function GenerateD3Table(stepOutcomeTable: StepOutcomeTable, handler: StepOutcom
                         .style("border-bottom", "1px solid")
                         .style("border-top", "1px solid")
                         .text(d => d.streamId),
-                    update => update.text(d => d.streamId),
+                    update => update
+                        .attr("rowspan", d => d.rowspan)
+                        .text(d => d.streamId),
                     exit => exit.remove()
                 );
 
diff --git a/Engine/Source/Programs/Horde/HordeDashboard/src/components/buildhealth/stepoutcome/StepOutcomeDataTypes.ts b/Engine/Source/Programs/Horde/HordeDashboard/src/components/buildhealth/stepoutcome/StepOutcomeDataTypes.ts
--- a/Engine/Source/Programs/Horde/HordeDashboard/src/components/buildhealth/stepoutcome/StepOutcomeDataTypes.ts
+++ b/Engine/Source/Programs/Horde/HordeDashboard/src/components/buildhealth/stepoutcome/StepOutcomeDataTypes.ts
@@ -540,6 +540,14 @@ export type StepNameHeader = {
 
 // #region -- Table Header Types --
 
+/**
+ * Utility type that contains metadata about the stream.
+ */
+export type StreamMetadata = {
+    stepCount: number;
+    labelCount: number;
+};
+
 /**
  * Model for the StepOutcomeTable.
  */
@@ -577,7 +585,7 @@ export class StepOutcomeTable {
      */
     summaryLookup: Map<string, SummaryTableEntry> = new Map<string, SummaryTableEntry>();
 
-    tableStreamMetadata: Map<string, number> = new Map<string, number>();
+    tableStreamMetadata: Map<string, StreamMetadata> = new Map<string, StreamMetadata>();
 
     @observable
     viewModelOptions: StepOutcomeViewOptions;
@@ -601,13 +609,16 @@ export class StepOutcomeTable {
         }
     }
 
-    private annotateStreamRow(streamId: string) {
-        let priorCount: number = this.tableStreamMetadata.get(streamId) ?? 0;
-        if (priorCount) {
-            this.tableStreamMetadata.set(streamId, priorCount);
+    private annotateStreamRow(streamId: string, isLabelOutcome: boolean) {
+        const metadata = this.tableStreamMetadata.get(streamId) ?? { stepCount: 0, labelCount: 0 };
+
+        if (isLabelOutcome) {
+            metadata.labelCount++;
+        } else {
+            metadata.stepCount++;
         }
-        priorCount = priorCount + 1;
-        this.tableStreamMetadata.set(streamId, priorCount);
+
+        this.tableStreamMetadata.set(streamId, metadata);
     }
 
     /**
@@ -681,7 +692,8 @@ export class StepOutcomeTable {
         // Resolve or create row index
         let insertionRow = this.stepRowLookup.get(normalizedStepName);
         if (insertionRow === undefined) {
-            this.annotateStreamRow(tableEntry.streamId);
+
+            this.annotateStreamRow(tableEntry.streamId, isLabelOutcome(tableEntry));
             insertionRow = this.stepNameRowHeaders.length;
             this.stepRowLookup.set(normalizedStepName, insertionRow); // cache the normalized step name -> row number
             this.stepRowNumberLookup.set(insertionRow, normalizedStepName); // cache the row number -> normalized step name
@@ -849,7 +861,7 @@ export class StepOutcomeTable {
             else {
                 const colIdx = this.changeColLookup.get(entry.change);
                 const tableEntry = colIdx !== undefined ? rowEntries[colIdx] : null;
-                
+
                 yield tableEntry !== null && this.viewModelOptions.validStates.has(tableEntry.state) ? tableEntry : undefined;
             }
         }
