[Horde][Analytics][Server] Adding isPreflight column to JobSummary data.

diff --git a/Engine/Source/Programs/Horde/Plugins/Build/HordeServer.Build/Jobs/JobTelemetry.cs b/Engine/Source/Programs/Horde/Plugins/Build/HordeServer.Build/Jobs/JobTelemetry.cs
--- a/Engine/Source/Programs/Horde/Plugins/Build/HordeServer.Build/Jobs/JobTelemetry.cs
+++ b/Engine/Source/Programs/Horde/Plugins/Build/HordeServer.Build/Jobs/JobTelemetry.cs
@@ -109,6 +109,11 @@ namespace HordeServer.Jobs
 		/// </summary>
 		public int StepTotalCount { get; init; }
 
+		/// <summary>
+		/// Whether the job was a preflight or not.
+		/// </summary>
+		public bool IsPreflight { get; init; }
+
 		/// <summary>
 		/// Default event name for the JobSummaryTelemetry.
 		/// </summary>
@@ -117,7 +122,7 @@ namespace HordeServer.Jobs
 		/// <summary>
 		/// Current schema version.
 		/// </summary>
-		public static readonly int CurrentSchemaVersion = 1;
+		public static readonly int CurrentSchemaVersion = 2;
 
 		/// <summary>
 		/// Constructs a JobSummary Telemetry event.
@@ -136,6 +141,7 @@ namespace HordeServer.Jobs
 			CreateTimeUtc = job.CreateTimeUtc;
 			FinishTimeUtc = finishTimeUtc;
 			JobWallTime = (float)(FinishTimeUtc - CreateTimeUtc).TotalSeconds;
+			IsPreflight = job.PreflightCommitId != null;
 
 			JobStepsCompletionInfo jobStepsCompletionInfo = job.GetStepCompletionInfo();
 
