[Build Health Dashboard] Bring audit log dropdown inline with Horde standard, correctly set defaults, and update options based on backend fetches when live tailing

diff --git a/Engine/Source/Programs/Horde/HordeDashboard/src/components/AuditLog.tsx b/Engine/Source/Programs/Horde/HordeDashboard/src/components/AuditLog.tsx
--- a/Engine/Source/Programs/Horde/HordeDashboard/src/components/AuditLog.tsx
+++ b/Engine/Source/Programs/Horde/HordeDashboard/src/components/AuditLog.tsx
@@ -1,9 +1,9 @@
 // Copyright Epic Games, Inc. All Rights Reserved.
-import { ComboBox, DefaultButton, Dropdown, FocusZone, FocusZoneDirection, IDropdownOption, List, mergeStyleSets, SearchBox, Selection, SelectionMode, SelectionZone, Spinner, SpinnerSize, Stack, Text } from "@fluentui/react";
+import { ComboBox, DefaultButton, Dropdown, FocusZone, FocusZoneDirection, IComboBoxOption, List, mergeStyleSets, SearchBox, SelectableOptionMenuItemType, Selection, SelectionMode, SelectionZone, Spinner, SpinnerSize, Stack, Text } from "@fluentui/react";
 import { action, makeObservable, observable } from 'mobx';
 import { observer } from "mobx-react-lite";
 import moment from "moment";
-import { useEffect, useState } from "react";
+import { useEffect, useMemo, useRef, useState } from "react";
 import Highlight from "react-highlighter";
 import { useParams } from "react-router";
 import { Link } from "react-router-dom";
@@ -59,6 +59,10 @@ export type TemplateAudit = {
    templateId: string;
 }
 
+export type AuditLogItem = {
+      entry: AuditLogEntry;
+      type: string;
+   }
 
 
 class AuditLogHandler {
@@ -79,10 +83,13 @@ class AuditLogHandler {
 
       this.timeSelectKey = "time_live_tail";
 
+      this.typedEntries = [];
+      this.typeOptions = [];
+
       clearTimeout(this.timeoutId);
       this.timeoutId = undefined;
 
-      // cancel any pending        
+      // cancel any pending
       for (let i = 0; i < this.cancelId; i++) {
          this.canceled.add(i);
       }
@@ -205,7 +212,7 @@ class AuditLogHandler {
 
       }
 
-      // cancel any pending        
+      // cancel any pending
       for (let i = 0; i < this.cancelId; i++) {
          this.canceled.add(i);
       }
@@ -269,6 +276,7 @@ class AuditLogHandler {
          }
 
          this.entries = entries;
+         this.typedEntries = this.deriveTypedEntries();
 
          this.haveUpdated = true;
 
@@ -285,6 +293,57 @@ class AuditLogHandler {
 
    }
 
+   /**
+    * Creates an array of @see AuditLogItem objects and calls @see this.updateTypesIfChanged.
+    * @returns Sorted array of typed entries for the current audit log.
+    */
+   deriveTypedEntries() {
+      let types: Set<string> = new Set<string>();
+      let typedEntries = this.entries.map((entry) => {
+
+         let type = entry.level.toString();
+         const TYPE_PROPERTY = "Type";
+         const COMPUTE_TYPE = "Compute";
+         if (entry.properties && entry.properties[TYPE_PROPERTY]) {
+            type = entry.properties[TYPE_PROPERTY] as string;
+         } else {
+            // @todo:
+            if (entry.message.startsWith(COMPUTE_TYPE)) {
+               type = COMPUTE_TYPE;
+            }
+         }
+
+         types.add(type);
+
+         return { entry: entry, type: type }
+      })
+
+      this.updateTypesIfChanged(types);
+
+      // Display issue and job audit logs in chronological order, entries in all other audit logs are sorted from most recent to least recent to better surface information
+      if (this.issueId || this.jobId) {
+         typedEntries.reverse();
+      }
+
+      return typedEntries
+   }
+
+    /**
+    * Helper function to update the types present in the audit log, if they have changed since the last handler update.
+    * @param newTypes The set of new types found in the latest fetch of audit log entries.
+    * @returns Sorted array of typed entries for the current audit log.
+    * @remark The reassignment of @see this.typeOptions is intentional to update the ref and let the @see AuditLogPanel component know the set of types has changed.
+    */
+   @action
+   updateTypesIfChanged(newTypes: Set<string>) {
+      // Types unchanged, no need to update
+      if (this.typeOptions.length === newTypes.size && this.typeOptions.every(type => newTypes.has(type))) {
+        return;
+      }
+
+      this.typeOptions = [...newTypes];
+   }
+
    setTimeSelection(time: TimeSelection) {
 
       const live = time.key === "time_live_tail";
@@ -335,6 +394,11 @@ class AuditLogHandler {
 
    timeSelectKey?: string;
 
+   typedEntries: AuditLogItem[];
+
+   @observable
+   typeOptions: string[] = [];
+
    minDate?: Date;
    maxDate?: Date;
 
@@ -408,7 +472,8 @@ const AuditLogPanel: React.FC<{ agentId?: string, deviceId?: string, jobId?: str
    const [showDatePicker, setShowDatePicker] = useState(false);
    const [tsFormat, setTSFormat] = useState(dashboard.displayUTC ? 'UTC' : 'Local');
    const [search, setSearch] = useState<string | undefined>(undefined);
-   const [selectedTypes, setSelectedTypes] = useState<string[]>(["All"]);
+   const [selectedTypes, setSelectedTypes] = useState<string[]>([]);
+   const [typeModifiedByUser, setTypeModifiedByUser] = useState<boolean>(false);
    const [viewAgent, setViewAgent] = useState(false);
    const [viewIssue, setViewIssue] = useState(false);
 
@@ -424,6 +489,7 @@ const AuditLogPanel: React.FC<{ agentId?: string, deviceId?: string, jobId?: str
    const auditStyleNormal = getAuditStyleNormal();
    const { modeColors } = getHordeStyling();
 
+   // Read handler.updated so MobX knows to re-ender the component if it has changed
    if (handler.updated) { }
 
    if (agentId && handler.agentId !== agentId) {
@@ -452,35 +518,10 @@ const AuditLogPanel: React.FC<{ agentId?: string, deviceId?: string, jobId?: str
 
    const vw = Math.max(document.documentElement.clientWidth, window.innerWidth || 0);
 
-   type AuditLogItem = {
-      entry: AuditLogEntry;
-      type: string;
-   }
-
    const searchTerm = search?.toLowerCase();
 
-   const entryTypes: Set<string> = new Set();
-
-   let logItems: AuditLogItem[] = handler.entries.map((entry) => {
-
-      let type = entry.level.toString();
-      if (entry.properties && entry.properties["Type"]) {
-         type = entry.properties["Type"] as string;
-      } else {
-         // @todo: 
-         if (entry.message.startsWith("Compute")) {
-            type = "Compute";
-         }
-      }
-
-      entryTypes.add(type);
-
-      return { entry: entry, type: type }
-
-   })
-
    // filter
-   logItems = logItems.filter(item => {
+   const logItems: AuditLogItem[] = handler.typedEntries?.filter(item => {
 
       if (selectedTypes.indexOf("All") === -1) {
          if (selectedTypes.indexOf(item.type) === -1) {
@@ -500,21 +541,59 @@ const AuditLogPanel: React.FC<{ agentId?: string, deviceId?: string, jobId?: str
 
    });
 
-   if (issueId || jobId) {
-      logItems.reverse();
+   // keep type data and selection logic in this component to avoid re-rendering log items on changes to selection
+   const typeOptions: IComboBoxOption[] = useMemo(() => {
+      
+      let options: IComboBoxOption[]  = Array.from(handler.typeOptions).map(type => {
+         return {
+            key: type,
+            text: type
+         };
+      }).sort((a, b) => { if (a.text === b.text) return 0; return a.text < b.text ? -1 : 1 })
+         
+      if (options.length > 1) {
+         options.unshift({
+            key: 'All',
+            text: 'Select All',
+            itemType: SelectableOptionMenuItemType.SelectAll
+         });
+      }
+
+      return options;
+
+   }, [handler.typeOptions]);
+
+   const prevAllTypesSelected = useRef(false);
+   const allTypesSelected =  selectedTypes.length > 1 && selectedTypes.length === handler.typeOptions.length;
+
+   // derive selected keys at render time to drive selectAll behaviour and avoid adding sentinel selectAll node to array of actual types
+   let selectedTypesDisplayed: string[] = [...selectedTypes];
+   if (allTypesSelected) {
+      selectedTypesDisplayed.push('All');
    }
 
-   const typeItems: IDropdownOption[] = Array.from(entryTypes.keys()).map(type => {
-      return {
-         key: type,
-         text: type
-      };
-   }).sort((a, b) => { if (a.text === b.text) return 0; return a.text < b.text ? -1 : 1 });
+   // types supplied from the backend, so we need to maintain which types should be disbaled by default across all audit logs
+   const disabledByDefaultTypes = ["Trace", "Debug"];
+
+   useEffect(() => {
+
+      // set defaults only if user hasn't interacteds
+      if (!typeModifiedByUser){
+         const defaultTypes: string[] = Array.from(handler.typeOptions).filter(option => !disabledByDefaultTypes.includes(option));
+         setSelectedTypes(defaultTypes);
+      }
+
+      // automatically select new options if a new type is added and all options were selected previously
+      if (prevAllTypesSelected.current) {
+         setSelectedTypes([...handler.typeOptions]);
+      }
+
+   }, [handler.typeOptions]);
 
-   typeItems.unshift({
-      key: "All",
-      text: "All"
-   })
+   // keep track of allTypesSelected value from previous render
+   useEffect(() => {
+      prevAllTypesSelected.current = allTypesSelected;
+   });
 
    let crumbItems: BreadcrumbItem[] = [];
    if (agentId) {
@@ -629,37 +708,29 @@ const AuditLogPanel: React.FC<{ agentId?: string, deviceId?: string, jobId?: str
             <Stack tokens={{ childrenGap: 0 }} styles={{ root: { backgroundColor: modeColors.background, margin: "auto", paddingTop: 12, paddingRight: 10 } }}>
                <Stack horizontal styles={{ root: { paddingLeft: 0, paddingBottom: 4, paddingRight: 12, width: 1440 } }}>
                   <Stack horizontal tokens={{ childrenGap: 12 }}>
-                     <Stack>
-                        <SearchBox style={{ width: 400 }} onEscape={() => setSearch(undefined)} autoComplete="off" disableAnimation={true} spellCheck={false} onChange={(ev, value) => setSearch(value)} />
-                     </Stack>
-                     <Stack>
-                        <Dropdown style={{ width: 240 }} options={typeItems} multiSelect selectedKeys={selectedTypes}
-                           onChange={(event, option, index) => {
 
-                              if (option) {
-
-                                 let filter = [...selectedTypes];
-                                 if (!option.selected) {
-                                    filter = filter.filter(k => k !== option.key);
-                                 } else {
-                                    if (filter.indexOf(option.key as string) === -1) {
-                                       filter.push(option.key as string);
-                                    }
-                                 }
-
-                                 if (!filter.length || (option.selected && option.key === "All")) {
-                                    filter = ["All"];
-                                 }
-
-                                 if (filter.find(k => k === "All") && filter.length > 1) {
-                                    filter = filter.filter(k => k !== "All");
-                                 }
+                        <SearchBox style={{ width: 400 }} onEscape={() => setSearch(undefined)} autoComplete="off" disableAnimation={true} spellCheck={false} onChange={(ev, value) => setSearch(value)} />
 
-                                 setSelectedTypes(filter);
+                        <ComboBox style={{ width: 240 }} multiSelect options={typeOptions} selectedKey={selectedTypesDisplayed}
+                        onChange={(event, option, index) => { 
+                           setTypeModifiedByUser(true);
+                           const selected = option?.selected;
+
+                           if (option) {
+                              // toggle selectAll
+                              if (option?.itemType === SelectableOptionMenuItemType.SelectAll) {
+                                 allTypesSelected ? setSelectedTypes([]) : setSelectedTypes([...handler.typeOptions]);
+                              } else {
+                              // toggle individual type
+                              let updatedKeys = selected ? [...selectedTypes, option!.key as string] : selectedTypes.filter(k => k !== option.key);
+                              if (updatedKeys.length !== selectedTypes.length) {
+                                 updatedKeys = updatedKeys.filter(k => k!== 'All');
                               }
-                           }}
+                              setSelectedTypes(updatedKeys);
+                              }
+                           }
+                        }}
                         />
-                     </Stack>
                   </Stack>
 
                   <Stack grow />
