Horde: Use ManagedWorkspace for AutoSDK when PerforceMaterializer is active

AutoSDK heavily rely on Horde-based view filtering which is not implemented in PerforceMaterializer.
This avoids massive syncs during experimentation.

diff --git a/Engine/Source/Programs/Horde/Drivers/JobDriver/Execution/PerforceMaterializer.cs b/Engine/Source/Programs/Horde/Drivers/JobDriver/Execution/PerforceMaterializer.cs
--- a/Engine/Source/Programs/Horde/Drivers/JobDriver/Execution/PerforceMaterializer.cs
+++ b/Engine/Source/Programs/Horde/Drivers/JobDriver/Execution/PerforceMaterializer.cs
@@ -115,7 +115,15 @@ public sealed class PerforceMaterializer : IWorkspaceMaterializer
 		RpcAgentWorkspace raw = _options.AgentWorkspace;
 		PerforceSettings tempSettings = new(raw.ServerAndPort, raw.UserName) { PreferNativeClient = PreferNative, Password = raw.Ticket };
 		_perforceWithoutClient = await PerforceConnection.CreateAsync(tempSettings, _logger);
-		await _perforceWithoutClient.LoginAsync(raw.Password, cancellationToken);
+		if (String.IsNullOrEmpty(raw.Password))
+		{
+			_logger.LogInformation("Perforce password not set - skipping login");
+		}
+		else
+		{
+			await _perforceWithoutClient.LoginAsync(raw.Password, cancellationToken);
+		}
+		
 		
 		// Get the host name, and fill in any missing metadata about the connection
 		InfoRecord info = await _perforceWithoutClient.GetInfoAsync(InfoOptions.ShortOutput, cancellationToken);
@@ -278,10 +286,11 @@ public sealed class PerforceMaterializer : IWorkspaceMaterializer
 		const string PlaceholderName = "horde-placeholder-name-pm";
 		
 		// Assume latest stream version is wanted. Should possibly be tied to current CL being synced
-		ClientRecord streamView = await _perforceWithoutClient.GetStreamViewAsync(PlaceholderName, streamName, -1, cancellationToken);
+		ClientRecord streamClient = await _perforceWithoutClient.GetStreamViewAsync(PlaceholderName, streamName, -1, cancellationToken);
 		
+		// Horde-based views are unused for now until a filtering syntax compatible with ManagedWorkspace is figured out
 		List<string> hordeView = _options.AgentWorkspace.View.Select(x => ConvertHordeViewToClientView(x, streamName, clientName)).ToList();
-		List<string> view = [..streamView.View.Select(x => x.Replace(PlaceholderName, clientName, StringComparison.Ordinal))];
+		List<string> view = [..streamClient.View.Select(x => x.Replace(PlaceholderName, clientName, StringComparison.Ordinal))];
 		
 		// Create a classic client while still applying the view from stream together with extra Horde-defined view entries.
 		// A stream-based client won't allow changing the view
@@ -393,15 +402,22 @@ public sealed class PerforceMaterializer : IWorkspaceMaterializer
 class PerforceMaterializerFactory(IServiceProvider serviceProvider) : IWorkspaceMaterializerFactory
 {
 	/// <inheritdoc/>
-	public Task<IWorkspaceMaterializer?> CreateMaterializerAsync(string name, RpcAgentWorkspace workspaceInfo, DirectoryReference workspaceDir, bool forAutoSdk, CancellationToken cancellationToken)
+	public async Task<IWorkspaceMaterializer?> CreateMaterializerAsync(string name, RpcAgentWorkspace workspaceInfo, DirectoryReference workspaceDir, bool forAutoSdk, CancellationToken cancellationToken)
 	{
 		if (name.Equals(PerforceMaterializer.TypeName, StringComparison.OrdinalIgnoreCase))
 		{
 			Tracer tracer = serviceProvider.GetRequiredService<Tracer>();
+			if (forAutoSdk)
+			{
+				// Use ManagedWorkspace for AutoSDK until Horde-based views are implemented in PerforceMaterializer
+				ILogger<ManagedWorkspaceMaterializer> mwLogger = serviceProvider.GetRequiredService<ILogger<ManagedWorkspaceMaterializer>>();
+				return await ManagedWorkspaceMaterializer.CreateAsync(workspaceInfo, workspaceDir, true, false, tracer, mwLogger, cancellationToken);
+			}
+			
 			ILogger<PerforceMaterializer> logger = serviceProvider.GetRequiredService<ILogger<PerforceMaterializer>>();
 			PerforceMaterializerOptions options = new (workspaceDir.FullName, workspaceInfo);
-			return Task.FromResult<IWorkspaceMaterializer?>(new PerforceMaterializer(options, tracer, logger));
+			return new PerforceMaterializer(options, tracer, logger);
 		}
-		return Task.FromResult<IWorkspaceMaterializer?>(null);
+		return null;
 	}
 }
\ No newline at end of file
