Keep invalid meta filters for clarity

#rnx

diff --git a/Engine/Source/Programs/Horde/HordeDashboard/plugins/test-automation/testAutomationView.tsx b/Engine/Source/Programs/Horde/HordeDashboard/plugins/test-automation/testAutomationView.tsx
--- a/Engine/Source/Programs/Horde/HordeDashboard/plugins/test-automation/testAutomationView.tsx
+++ b/Engine/Source/Programs/Horde/HordeDashboard/plugins/test-automation/testAutomationView.tsx
@@ -1,6 +1,6 @@
 // Copyright Epic Games, Inc. All Rights Reserved.
 
-import { Stack, FontIcon, IComboBoxOption, IContextualMenuItem, IContextualMenuProps, Text, Spinner, SpinnerSize, Label, DefaultButton, ITag, TagPicker, IPickerItemProps, PrimaryButton, Modal, IconButton } from "@fluentui/react";
+import { Stack, FontIcon, IComboBoxOption, IContextualMenuItem, IContextualMenuProps, Text, Spinner, SpinnerSize, Label, DefaultButton, ITag, TagPicker, IPickerItemProps, PrimaryButton, Modal, IconButton, TooltipHost, DirectionalHint } from "@fluentui/react";
 import { Breadcrumbs } from "horde/components/Breadcrumbs";
 import { TopNav } from "horde/components/TopNav";
 import { getHordeStyling } from "horde/styles/Styles";
@@ -10,12 +10,12 @@ import React, { useEffect, useState, useCallback, useMemo } from "react";
 import { projectStore } from 'horde/backend/ProjectStore';
 import { TestDataHandler, TestDataVersionRegistrar, TestSessionDetails } from "./testData";
 import { TestSummary } from "./testAutomationSummary";
-import { getStreamOptions, getTestSessionStatusColor, MultiOptionChooser, SessionStatusBar } from "./testAutomationCommon";
+import { getStatusColors, getStreamOptions, getTestSessionStatusColor, MultiOptionChooser, SessionStatusBar } from "./testAutomationCommon";
 import { TestPhasesView } from "./testAutomationPhaseView";
 import { JobStepTestDataItem } from "./api";
 
 import { TestDataV2 } from "./testDataV2Fetcher";
-import dashboard from "horde/backend/Dashboard";
+import dashboard, { StatusColor } from "horde/backend/Dashboard";
 
 const handler = new TestDataHandler();
 
@@ -33,7 +33,7 @@ const StreamChooser: React.FC = observer(() => {
     
     return <Stack style={{ paddingTop: 12, paddingBottom: 4 }} horizontal>
             <Label style={{ paddingTop: 5, paddingRight: 5 }}>Stream</Label>
-            <DefaultButton style={{ width: 270, textAlign: "left", height: 30 }} menuProps={menuProps} text={stream ? projectStore.streamById(stream)?.fullname ?? stream : "Select"} />
+            <DefaultButton style={{ width: 270, textAlign: "left", height: 30 }} menuProps={menuProps} text={stream ? projectStore.streamById(stream)?.fullname ?? stream : "Select a stream first"} />
         </Stack>   
 });
 
@@ -85,16 +85,13 @@ const TestChooser: React.FC = observer(() => {
         return closestTags;
     }, [handler.queryLoading, handler.searchUpdated]);
 
-    if (!streamTests.length) {
-        return null;
-    }
-
     const selectedItems: ITag[] = testsSelection.keys().toArray().map((key) => ({key: key.toLowerCase(), name: key} as ITag));
 
     return <Stack style={{ paddingTop: 12, paddingBottom: 4 }} horizontal>
             <Label style={{ paddingTop: 5, paddingRight: 5 }}>Tests</Label>
             <Stack tokens={{ childrenGap: 6 }} style={{ paddingLeft: 12 }} grow>
                 <TagPicker
+                    disabled={!streamTests.length}
                     styles={{ input: { height: 28, width: 40 }, itemsWrapper: { marginRight: 3, marginBottom: 2 } }}
                     onRenderItem={onRenderItem}
                     onRenderSuggestionsItem={onRenderSuggestionsItem}
@@ -123,15 +120,11 @@ const MetaKeyChooser: React.FC<{ onUpdateKey: (key: string) => void }> = observe
     handler.subscribeToFilter();
     handler.subscribeToSearch();
 
-   const cmetakeys: Set<string> = new Set(handler.filterState.metadata?.keys() ?? []);
-   const streamMetadataKeys = handler.allStreamMetadataKeys;
+    const cmetakeys: Set<string> = new Set(handler.filterState.metadata?.keys() ?? []);
+    const streamMetadataKeys = handler.allStreamMetadataKeys;
 
-   if (!streamMetadataKeys.length) {
-        return null;
-   }
-
-   const options: IContextualMenuItem[] = [];
-   streamMetadataKeys.forEach(a => {
+    const options: IContextualMenuItem[] = [];
+    streamMetadataKeys.forEach(a => {
         if (cmetakeys.has(a)) {
             return;
         }
@@ -140,62 +133,63 @@ const MetaKeyChooser: React.FC<{ onUpdateKey: (key: string) => void }> = observe
                 onUpdateKey(a);
             }
         });
-   });
+    });
 
-   const menuProps: IContextualMenuProps = {
+    const menuProps: IContextualMenuProps = {
         shouldFocusOnMount: true,
         subMenuHoverDelay: 0,
         items: options,
-   };
+    };
 
-   return <Stack style={{ paddingTop: 12, paddingBottom: 4 }} horizontal>
-            <Label style={{ paddingTop: 5, paddingRight: 5, whiteSpace: 'nowrap' }}>Filter by</Label>
-            <DefaultButton style={{width: 270, textAlign: "left", height: 30 }} menuProps={menuProps} text="Metadata" />
-        </Stack>   
+    return <Stack style={{ paddingTop: 12, paddingBottom: 4 }} horizontal horizontalAlign="end" title="add a meta filter">
+                <Stack style={{ padding: 5, fontSize: 18, cursor: "default" }} verticalAlign="center"><FontIcon iconName="AddTo" /></Stack>
+                <DefaultButton style={{width: 150, textAlign: "left", height: 30 }} menuProps={menuProps} text="More filters" disabled={!streamMetadataKeys.length}/>
+            </Stack>   
 });
 
 const MetaValueChooser: React.FC<{ metakey: string, onRemove: (key: string) => void }> = observer(({ metakey, onRemove }) => {
 
     handler.subscribeToFilter();
 
-   const cmetavalues: Set<string> = new Set(handler.filterState.metadata?.get(metakey) ?? []);
-   const metavalues = handler.getAllStreamMetadataValues(metakey);
+    const cmetavalues: Set<string> = new Set(handler.filterState.metadata?.get(metakey) ?? []);
+    const metavalues = handler.getAllStreamMetadataValues(metakey);
 
-   if (!metavalues.length) {
-        return null;
-   }
-
-   const options: IComboBoxOption[] = metavalues.map(t => { return { key: t, text: t } });
-
-   return <Stack style={{ paddingTop: 12, paddingBottom: 4, }} horizontal>
-            <Label style={{ paddingTop: 5, paddingRight: 5 }}>{metakey}</Label>
-            <Stack style={{ padding: 1, width: 270 }}>
-                <MultiOptionChooser style={{height: 28}} options={options} initialSelection={cmetavalues.keys().toArray()}
-                updateKeys={(keys) => {
-                    metavalues.forEach(t => {
-                    const selected = keys.includes(t);
-                    if (!selected && cmetavalues.has(t)) {
-                        handler.removeMetadata(metakey, t);
-                    }
-                    else if (selected && !cmetavalues.has(t)) {
-                        handler.addMetadata(metakey, t);
-                    }
-                    });
-                }} />
+    let options: IComboBoxOption[] = !!metavalues.length ? metavalues.map(t => { return { key: t, text: t } }) : cmetavalues.keys().toArray().map(t => { return { key: t, text: `! ${t}` } });
+    const diffmetavalues = cmetavalues.difference(new Set(metavalues));
+    if (!!metavalues.length && diffmetavalues.size > 0) {
+        options = options.concat(diffmetavalues.keys().toArray().map(t => { return {key: t , text: `! ${t}`} })).sort((a, b) => (a.key as string).localeCompare(b.key as string));
+    }
+    const noMatch = !metavalues.length || (cmetavalues.size > 0 && diffmetavalues.size === cmetavalues.size);
+
+    return <TooltipHost
+            content={noMatch ? "meta key or value(s) missing in this stream" : undefined}
+            directionalHint={DirectionalHint.rightCenter}>
+                <Stack style={{ paddingTop: 12, paddingBottom: 4 }} horizontal>
+                <Label style={{ paddingTop: 5, paddingRight: 5 }}>{metakey}</Label>
+                <Stack style={{ width: 270, borderWidth: 1, borderStyle: "solid", borderColor: noMatch ? getStatusColors().get(StatusColor.Failure) : "transparent" }}>
+                    <MultiOptionChooser style={{height: 28}} options={options} initialSelection={cmetavalues.keys().toArray()} disabled={!metavalues.length}
+                        updateKeys={(keys) => {
+                            keys.forEach(t => {
+                                if (!cmetavalues.has(t))  handler.addMetadata(metakey, t);
+                            });
+                            cmetavalues.forEach(t => {
+                                if (!keys.includes(t))  handler.removeMetadata(metakey, t);
+                            });
+                        }} />
+                </Stack>
+                <DefaultButton style={{ minWidth: 0, fontSize: 11, padding: 5, height: 30 }} title="remove filter"
+                        onClick={() => {
+                            if (handler.filterState?.metadata?.keys().some((key) => key === metakey)) {
+                                handler.removeMetadata(metakey);
+                            }
+                            onRemove(metakey);
+                        }}>
+                    <FontIcon iconName="CalculatorMultiply" />
+                </DefaultButton>
             </Stack>
-            <DefaultButton style={{ minWidth: 0, fontSize: 11, padding: 5, height: 30 }} title="remove filter"
-                    onClick={() => {
-                        if (handler.filterState?.metadata?.keys().some((key) => key === metakey)) {
-                            handler.removeMetadata(metakey);
-                        }
-                        onRemove(metakey);
-                    }}>
-                <FontIcon iconName="CalculatorMultiply" />
-            </DefaultButton>
-        </Stack>
+        </TooltipHost>
 });
 
-
 const TestAutomationSidebarLeft: React.FC = observer(() => {
 
     handler.subscribeToFilter();
