Horde: Attach exception to trace span when agent relay allocations fail

diff --git a/Engine/Source/Programs/Horde/Plugins/Compute/HordeServer.Compute/Compute/ComputeService.cs b/Engine/Source/Programs/Horde/Plugins/Compute/HordeServer.Compute/Compute/ComputeService.cs
--- a/Engine/Source/Programs/Horde/Plugins/Compute/HordeServer.Compute/Compute/ComputeService.cs
+++ b/Engine/Source/Programs/Horde/Plugins/Compute/HordeServer.Compute/Compute/ComputeService.cs
@@ -511,6 +511,8 @@ namespace HordeServer.Compute
 			}
 			catch (AgentRelayException are)
 			{
+				span.SetStatus(Status.Error);
+				span.RecordException(are);
 				throw new ComputeServiceException(are.Message, are) { ShowToUser = true };
 			}
 		}
