[Horde][BuildHealth][Dashboard] Batch API for JobStepRef

diff --git a/Engine/Source/Programs/Horde/Plugins/Build/HordeServer.Build/Streams/StreamsController.cs b/Engine/Source/Programs/Horde/Plugins/Build/HordeServer.Build/Streams/StreamsController.cs
--- a/Engine/Source/Programs/Horde/Plugins/Build/HordeServer.Build/Streams/StreamsController.cs
+++ b/Engine/Source/Programs/Horde/Plugins/Build/HordeServer.Build/Streams/StreamsController.cs
@@ -276,7 +276,7 @@ namespace HordeServer.Streams
 		[HttpPost]
 		[Route("/api/v1/streams/{streamId}/changes")]
 		[ProducesResponseType(typeof(List<GetCommitResponse>), 200)]
-		public async Task<ActionResult<List<object>>> GetBatchedChangesAsync(StreamId streamId, [FromBody] CommitId[]? changeNumbers = null,  PropertyFilter? filter = null, CancellationToken cancellationToken = default)
+		public async Task<ActionResult<List<object>>> GetBatchedChangesAsync(StreamId streamId, [FromBody] CommitId[]? changeNumbers = null, PropertyFilter? filter = null, CancellationToken cancellationToken = default)
 		{
 			StreamConfig? streamConfig;
 			if (!_buildConfig.Value.TryGetStream(streamId, out streamConfig))
@@ -436,6 +436,81 @@ namespace HordeServer.Streams
 			return steps.ConvertAll(x => PropertyFilter.Apply(CreateGetJobStepRefResponse(x), filter));
 		}
 
+		private const int MaxConcurrentStepHistoryRequests = 32;
+
+		/// <summary>
+		/// Gets the history of a set of steps given a stream and template context.
+		/// </summary>
+		/// <param name="streamId">The owning stream's id.</param>
+		/// <param name="templateId">The owning template's id.</param>
+		/// <param name="steps">The steps to obtain the history of.</param>
+		/// <param name="change">Maximum changelist number to return</param>
+		/// <param name="count">Number of results to return</param>
+		/// <param name="filter">The filter to apply to the results</param>
+		/// <param name="cancellationToken">The cancellation token to use for the job.</param>
+		/// <returns>Http result code along with the corresponding filtered list of <see cref="GetJobStepRefResponse"/>.</returns>
+		[HttpPost]
+		[Route("/api/v1/streams/{streamId}/history")]
+		[ProducesResponseType(typeof(List<GetJobStepRefResponse>), 200)]
+		public async Task<ActionResult<List<object>>> GetStepHistoryAsync(StreamId streamId, [FromQuery] string templateId, [FromQuery] string[] steps, [FromQuery] CommitId? change = null, [FromQuery] int count = 10, [FromQuery] PropertyFilter? filter = null, CancellationToken cancellationToken = default)
+		{
+			StreamConfig? streamConfig;
+			if (!_buildConfig.Value.TryGetStream(streamId, out streamConfig))
+			{
+				return NotFound(streamId);
+			}
+			if (!streamConfig.Authorize(JobAclAction.ViewJob, User))
+			{
+				return Forbid(JobAclAction.ViewJob, streamId);
+			}
+			if (steps == null)
+			{
+				return BadRequest("Null step argument provided.");
+			}
+
+			TemplateId templateIdValue = new TemplateId(templateId);
+
+			CommitIdWithOrder? commitIdWithOrder = null;
+			if (change != null)
+			{
+				commitIdWithOrder = await _commitService.GetOrderedAsync(streamId, change, cancellationToken);
+			}
+
+			List<IJobStepRef> jobStepRefs = [];
+
+			// create a batch
+			using SemaphoreSlim maxStepHistoryActiveRequests = new SemaphoreSlim(MaxConcurrentStepHistoryRequests);
+			Task<List<IJobStepRef>?>[] jobStepRefCollectionTasks = [.. steps.Select(async step =>
+				{
+					await maxStepHistoryActiveRequests.WaitAsync(cancellationToken);
+					try
+					{
+						return await _jobStepRefCollection.GetStepsForNodeAsync(streamId, templateIdValue, step, commitIdWithOrder, true, count, cancellationToken);
+					}
+					catch (OperationCanceledException)
+					{
+						throw;
+					}
+					catch (Exception)
+					{
+						// Swallow all error exceptions; the user will obtain no result items.
+						return null;
+					}
+					finally
+					{
+						maxStepHistoryActiveRequests.Release();
+					}
+				})];
+
+			List<IJobStepRef>?[] jobStepRefCollectionTasksResults = await Task.WhenAll(jobStepRefCollectionTasks);
+
+			jobStepRefs = jobStepRefCollectionTasksResults.Where(kv => kv != null).SelectMany(kv => kv!).ToList();
+
+			List<object> filteredList = jobStepRefs.Select(y => PropertyFilter.Apply(CreateGetJobStepRefResponse(y), filter)).ToList();
+
+			return filteredList;
+		}
+
 		/// <summary>
 		/// Constructor
 		/// </summary>
@@ -443,6 +518,7 @@ namespace HordeServer.Streams
 		internal static GetJobStepRefResponse CreateGetJobStepRefResponse(IJobStepRef jobStepRef)
 		{
 			GetJobStepRefResponse response = new GetJobStepRefResponse();
+			response.StepName = jobStepRef.NodeName;
 			response.JobId = jobStepRef.Id.JobId;
 			response.BatchId = jobStepRef.Id.BatchId;
 			response.StepId = jobStepRef.Id.StepId;
@@ -450,6 +526,7 @@ namespace HordeServer.Streams
 			response.LogId = jobStepRef.LogId.ToString();
 			response.PoolId = jobStepRef.PoolId?.ToString();
 			response.AgentId = jobStepRef.AgentId?.ToString();
+			response.State = jobStepRef.State;
 			response.Outcome = jobStepRef.Outcome;
 			response.IssueIds = jobStepRef.IssueIds?.Select(id => id).ToList();
 			response.JobStartTime = jobStepRef.JobStartTimeUtc;
diff --git a/Engine/Source/Programs/Shared/EpicGames.Horde/Jobs/IJobStepRef.cs b/Engine/Source/Programs/Shared/EpicGames.Horde/Jobs/IJobStepRef.cs
--- a/Engine/Source/Programs/Shared/EpicGames.Horde/Jobs/IJobStepRef.cs
+++ b/Engine/Source/Programs/Shared/EpicGames.Horde/Jobs/IJobStepRef.cs
@@ -28,7 +28,7 @@ namespace HordeServer.Jobs
 		public string JobName { get; }
 
 		/// <summary>
-		/// Name of the name
+		/// Name of the node
 		/// </summary>
 		public string NodeName { get; }
 
@@ -62,6 +62,11 @@ namespace HordeServer.Jobs
 		/// </summary>
 		public AgentId? AgentId { get; }
 
+		/// <summary>
+		/// State of the step.
+		/// </summary>
+		public JobStepState? State{ get; }
+
 		/// <summary>
 		/// Outcome of the step, once complete.
 		/// </summary>
diff --git a/Engine/Source/Programs/Shared/EpicGames.Horde/Jobs/JobMessages.cs b/Engine/Source/Programs/Shared/EpicGames.Horde/Jobs/JobMessages.cs
--- a/Engine/Source/Programs/Shared/EpicGames.Horde/Jobs/JobMessages.cs
+++ b/Engine/Source/Programs/Shared/EpicGames.Horde/Jobs/JobMessages.cs
@@ -1499,6 +1499,11 @@ namespace EpicGames.Horde.Jobs
 		/// </summary>
 		public JobId JobId { get; set; }
 
+		/// <summary>
+		/// The step name.
+		/// </summary>
+		public string? StepName { get; set; }
+
 		/// <summary>
 		/// The batch containing the step
 		/// </summary>
@@ -1545,6 +1550,11 @@ namespace EpicGames.Horde.Jobs
 		/// </summary>
 		public string? AgentId { get; set; }
 
+		/// <summary>
+		/// The step state.
+		/// </summary>
+		public JobStepState? State { get; set; }
+
 		/// <summary>
 		/// Outcome of the step, once complete.
 		/// </summary>
