Implement Test Audit API on Horde
Add /api/v3/testdata/audit/post
and /api/v3/testdata/audit/history

#rnx

diff --git a/Engine/Source/Programs/Horde/Plugins/Build/HordeServer.Build.Tests/Jobs/TestData/TestDataTests.cs b/Engine/Source/Programs/Horde/Plugins/Build/HordeServer.Build.Tests/Jobs/TestData/TestDataTests.cs
--- a/Engine/Source/Programs/Horde/Plugins/Build/HordeServer.Build.Tests/Jobs/TestData/TestDataTests.cs
+++ b/Engine/Source/Programs/Horde/Plugins/Build/HordeServer.Build.Tests/Jobs/TestData/TestDataTests.cs
@@ -570,6 +570,159 @@ namespace HordeServer.Tests.Jobs.TestData
 			Assert.AreEqual(testRecipes[0].Id, sessions[0].RecipeId);
 		}
 
+		[TestMethod]
+		public async Task SessionReportTestV2AuditAsync()
+		{
+			await TestDataService.StartAsync(CancellationToken.None);
+
+			string[] streamIds = new string[] { _mainStreamId.ToString(), _releaseStreamId.ToString() };
+
+			IJob job = CreateJob(_mainStreamId, 105, "Test Build", _graph);
+			IJobStep step = job.Batches[0].Steps[0];
+			IJob job2 = CreateJob(_releaseStreamId, 105, "Test Build", _graph);
+			IJobStep step2 = job2.Batches[0].Steps[0];
+			IJob job3 = CreateJob(_mainStreamId, 106, "Test Build", _graph);
+			IJobStep step3 = job3.Batches[0].Steps[0];
+
+			BsonDocument testData = BsonDocument.Parse(String.Join('\n', _sessionTestDataV2Lines4));
+			List<(string key, BsonDocument)> data = new List<(string key, BsonDocument)>();
+			foreach (BsonValue item in testData.GetValue("items").AsBsonArray.ToList())
+			{
+				BsonDocument value = item.AsBsonDocument.GetValue("data").AsBsonDocument;
+				data.Add((item.AsBsonDocument.GetValue("key").AsString, value));
+			}
+
+			BsonDocument testData2 = BsonDocument.Parse(String.Join('\n', _sessionTestDataV2Lines2));
+			List<(string key, BsonDocument)> data2 = new List<(string key, BsonDocument)>();
+			foreach (BsonValue item in testData2.GetValue("items").AsBsonArray.ToList())
+			{
+				BsonDocument value = item.AsBsonDocument.GetValue("data").AsBsonDocument;
+				data2.Add((item.AsBsonDocument.GetValue("key").AsString, value));
+			}
+
+			await TestDataCollectionV2.AddAsync(job, step, data.ToArray());
+			await TestDataCollectionV2.AddAsync(job2, step2, data.ToArray());
+			await TestDataCollectionV2.AddAsync(job3, step3, data2.ToArray());
+
+			await TestDataService.TickForTestingAsync();
+
+			GetStreamTestsRequest streamRequest = new GetStreamTestsRequest() { StreamIds = streamIds.ToList() };
+			ActionResult<List<GetTestSessionStreamResponse>> streamResult = await TestDataController.GetTestStreamsV2Async(streamRequest);
+			Assert.IsNotNull(streamResult);
+			Assert.IsNotNull(streamResult.Value);
+			List<GetTestSessionStreamResponse> streams = streamResult.Value;
+			Assert.AreNotEqual(0, streams.Count);
+			Assert.AreEqual(2, streams[0].Tests.Count);
+
+			GetTestAuditRequest testRequestAll = new GetTestAuditRequest() { TestIds = streams[0].Tests.Select(t => t.Id).ToList() };
+			ActionResult<List<GetTestAuditResponse>> testauditResult = await TestDataController.GetTestAuditsV2Async(testRequestAll);
+			Assert.IsNotNull(testauditResult);
+			Assert.IsNotNull(testauditResult.Value);
+			List<GetTestAuditResponse> testaudits = testauditResult.Value;
+			Assert.AreEqual(0, testaudits.Count);
+
+			string? userIdString = TestDataController.HttpContext.User.GetUserId()?.ToString();
+
+			// Add one entry for a test id
+			DateTime auditDate = DateTime.UtcNow;
+			string testId1 = streams[0].Tests.FirstOrDefault()!.Id;
+			AddTestAuditRequest testAuditPostRequest = new AddTestAuditRequest()
+			{
+				TestId = testId1,
+				Notes = "Some audit notes",
+				LastAuditDate = auditDate,
+			};
+			ActionResult<GetTestAuditResponse> testauditPostResult = await TestDataController.AddTestAuditV2Async(testAuditPostRequest);
+			Assert.IsNotNull(testauditPostResult);
+			Assert.IsNotNull(testauditPostResult.Value);
+			GetTestAuditResponse testauditPost = testauditPostResult.Value;
+			Assert.AreEqual(testId1, testauditPost.TestId);
+			Assert.AreEqual("Some audit notes", testauditPost.Notes);
+			Assert.AreEqual(auditDate, testauditPost.LastAuditDate);
+			Assert.AreEqual(true, testauditPost.EnableSchedule);
+			Assert.AreEqual(userIdString, testauditPost.UserId);
+			Assert.IsTrue(testauditPost.UserInputDate > auditDate);
+			Assert.IsNull(testauditPost.Intent);
+
+			// Add a second entry for another test id
+			string testId2 = streams[0].Tests[1].Id;
+			AddTestAuditRequest testAuditPostRequest2 = new AddTestAuditRequest()
+			{
+				TestId = testId2,
+				Notes = "Some audit notes 2",
+				LastAuditDate = auditDate,
+			};
+			testauditPostResult = await TestDataController.AddTestAuditV2Async(testAuditPostRequest2);
+			Assert.IsNotNull(testauditPostResult);
+			Assert.IsNotNull(testauditPostResult.Value);
+			GetTestAuditResponse testauditPost2 = testauditPostResult.Value;
+			Assert.AreEqual(testId2, testauditPost2.TestId);
+			Assert.AreEqual("Some audit notes 2", testauditPost2.Notes);
+			Assert.AreEqual(auditDate, testauditPost2.LastAuditDate);
+			Assert.AreEqual(true, testauditPost2.EnableSchedule);
+			Assert.IsNull(testauditPost2.Intent);
+
+			// Check history output
+			GetTestAuditRequest testRequestOne = new GetTestAuditRequest() { TestIds = [testId1] };
+			testauditResult = await TestDataController.GetTestAuditsV2Async(testRequestOne);
+			Assert.IsNotNull(testauditResult);
+			Assert.IsNotNull(testauditResult.Value);
+			testaudits = testauditResult.Value;
+			Assert.AreEqual(1, testaudits.Count);
+			testauditPost = testaudits[0];
+			Assert.AreEqual(testId1, testauditPost.TestId);
+			Assert.AreEqual("Some audit notes", testauditPost.Notes);
+			Assert.AreEqual(auditDate.ToString(), testauditPost.LastAuditDate.ToString());
+			Assert.AreEqual(true, testauditPost.EnableSchedule);
+			Assert.IsNull(testauditPost.Intent);
+			Assert.AreEqual(userIdString, testauditPost.UserId);
+			Assert.IsTrue(testauditPost.UserInputDate > auditDate);
+
+			// Add one more entry for first test id
+			testAuditPostRequest = new AddTestAuditRequest()
+			{
+				TestId = testId1,
+				Notes = "Some audit notes updated",
+				Intent = "Test intent",
+				Owner = userIdString,
+				Customers = [userIdString],
+				Notifications = new Dictionary<string, string>(){ { _mainStreamId.ToString(), "test-workflow1" } , { _releaseStreamId.ToString(), "test-workflow2" } }
+			};
+			testauditPostResult = await TestDataController.AddTestAuditV2Async(testAuditPostRequest);
+			Assert.IsNotNull(testauditPostResult);
+			Assert.IsNotNull(testauditPostResult.Value);
+			testauditPost = testauditPostResult.Value;
+			Assert.AreEqual(testId1, testauditPost.TestId);
+			Assert.AreEqual("Some audit notes updated", testauditPost.Notes);
+			Assert.AreEqual("Test intent", testauditPost.Intent);
+
+			// Check history output with time range for one test
+			testauditResult = await TestDataController.GetTestAuditsV2Async(testRequestOne, DateTimeOffset.UtcNow.AddHours(-1));
+			Assert.IsNotNull(testauditResult);
+			Assert.IsNotNull(testauditResult.Value);
+			testaudits = testauditResult.Value;
+			Assert.AreEqual(2, testaudits.Count);
+			Assert.IsNull(testaudits[0].LastAuditDate);
+			Assert.AreEqual(auditDate.ToString(), testaudits[1].LastAuditDate.ToString());
+
+			// Check history output for latest for 2 tests
+			testRequestOne = new GetTestAuditRequest() { TestIds = [testId1, testId2] };
+			testauditResult = await TestDataController.GetTestAuditsV2Async(testRequestOne);
+			Assert.IsNotNull(testauditResult);
+			Assert.IsNotNull(testauditResult.Value);
+			testaudits = testauditResult.Value;
+			Assert.AreEqual(2, testaudits.Count);
+			Assert.IsTrue(testRequestOne.TestIds.All(t => testaudits.Any( a => a.TestId == t)));
+			testauditPost = testaudits[0];
+			Assert.AreEqual("Some audit notes updated", testauditPost.Notes);
+			Assert.AreEqual(userIdString, testauditPost.Owner);
+			Assert.IsTrue(testauditPost.Customers?.All(c => c == userIdString));
+			Assert.IsTrue(streamIds.All(s => testauditPost.Notifications?.Any(n => n.Key == s) ?? false));
+			Assert.IsTrue(new List<string>(["test-workflow1", "test-workflow2"]).All(w => testauditPost.Notifications?.Any(n => n.Value == w) ?? false));
+			Assert.IsNull(testauditPost.LastAuditDate);
+			Assert.AreEqual("Test intent", testauditPost.Intent);
+		}
+
 		private readonly string[] _simpleTestDataLines =
 		{
 			@"{",
diff --git a/Engine/Source/Programs/Horde/Plugins/Build/HordeServer.Build/Jobs/TestData/ITestData.cs b/Engine/Source/Programs/Horde/Plugins/Build/HordeServer.Build/Jobs/TestData/ITestData.cs
--- a/Engine/Source/Programs/Horde/Plugins/Build/HordeServer.Build/Jobs/TestData/ITestData.cs
+++ b/Engine/Source/Programs/Horde/Plugins/Build/HordeServer.Build/Jobs/TestData/ITestData.cs
@@ -6,6 +6,7 @@ using EpicGames.Horde.Jobs;
 using EpicGames.Horde.Jobs.Templates;
 using EpicGames.Horde.Jobs.TestData;
 using EpicGames.Horde.Streams;
+using EpicGames.Horde.Users;
 using MongoDB.Bson;
 
 namespace HordeServer.Jobs.TestData
@@ -435,6 +436,103 @@ namespace HordeServer.Jobs.TestData
 		IReadOnlyList<ITestParamEntry>? Params { get; }
 	}
 
+	/// <summary>
+	/// A test notification settings
+	/// </summary>
+	public interface ITestNotificationSettings
+	{
+		/// <summary>
+		/// The Stream associated with this entry
+		/// </summary>
+		StreamId Stream { get; }
+
+		/// <summary>
+		/// The Workflow associated with this entry
+		/// </summary>
+		WorkflowId Workflow { get; }
+	}
+
+	/// <summary>
+	/// A test audit entry
+	/// </summary>
+	public interface ITestAudit
+	{
+		/// <summary>
+		/// The test audit id
+		/// </summary>
+		TestAuditId Id { get; }
+
+		/// <summary>
+		/// The test name id associated with this audit
+		/// </summary>
+		TestId TestId { get; }
+
+		/// <summary>
+		/// The date at which this audit was filed
+		/// </summary>
+		DateTime UserInputDate { get; }
+
+		/// <summary>
+		/// Whether or not this entry is archived
+		/// </summary>
+		bool Archived { get; }
+
+		/// <summary>
+		/// The user id that updated the audit
+		/// </summary>
+		UserId UserId { get; }
+
+		/// <summary>
+		/// Whether or not the associated test should be run on the schedule
+		/// </summary>
+		bool EnableSchedule { get; }
+
+		/// <summary>
+		/// The date at which the audit was last officially taken place
+		/// </summary>
+		DateTime? LastAuditDate { get; }
+
+		/// <summary>
+		/// The owner user id of the associated test
+		/// </summary>
+		UserId? Owner { get; }
+
+		/// <summary>
+		/// The customer user ids of the associated test
+		/// </summary>
+		IReadOnlyList<UserId>? Customers { get; }
+
+		/// <summary>
+		/// The workflow notification settings associated with this test
+		/// </summary>
+		IReadOnlyList<ITestNotificationSettings>? Notifications { get; }
+
+		/// <summary>
+		/// The team associated with this test
+		/// </summary>
+		string? Team { get; }
+
+		/// <summary>
+		/// The harness associated with this test
+		/// </summary>
+		string? Harness { get; }
+
+		/// <summary>
+		/// The intent associated with this test
+		/// </summary>
+		string? Intent { get; }
+
+		/// <summary>
+		/// The profile associated with this test
+		/// </summary>
+		string? Profile { get; }
+
+		/// <summary>
+		/// Notes associated with this test
+		/// </summary>
+		string? Notes { get; }
+	}
+
 	/// <summary>
 	/// A test tag reference
 	/// </summary>
diff --git a/Engine/Source/Programs/Horde/Plugins/Build/HordeServer.Build/Jobs/TestData/ITestDataCollection.cs b/Engine/Source/Programs/Horde/Plugins/Build/HordeServer.Build/Jobs/TestData/ITestDataCollection.cs
--- a/Engine/Source/Programs/Horde/Plugins/Build/HordeServer.Build/Jobs/TestData/ITestDataCollection.cs
+++ b/Engine/Source/Programs/Horde/Plugins/Build/HordeServer.Build/Jobs/TestData/ITestDataCollection.cs
@@ -1,8 +1,10 @@
 // Copyright Epic Games, Inc. All Rights Reserved.
 
 using EpicGames.Horde.Commits;
+using EpicGames.Horde.Issues;
 using EpicGames.Horde.Jobs;
 using EpicGames.Horde.Streams;
+using EpicGames.Horde.Users;
 using MongoDB.Bson;
 
 namespace HordeServer.Jobs.TestData
@@ -249,6 +251,25 @@ namespace HordeServer.Jobs.TestData
 		/// <returns></returns>
 		Task<IReadOnlyList<ITestPhaseSession>> FindTestPhaseSessionsAsync(StreamId[] streamIds, TestPhaseId[] phaseIds, DateTime? minCreateTime = null, DateTime? maxCreateTime = null, CommitId? minCommitId = null, CommitId? maxCommitId = null, CancellationToken cancellationToken = default);
 
+		/// <summary>
+		/// Add test audit entry
+		/// </summary>
+		/// <param name="userId"></param>
+		/// <param name="testId"></param>
+		/// <param name="options"></param>
+		/// <param name="cancellationToken">Cancellation token for the operation</param>
+		/// <returns></returns>
+		Task<ITestAudit> AddTestAuditAsync(UserId userId, TestId testId, TestAuditOptions options, CancellationToken cancellationToken = default);
+
+		/// <summary>
+		/// Find audits related to tests
+		/// </summary>
+		/// <param name="testIds"></param>
+		/// <param name="minCreateTime"></param>
+		/// <param name="cancellationToken">Cancellation token for the operation</param>
+		/// <returns></returns>
+		Task<IReadOnlyList<ITestAudit>> FindTestAuditsAsync(TestId[] testIds, DateTime? minCreateTime = null, CancellationToken cancellationToken = default);
+
 		/// <summary>
 		/// Prune test data from collection
 		/// </summary>
@@ -264,4 +285,25 @@ namespace HordeServer.Jobs.TestData
 		/// <returns></returns>
 		Task<bool> UpdateAsync(int retainMonths, CancellationToken cancellationToken = default);
 	}
+
+	public class TestNotificationSettingsOptions
+	{
+		public StreamId Stream { get; set; }
+
+		public WorkflowId Workflow { get; set; }
+	}
+
+	public class TestAuditOptions
+	{
+		public bool EnableSchedule { get; set; } = true;
+		public DateTime? LastAuditDate { get; set; }
+		public UserId? Owner { get; set; }
+		public List<UserId>? Customers { get; set; }
+		public List<TestNotificationSettingsOptions>? Notifications { get; set; }
+		public string? Team { get; set; }
+		public string? Harness { get; set; }
+		public string? Intent { get; set; }
+		public string? Profile { get; set; }
+		public string? Notes { get; set; }
+	}
 }
diff --git a/Engine/Source/Programs/Horde/Plugins/Build/HordeServer.Build/Jobs/TestData/TestAuditId.cs b/Engine/Source/Programs/Horde/Plugins/Build/HordeServer.Build/Jobs/TestData/TestAuditId.cs
new file mode 100644
--- /dev/null
+++ b/Engine/Source/Programs/Horde/Plugins/Build/HordeServer.Build/Jobs/TestData/TestAuditId.cs
@@ -0,0 +1,67 @@
+// Copyright Epic Games, Inc. All Rights Reserved.
+
+using System.ComponentModel;
+using EpicGames.Core;
+using HordeServer.Utilities;
+using MongoDB.Bson;
+
+namespace HordeServer.Jobs.TestData
+{
+	/// <summary>
+	/// Identifier for a test audit
+	/// </summary>
+	/// <param name="Id">Id to construct from</param>
+	[JsonSchemaString]
+	[TypeConverter(typeof(ObjectIdTypeConverter<TestAuditId, TestAuditIdConverter>))]
+	[ObjectIdConverter(typeof(TestAuditIdConverter))]
+	public record struct TestAuditId(ObjectId Id)
+	{
+		/// <summary>
+		/// Default empty value for test audit id
+		/// </summary>
+		public static TestAuditId Empty { get; } = default;
+
+		/// <summary>
+		/// Creates a new <see cref="TestAuditId"/>
+		/// </summary>
+		public static TestAuditId GenerateNewId() => new TestAuditId(ObjectId.GenerateNewId());
+
+		/// <summary>
+		/// Creates a new <see cref="TestAuditId"/>
+		/// </summary>
+		/// <param name="time"></param>
+		/// <returns></returns>
+		public static TestAuditId GenerateNewId(DateTime time) => new TestAuditId(ObjectId.GenerateNewId(time));
+
+		/// <inheritdoc cref="ObjectId.Parse(System.String)"/>
+		public static TestAuditId Parse(string text) => new TestAuditId(ObjectId.Parse(text));
+
+		/// <inheritdoc cref="ObjectId.TryParse(System.String, out ObjectId)"/>
+		public static bool TryParse(string text, out TestAuditId id)
+		{
+			if (ObjectId.TryParse(text, out ObjectId objectId))
+			{
+				id = new TestAuditId(objectId);
+				return true;
+			}
+
+			id = default;
+			return false;
+		}
+
+		/// <inheritdoc/>
+		public override string ToString() => Id.ToString();
+	}
+
+	/// <summary>
+	/// Converter to and from <see cref="ObjectId"/> instances.
+	/// </summary>
+	class TestAuditIdConverter : ObjectIdConverter<TestAuditId>
+	{
+		/// <inheritdoc/>
+		public override TestAuditId FromObjectId(ObjectId id) => new TestAuditId(id);
+
+		/// <inheritdoc/>
+		public override ObjectId ToObjectId(TestAuditId value) => value.Id;
+	}
+}
diff --git a/Engine/Source/Programs/Horde/Plugins/Build/HordeServer.Build/Jobs/TestData/TestDataCollectionV2.cs b/Engine/Source/Programs/Horde/Plugins/Build/HordeServer.Build/Jobs/TestData/TestDataCollectionV2.cs
--- a/Engine/Source/Programs/Horde/Plugins/Build/HordeServer.Build/Jobs/TestData/TestDataCollectionV2.cs
+++ b/Engine/Source/Programs/Horde/Plugins/Build/HordeServer.Build/Jobs/TestData/TestDataCollectionV2.cs
@@ -6,6 +6,7 @@ using EpicGames.Horde.Jobs;
 using EpicGames.Horde.Jobs.Templates;
 using EpicGames.Horde.Jobs.TestData;
 using EpicGames.Horde.Streams;
+using EpicGames.Horde.Users;
 using HordeServer.Commits;
 using HordeServer.Server;
 using HordeServer.Utilities;
@@ -508,6 +509,102 @@ namespace HordeServer.Jobs.TestData
 			}
 		}
 
+		/// <summary>
+		/// Store test notification settings for a stream
+		/// </summary>
+		class TestNotificationSettingsEntry : ITestNotificationSettings
+		{
+			[BsonRequired, BsonElement("s")]
+			public StreamId Stream { get; set; }
+
+			[BsonRequired, BsonElement("w")]
+			public WorkflowId Workflow { get; set; }
+		}
+
+		/// <summary>
+		/// Store a test audit entry
+		/// </summary>
+		class TestAuditDocument : ITestAudit
+		{
+			[BsonRequired, BsonId]
+			public TestAuditId Id { get; set; }
+
+			[BsonRequired, BsonElement("tid")]
+			public TestId TestId { get; set; }
+
+			[BsonRequired, BsonElement("ui")]
+			public DateTime UserInputDate { get; set; }
+
+			[BsonRequired, BsonElement("a")]
+			public bool Archived { get; set; }
+
+			[BsonRequired, BsonElement("uid")]
+			public UserId UserId { get; set; }
+
+			[BsonRequired, BsonElement("e")]
+			public bool EnableSchedule { get; set; }
+
+			[BsonIgnoreIfNull, BsonElement("la")]
+			public DateTime? LastAuditDate { get; set; }
+
+			[BsonIgnoreIfNull, BsonElement("o")]
+			public UserId? Owner { get; set; }
+
+			[BsonIgnoreIfNull, BsonElement("c")]
+			public IReadOnlyList<UserId>? Customers { get; set; }
+
+			[BsonIgnoreIfNull, BsonElement("ns")]
+			public IReadOnlyList<TestNotificationSettingsEntry>? Notifications { get; set; }
+			IReadOnlyList<ITestNotificationSettings>? ITestAudit.Notifications => Notifications;
+
+			[BsonIgnoreIfNull, BsonElement("t")]
+			public string? Team { get; set; }
+
+			[BsonIgnoreIfNull, BsonElement("h")]
+			public string? Harness { get; set; }
+
+			[BsonIgnoreIfNull, BsonElement("i")]
+			public string? Intent { get; set; }
+
+			[BsonIgnoreIfNull, BsonElement("p")]
+			public string? Profile { get; set; }
+
+			[BsonIgnoreIfNull, BsonElement("n")]
+			public string? Notes { get; set; }
+
+			private TestAuditDocument() { }
+
+			public TestAuditDocument(UserId userId, TestId testId)
+			{
+				Id = TestAuditId.GenerateNewId();
+				TestId = testId;
+				UserId = userId;
+				UserInputDate = DateTime.UtcNow;
+				Archived = false;
+				EnableSchedule = true;
+			}
+
+			public void UpdateOptions(TestAuditOptions options)
+			{
+				EnableSchedule = options.EnableSchedule;
+				LastAuditDate = options.LastAuditDate;
+				Owner = options.Owner;
+				Customers = options.Customers;
+				Notifications = options.Notifications?.Select(
+					n => new TestNotificationSettingsEntry()
+					{
+						Stream = n.Stream,
+						Workflow = n.Workflow
+					}
+				).ToList();
+				Team = options.Team;
+				Harness = options.Harness;
+				Intent = options.Intent;
+				Profile = options.Profile;
+				Notes = options.Notes;
+			}
+		}
+
 		/// <summary>
 		/// Information about a test data document
 		/// </summary>
@@ -572,6 +669,7 @@ namespace HordeServer.Jobs.TestData
 		readonly IMongoCollection<TestSessionDocument> _testSessions;
 		readonly IMongoCollection<TestPhaseSessionDocument> _testPhaseSessions;
 		readonly IMongoCollection<TestSessionStreamDocument> _testStreams;
+		readonly IMongoCollection<TestAuditDocument> _testAudits;
 		readonly ICommitService _commitService;
 		readonly Tracer _tracer;
 		readonly ILogger _logger;
@@ -620,6 +718,10 @@ namespace HordeServer.Jobs.TestData
 			List<MongoIndex<TestSessionStreamDocument>> streamIndexes = new List<MongoIndex<TestSessionStreamDocument>>();
 			streamIndexes.Add(keys => keys.Ascending(x => x.StreamId), unique: true);
 			_testStreams = mongoService.GetCollection<TestSessionStreamDocument>("TestDataV2.TestStreams", streamIndexes);
+
+			List<MongoIndex<TestAuditDocument>> auditIndexes = new List<MongoIndex<TestAuditDocument>>();
+			auditIndexes.Add(keys => keys.Ascending(x => x.TestId).Descending(x => x.Archived));
+			_testAudits = mongoService.GetCollection<TestAuditDocument>("TestDataV2.TestAudits", auditIndexes);
 		}
 
 		/// <inheritdoc/>
@@ -1317,6 +1419,49 @@ namespace HordeServer.Jobs.TestData
 			return phase;
 		}
 
+		/// <inheritdoc/>
+		public async Task<ITestAudit> AddTestAuditAsync(UserId userId, TestId testId, TestAuditOptions options, CancellationToken cancellationToken = default)
+		{
+			// update previous one(s) and flag them as archive
+			UpdateDefinitionBuilder<TestAuditDocument> updateBuilder = Builders<TestAuditDocument>.Update;
+			List<UpdateDefinition<TestAuditDocument>> updates = new List<UpdateDefinition<TestAuditDocument>>();
+			updates.Add(updateBuilder.Set(x => x.Archived, true));
+
+			await _testAudits.UpdateManyAsync(x => x.TestId == testId && x.Archived == false, updateBuilder.Combine(updates), null, cancellationToken);
+
+			// add new entry
+			TestAuditDocument audit = new TestAuditDocument(userId, testId);
+			audit.UpdateOptions(options);
+
+			await _testAudits.InsertOneAsync(audit, null, cancellationToken);
+
+			return audit;
+		}
+
+		/// <inheritdoc/>
+		public async Task<IReadOnlyList<ITestAudit>> FindTestAuditsAsync(TestId[] testIds, DateTime? minCreateTime = null, CancellationToken cancellationToken = default)
+		{
+			FilterDefinition<TestAuditDocument> filter = FilterDefinition<TestAuditDocument>.Empty;
+			FilterDefinitionBuilder<TestAuditDocument> filterBuilder = Builders<TestAuditDocument>.Filter;
+			SortDefinition<TestAuditDocument> sortDefinition = Builders<TestAuditDocument>.Sort.Descending(x => x.Id);
+
+			filter &= filterBuilder.In(x => x.TestId, testIds);
+
+			if (minCreateTime != null)
+			{
+				TestAuditId minTime = TestAuditId.GenerateNewId(minCreateTime.Value);
+				filter &= filterBuilder.Eq(x => x.Archived, false) | filterBuilder.Gte(x => x.Id, minTime);
+			}
+			else
+			{
+				filter &= filterBuilder.Eq(x => x.Archived, false);
+			}
+
+			List<TestAuditDocument> results = await _testAudits.Find(filter).Sort(sortDefinition).ToListAsync(cancellationToken);
+
+			return results.ConvertAll<ITestAudit>(x => x);
+		}
+
 		/// <inheritdoc/>
 		public async Task<bool> UpdateAsync(int retainMonths, CancellationToken cancellationToken)
 		{
@@ -1398,6 +1543,20 @@ namespace HordeServer.Jobs.TestData
 					_logger.LogInformation("Expired {NumDeleted} test phase session documents", result.DeletedCount);
 				}
 			}
+
+			// expire test audits
+			{
+				TestAuditId auditExpireTime = TestAuditId.GenerateNewId(expireTime);
+				FilterDefinition<TestAuditDocument> filter = FilterDefinition<TestAuditDocument>.Empty;
+				FilterDefinitionBuilder<TestAuditDocument> filterBuilder = Builders<TestAuditDocument>.Filter;
+				filter &= filterBuilder.Eq(x => x.Archived, true);
+				filter &= filterBuilder.Lte(x => x.Id!, auditExpireTime);
+				DeleteResult result = await _testAudits.DeleteManyAsync(filter, cancellationToken);
+				if (result.DeletedCount > 0)
+				{
+					_logger.LogInformation("Expired {NumDeleted} test audit documents", result.DeletedCount);
+				}
+			}
 		}
 
 		/// <inheritdoc/>
diff --git a/Engine/Source/Programs/Horde/Plugins/Build/HordeServer.Build/Jobs/TestData/TestDataController.cs b/Engine/Source/Programs/Horde/Plugins/Build/HordeServer.Build/Jobs/TestData/TestDataController.cs
--- a/Engine/Source/Programs/Horde/Plugins/Build/HordeServer.Build/Jobs/TestData/TestDataController.cs
+++ b/Engine/Source/Programs/Horde/Plugins/Build/HordeServer.Build/Jobs/TestData/TestDataController.cs
@@ -3,9 +3,11 @@
 using System.Text.Json;
 using EpicGames.Core;
 using EpicGames.Horde.Commits;
+using EpicGames.Horde.Issues;
 using EpicGames.Horde.Jobs;
 using EpicGames.Horde.Jobs.TestData;
 using EpicGames.Horde.Streams;
+using EpicGames.Horde.Users;
 using HordeServer.Streams;
 using HordeServer.Utilities;
 using Microsoft.AspNetCore.Authorization;
@@ -906,5 +908,88 @@ namespace HordeServer.Jobs.TestData
 
 			return responses;
 		}
+
+		/// <summary>
+		/// Post a test audit
+		/// </summary>
+		/// <param name="request"></param>
+		/// <returns></returns>
+		[HttpPost]
+		[Route("/api/v3/testdata/audit/post")]
+		[ProducesResponseType(typeof(GetTestAuditResponse), 200)]
+		public async Task<ActionResult<GetTestAuditResponse>> AddTestAuditV2Async([FromBody] AddTestAuditRequest request)
+		{
+			GetTestAuditResponse responses = new GetTestAuditResponse();
+
+			UserId? userId = User.GetUserId();
+			if (userId == null)
+			{
+				return BadRequest("No user logged in.");
+			}
+
+			if (request.TestId == null || request.TestId.Length == 0)
+			{
+				return BadRequest("Invalid test id.");
+			}
+
+			TestAuditOptions options = new TestAuditOptions()
+			{
+				LastAuditDate = request.LastAuditDate?.UtcDateTime,
+				EnableSchedule = request.EnableSchedule ?? true,
+				Owner = request.Owner == null || request.Owner.Length == 0? null : UserId.Parse(request.Owner),
+				Customers = request.Customers?.Where(u => u.Length > 0).Select(u => UserId.Parse(u)).ToList(),
+				Notifications = request.Notifications?.Select(
+					n => new TestNotificationSettingsOptions()
+					{
+						Stream = new StreamId(n.Key),
+						Workflow = new WorkflowId(n.Value)
+					}).ToList(),
+				Team = request.Team,
+				Harness = request.Harness,
+				Intent = request.Intent,
+				Profile = request.Profile,
+				Notes = request.Notes
+			};
+
+			ITestAudit audit = await _testDataService.AddTestAuditAsync(userId!.Value, TestId.Parse(request.TestId), options);
+
+			return GetNewTestAutditResponse(audit);
+		}
+
+		[HttpPost]
+		[Route("/api/v3/testdata/audit/history")]
+		[ProducesResponseType(typeof(List<GetTestAuditResponse>), 200)]
+		public async Task<ActionResult<List<GetTestAuditResponse>>> GetTestAuditsV2Async(
+			[FromBody] GetTestAuditRequest request,
+			[FromQuery] DateTimeOffset? minCreateTime = null)
+		{
+			TestId[] testIds = request.TestIds.Select(x => TestId.Parse(x)).ToHashSet().ToArray();
+
+			IReadOnlyList<ITestAudit> audits = await _testDataService.FindTestAuditsV2Async(testIds, minCreateTime?.UtcDateTime);
+
+			return audits.Select(a => GetNewTestAutditResponse(a)).ToList();
+		}
+
+		private GetTestAuditResponse GetNewTestAutditResponse(ITestAudit audit)
+		{
+			return new GetTestAuditResponse()
+			{
+				TestId = audit.TestId.ToString(),
+				EnableSchedule = audit.EnableSchedule,
+				UserId = audit.UserId.ToString(),
+				UserInputDate = audit.UserInputDate,
+				LastAuditDate = audit.LastAuditDate,
+				Owner = audit.Owner?.ToString(),
+				Customers = audit.Customers == null ? null : audit.Customers.Select(c => c.ToString()).ToList(),
+				Notifications = audit.Notifications?.Where(
+						n => _buildConfig.Value.TryGetStream(n.Stream, out StreamConfig? streamConfig) && streamConfig.Authorize(JobAclAction.ViewJob, User)
+					).ToDictionary(n => n.Stream.ToString(), n => n.Workflow.ToString()),
+				Team = audit.Team?.ToString(),
+				Harness = audit.Harness?.ToString(),
+				Intent = audit.Intent?.ToString(),
+				Profile = audit.Profile?.ToString(),
+				Notes = audit.Notes?.ToString()
+			};
+		}
 	}
 }
diff --git a/Engine/Source/Programs/Horde/Plugins/Build/HordeServer.Build/Jobs/TestData/TestDataService.cs b/Engine/Source/Programs/Horde/Plugins/Build/HordeServer.Build/Jobs/TestData/TestDataService.cs
--- a/Engine/Source/Programs/Horde/Plugins/Build/HordeServer.Build/Jobs/TestData/TestDataService.cs
+++ b/Engine/Source/Programs/Horde/Plugins/Build/HordeServer.Build/Jobs/TestData/TestDataService.cs
@@ -4,6 +4,7 @@ using EpicGames.Core;
 using EpicGames.Horde.Commits;
 using EpicGames.Horde.Jobs;
 using EpicGames.Horde.Streams;
+using EpicGames.Horde.Users;
 using Microsoft.Extensions.Hosting;
 using Microsoft.Extensions.Logging;
 using Microsoft.Extensions.Options;
@@ -302,5 +303,30 @@ namespace HordeServer.Jobs.TestData
 		{
 			return await _testDataV2.FindTestPhaseSessionsAsync(streamIds, phaseIds, minCreateTime, maxCreateTime, minCommitId, maxCommitId, cancellationToken);
 		}
+
+		/// <summary>
+		/// Add a test audit entry
+		/// </summary>
+		/// <param name="userId"></param>
+		/// <param name="testId"></param>
+		/// <param name="options"></param>
+		/// <param name="cancellationToken"></param>
+		/// <returns></returns>
+		public async Task<ITestAudit> AddTestAuditAsync(UserId userId, TestId testId, TestAuditOptions options, CancellationToken cancellationToken = default)
+		{
+			return await _testDataV2.AddTestAuditAsync(userId, testId, options, cancellationToken);
+		}
+
+		/// <summary>
+		/// Find audit associated with the tests
+		/// </summary>
+		/// <param name="testIds"></param>
+		/// <param name="minCreateTime"></param>
+		/// <param name="cancellationToken"></param>
+		/// <returns></returns>
+		public async Task<IReadOnlyList<ITestAudit>> FindTestAuditsV2Async(TestId[] testIds, DateTime? minCreateTime = null, CancellationToken cancellationToken = default)
+		{
+			return await _testDataV2.FindTestAuditsAsync(testIds, minCreateTime, cancellationToken);
+		}
 	}
 }
\ No newline at end of file
diff --git a/Engine/Source/Programs/Shared/EpicGames.Horde/Jobs/TestData/TestDataMessages.cs b/Engine/Source/Programs/Shared/EpicGames.Horde/Jobs/TestData/TestDataMessages.cs
--- a/Engine/Source/Programs/Shared/EpicGames.Horde/Jobs/TestData/TestDataMessages.cs
+++ b/Engine/Source/Programs/Shared/EpicGames.Horde/Jobs/TestData/TestDataMessages.cs
@@ -405,6 +405,149 @@ namespace EpicGames.Horde.Jobs.TestData
 		public List<string> PhaseIds { get; set; } = new List<string>();
 	}
 
+	/// <summary>
+	/// Add an audit entry for a test
+	/// </summary>
+	public class AddTestAuditRequest
+	{
+		/// <summary>
+		/// The test name reference id
+		/// </summary>
+		public string TestId { get; set; } = String.Empty;
+
+		/// <summary>
+		/// Whether the test schedules should be enabled
+		/// </summary>
+		public bool? EnableSchedule { get; set; } = true;
+
+		/// <summary>
+		/// The date of the last audit
+		/// </summary>
+		public DateTimeOffset? LastAuditDate { get; set; }
+
+		/// <summary>
+		/// The owner id of the related test
+		/// </summary>
+		public string? Owner { get; set; }
+
+		/// <summary>
+		/// The customer ids of the related test
+		/// </summary>
+		public List<string>? Customers { get; set; }
+
+		/// <summary>
+		/// key/value pairs of stream id mapping to workflow id
+		/// </summary>
+		public Dictionary<string, string>? Notifications { get; set; }
+
+		/// <summary>
+		/// The team for which the test was intended for
+		/// </summary>
+		public string? Team { get; set; }
+
+		/// <summary>
+		/// The test harness
+		/// </summary>
+		public string? Harness { get; set; }
+
+		/// <summary>
+		/// The intent of the test
+		/// </summary>
+		public string? Intent { get; set; }
+
+		/// <summary>
+		/// The test profile
+		/// </summary>
+		public string? Profile { get; set; }
+
+		/// <summary>
+		/// Additional notes attached to the test audit
+		/// </summary>
+		public string? Notes { get; set; }
+	}
+
+	/// <summary>
+	/// Get test audit entries for test ids
+	/// </summary>
+	public class GetTestAuditRequest
+	{
+		/// <summary>
+		/// The test ids
+		/// </summary>
+		public List<string> TestIds { get; set; } = new List<string>();
+	}
+	
+	/// <summary>
+	/// A test audit entry
+	/// </summary>
+	public class GetTestAuditResponse
+	{
+		/// <summary>
+		/// The test id
+		/// </summary>
+		public string TestId { get; set; } = String.Empty;
+
+		/// <summary>
+		/// The user id that send the test audit
+		/// </summary>
+		public string UserId { get; set; } = String.Empty;
+
+		/// <summary>
+		/// The date at which this audit entry was done
+		/// </summary>
+		public DateTime UserInputDate { get; set; }
+
+		/// <summary>
+		/// Whether the test schedules should be enabled
+		/// </summary>
+		public bool? EnableSchedule { get; set; } = true;
+
+		/// <summary>
+		/// The date of the last audit
+		/// </summary>
+		public DateTime? LastAuditDate { get; set; }
+
+		/// <summary>
+		/// The owner id of the related test
+		/// </summary>
+		public string? Owner { get; set; }
+
+		/// <summary>
+		/// The customer ids of the related test
+		/// </summary>
+		public List<string>? Customers { get; set; }
+
+		/// <summary>
+		/// key/value pairs of stream id mapping to workflow id
+		/// </summary>
+		public Dictionary<string, string>? Notifications { get; set; }
+
+		/// <summary>
+		/// The team for which the test was intended for
+		/// </summary>
+		public string? Team { get; set; }
+
+		/// <summary>
+		/// The test harness
+		/// </summary>
+		public string? Harness { get; set; }
+
+		/// <summary>
+		/// The intent of the test
+		/// </summary>
+		public string? Intent { get; set; }
+
+		/// <summary>
+		/// The test profile
+		/// </summary>
+		public string? Profile { get; set; }
+
+		/// <summary>
+		/// Additional notes attached to the test audit
+		/// </summary>
+		public string? Notes { get; set; }
+	}
+
 	/// <summary>
 	/// A test suite that runs in a stream, contain subtests
 	/// </summary>
