Fixing warnings (missing  comments)

#rnx

diff --git a/Engine/Source/Programs/Horde/Plugins/Build/HordeServer.Build/Jobs/TestData/ITestDataCollection.cs b/Engine/Source/Programs/Horde/Plugins/Build/HordeServer.Build/Jobs/TestData/ITestDataCollection.cs
--- a/Engine/Source/Programs/Horde/Plugins/Build/HordeServer.Build/Jobs/TestData/ITestDataCollection.cs
+++ b/Engine/Source/Programs/Horde/Plugins/Build/HordeServer.Build/Jobs/TestData/ITestDataCollection.cs
@@ -286,24 +286,75 @@ namespace HordeServer.Jobs.TestData
 		Task<bool> UpdateAsync(int retainMonths, CancellationToken cancellationToken = default);
 	}
 
+	/// <summary>
+	/// Test notification setting options
+	/// </summary>
 	public class TestNotificationSettingsOptions
 	{
+		/// <summary>
+		/// The stream id
+		/// </summary>
 		public StreamId Stream { get; set; }
 
+		/// <summary>
+		/// The workflow id associated with the stream
+		/// </summary>
 		public WorkflowId Workflow { get; set; }
 	}
 
+	/// <summary>
+	/// The test audit options
+	/// </summary>
 	public class TestAuditOptions
 	{
+		/// <summary>
+		/// Whether or not the test schedule is enabled
+		/// </summary>
 		public bool EnableSchedule { get; set; } = true;
+
+		/// <summary>
+		/// The last audit Date
+		/// </summary>
 		public DateTime? LastAuditDate { get; set; }
+
+		/// <summary>
+		/// The owner of the related test
+		/// </summary>
 		public UserId? Owner { get; set; }
+
+		/// <summary>
+		/// The customers of the related test
+		/// </summary>
 		public List<UserId>? Customers { get; set; }
+
+		/// <summary>
+		/// The notifications settings options
+		/// </summary>
 		public List<TestNotificationSettingsOptions>? Notifications { get; set; }
+
+		/// <summary>
+		/// The team associated with the test
+		/// </summary>
 		public string? Team { get; set; }
+
+		/// <summary>
+		/// The harness associated with the test
+		/// </summary>
 		public string? Harness { get; set; }
+
+		/// <summary>
+		/// The intent of the test
+		/// </summary>
 		public string? Intent { get; set; }
+
+		/// <summary>
+		/// The profile associated with the test
+		/// </summary>
 		public string? Profile { get; set; }
+
+		/// <summary>
+		/// The notes associated with the test audit
+		/// </summary>
 		public string? Notes { get; set; }
 	}
 }
diff --git a/Engine/Source/Programs/Horde/Plugins/Build/HordeServer.Build/Jobs/TestData/TestDataController.cs b/Engine/Source/Programs/Horde/Plugins/Build/HordeServer.Build/Jobs/TestData/TestDataController.cs
--- a/Engine/Source/Programs/Horde/Plugins/Build/HordeServer.Build/Jobs/TestData/TestDataController.cs
+++ b/Engine/Source/Programs/Horde/Plugins/Build/HordeServer.Build/Jobs/TestData/TestDataController.cs
@@ -956,6 +956,12 @@ namespace HordeServer.Jobs.TestData
 			return GetNewTestAutditResponse(audit);
 		}
 
+		/// <summary>
+		/// Get test audit history
+		/// </summary>
+		/// <param name="request"></param>
+		/// <param name="minCreateTime"></param>
+		/// <returns></returns>
 		[HttpPost]
 		[Route("/api/v3/testdata/audit/history")]
 		[ProducesResponseType(typeof(List<GetTestAuditResponse>), 200)]
