[Horde][BuildHealth][Server] Merge "Root Cause" Slack workflow into the "Mark Fixed" button

diff --git a/Engine/Source/Programs/Horde/Plugins/Build/HordeServer.Build/Notifications/Sinks/SlackNotificationSink.cs b/Engine/Source/Programs/Horde/Plugins/Build/HordeServer.Build/Notifications/Sinks/SlackNotificationSink.cs
--- a/Engine/Source/Programs/Horde/Plugins/Build/HordeServer.Build/Notifications/Sinks/SlackNotificationSink.cs
+++ b/Engine/Source/Programs/Horde/Plugins/Build/HordeServer.Build/Notifications/Sinks/SlackNotificationSink.cs
@@ -1170,7 +1170,6 @@ namespace HordeServer.Notifications.Sinks
 						actions.AddButton("Assign to Me", value: $"issue_{issue.Id}_ack", style: ButtonStyle.Primary);
 						actions.AddButton("Not Me", value: $"issue_{issue.Id}_decline", style: ButtonStyle.Danger);
 						actions.AddButton("Mark Fixed", value: $"issue_{issue.Id}_markfixed");
-						actions.AddButton("Root Cause", value: $"issue_{issue.Id}_rootcause");
 
 						string? context = null;
 						if (issue.OwnerId != null)
@@ -3231,24 +3230,15 @@ namespace HordeServer.Notifications.Sinks
 							}
 
 							string? fixChangeStr;
+							CommitId? fixCommitId = null;
 							if (payload.View.State.TryGetValue("fix_cl", "fix_cl_action", out fixChangeStr))
 							{
-								CommitId commitId = new CommitId(fixChangeStr);
-								await _issueService.UpdateIssueAsync(issueId, fixCommitId: commitId, resolvedById: resolvedById, initiatedById: userId, cancellationToken: cancellationToken);
-								_logger.LogInformation("Marked issue {IssueId} fixed by user {UserId} in {Change}", issueId, resolvedById, commitId);
+								fixCommitId = new CommitId(fixChangeStr);
+								_logger.LogInformation("Marking issue {IssueId} fixed by user {UserId} in {Change}", issueId, resolvedById, fixCommitId);
 							}
-						}
-						else if (TryMatch(payload.View.CallbackId, @"^issue_(\d+)_rootcause_([a-fA-F0-9]{24})$", out match))
-						{
-							int issueId = Int32.Parse(match.Groups[1].Value, NumberStyles.Integer, CultureInfo.InvariantCulture);
-							UserId userId = UserId.Parse(match.Groups[2].Value);
 
 							string? rootChange;
-							string? rootCategory;
-							string? rootSummary;
-							int? dupeIssueId = null;
 							CommitId rootCommit = new CommitId("");
-
 							if (payload.View.State.TryGetValue("rootcause_cl", "rootcause_cl_action", out rootChange))
 							{
 								if (rootChange != null && rootChange.Length > 0)
@@ -3265,6 +3255,7 @@ namespace HordeServer.Notifications.Sinks
 								}
 							}
 
+							string? rootCategory;
 							if (payload.View.State.TryGetValue("rootcause_category", "rootcause_category_action", out rootCategory))
 							{
 								if (rootCategory != null && rootCategory.Length > 0)
@@ -3281,6 +3272,7 @@ namespace HordeServer.Notifications.Sinks
 								rootCategory = "";
 							}
 
+							string? rootSummary;
 							if (payload.View.State.TryGetValue("rootcause_summary", "rootcause_summary_action", out rootSummary))
 							{
 								if (rootSummary != null && rootSummary.Length > 0)
@@ -3297,6 +3289,7 @@ namespace HordeServer.Notifications.Sinks
 								rootSummary = "";
 							}
 
+							int? dupeIssueId = null;
 							if (payload.View.State.TryGetValue("rootcause_dupeid", "rootcause_dupeid_action", out string? dupeIdStr))
 							{
 								if (dupeIdStr != null && dupeIdStr.Length > 0)
@@ -3338,7 +3331,7 @@ namespace HordeServer.Notifications.Sinks
 								summary = curSummary;
 							}
 
-							await _issueService.UpdateIssueAsync(issueId, summary: summary, duplicateIssueId: dupeIssueId, rootCommitId: rootCommit, rootCauseCategory: rootCategory, rootCauseSummary: rootSummary, cancellationToken: cancellationToken);
+							await _issueService.UpdateIssueAsync(issueId, summary: summary, fixCommitId: fixCommitId, resolvedById: fixCommitId != null ? resolvedById : null, initiatedById: userId, duplicateIssueId: dupeIssueId, rootCommitId: rootCommit, rootCauseCategory: rootCategory, rootCauseSummary: rootSummary, cancellationToken: cancellationToken);
 
 						}
 						else if (TryMatch(payload.View.CallbackId, @"^issue_(\d+)_ack_([a-fA-F0-9]{24})$", out match))
@@ -3477,39 +3470,34 @@ namespace HordeServer.Notifications.Sinks
 				view.CallbackId = $"issue_{issueId}_markfixed_{user.Id}";
 
 				IIssue? issue = await _issueService.Collection.GetIssueAsync(issueId, cancellationToken);
-				if (issue != null && issue.OwnerId != null && issue.OwnerId != user.Id)
+
+				if (issue == null)
+				{
+					return;
+				}
+
+				if (issue.OwnerId != null && issue.OwnerId != user.Id)
 				{
 					IUser? owner = await _userCollection.GetCachedUserAsync(issue.OwnerId.Value, cancellationToken);
 					if (owner != null)
 					{
-						List<SlackOption> options = new List<SlackOption>();
-						options.Add(new SlackOption("Me", user.Id.ToString()));
-						options.Add(new SlackOption($"Current owner ({owner.Name})", owner.Id.ToString()));
+						List<SlackOption> ownerOptions = new List<SlackOption>();
+						ownerOptions.Add(new SlackOption("Me", user.Id.ToString()));
+						ownerOptions.Add(new SlackOption($"Current owner ({owner.Name})", owner.Id.ToString()));
 
-						RadioButtonGroupElement ownership = new RadioButtonGroupElement("fixed_by_action", options);
-						ownership.InitialOption = options[0];
+						RadioButtonGroupElement ownership = new RadioButtonGroupElement("fixed_by_action", ownerOptions);
+						ownership.InitialOption = ownerOptions[0];
 						view.AddInput("Fixed By:", ownership).BlockId = "fixed_by";
 					}
 				}
 
-				view.AddInput("Fix Changelist:", new PlainTextInputElement("fix_cl_action", placeholder: "Number")).BlockId = "fix_cl";
-
-				view.Close = "Cancel";
-				view.Submit = "Mark Fixed";
-
-				await _slackClient.OpenViewAsync(triggerId, view, cancellationToken);
-				return;
-			}
-			else if (String.Equals(verb, "rootcause", StringComparison.Ordinal))
-			{
-				IIssue? issue = await _issueService.Collection.GetIssueAsync(issueId, cancellationToken);
+				InputBlock fixBlock = view.AddInput("Fix Changelist:", new PlainTextInputElement("fix_cl_action", placeholder: "Fix Changelist"));
+				fixBlock.BlockId = "fix_cl";
+				fixBlock.Hint = "Change which fixed this issue. Enter 0 to mark fixed, without a change";
 
-				if (issue == null)
-				{
-					return;
-				}
+				view.AddDivider();
 
-				List<SlackOption> options = new List<SlackOption>
+				List<SlackOption> categoryOptions = new List<SlackOption>
 				{
 				   new SlackOption( "Asset/Content", "Asset/Content"),
 				   new SlackOption( "Code", "Code"),
@@ -3522,16 +3510,13 @@ namespace HordeServer.Notifications.Sinks
 				   new SlackOption( "Unintended Side Effect", "Unintended Side Effect"),
 				   new SlackOption( "Build Agent", "Build Agent"),
 				   new SlackOption( "Change Backout", "Change Backout"),
-				   new SlackOption( "Other", "Other"),				   
-				};						
-
-				SlackView view = new SlackView($"Issue {issueId}");
-				view.CallbackId = $"issue_{issueId}_rootcause_{user.Id}";
+				   new SlackOption( "Other", "Other"),
+				};
 
-				StaticSelectMenuElement categorySelect = new StaticSelectMenuElement("rootcause_category_action", "Select an Issue Category", options);
+				StaticSelectMenuElement categorySelect = new StaticSelectMenuElement("rootcause_category_action", "Select an Issue Category", categoryOptions);
 				if (!String.IsNullOrEmpty(issue.RootCauseCategory))
 				{
-					SlackOption? existing = options.FirstOrDefault(i => i.Value == issue.RootCauseCategory);
+					SlackOption? existing = categoryOptions.FirstOrDefault(i => i.Value == issue.RootCauseCategory);
 					if (existing != null)
 					{
 						// Only can get, issue with Epic Slack
