Horde: Make disposing of JobStepLogger more resilient

- If a network error occurs during disposal it may prevent job step from being updated
- Also make Polly retry on connect timeouts (speculative fix)

diff --git a/Engine/Source/Programs/Horde/Drivers/JobDriver/Utility/JobStepLogger.cs b/Engine/Source/Programs/Horde/Drivers/JobDriver/Utility/JobStepLogger.cs
--- a/Engine/Source/Programs/Horde/Drivers/JobDriver/Utility/JobStepLogger.cs
+++ b/Engine/Source/Programs/Horde/Drivers/JobDriver/Utility/JobStepLogger.cs
@@ -57,14 +57,21 @@ namespace JobDriver.Utility
 		/// <inheritdoc/>
 		public async ValueTask DisposeAsync()
 		{
-			if (!_updateOutcomeTask.IsCompleted)
+			try
 			{
-				await _cancellationSource.CancelAsync();
-				await _updateOutcomeTask;
-				_cancellationSource.Dispose();
-			}
+				if (!_updateOutcomeTask.IsCompleted)
+				{
+					await _cancellationSource.CancelAsync();
+					await _updateOutcomeTask;
+					_cancellationSource.Dispose();
+				}
 
-			await _serverLogger.DisposeAsync();
+				await _serverLogger.DisposeAsync();
+			}
+			catch (Exception e)
+			{
+				_internalLogger.LogError(e, "Error while disposing job step logger");
+			}
 		}
 
 		/// <inheritdoc/>
diff --git a/Engine/Source/Programs/Shared/EpicGames.Horde/HordeHttpMessageHandler.cs b/Engine/Source/Programs/Shared/EpicGames.Horde/HordeHttpMessageHandler.cs
--- a/Engine/Source/Programs/Shared/EpicGames.Horde/HordeHttpMessageHandler.cs
+++ b/Engine/Source/Programs/Shared/EpicGames.Horde/HordeHttpMessageHandler.cs
@@ -92,8 +92,11 @@ namespace EpicGames.Horde
 		{
 			return HttpPolicyExtensions
 				.HandleTransientHttpError()
+				.Or<TaskCanceledException>(ex => 
+					ex.InnerException is TimeoutException && 
+					ex.InnerException.Message.Contains("ConnectTimeout", StringComparison.Ordinal))
 				.OrResult(response => response.StatusCode == HttpStatusCode.ServiceUnavailable)
-				.WaitAndRetryAsync(3, retryAttempt => TimeSpan.FromSeconds(Math.Pow(2, retryAttempt)));
+				.WaitAndRetryAsync(6, retryAttempt => TimeSpan.FromSeconds(Math.Pow(2, retryAttempt)));
 		}
 	}
 }
