Horde: Update release notes

#rnx

diff --git a/Engine/Source/Programs/Horde/Docs/ReleaseNotes.md b/Engine/Source/Programs/Horde/Docs/ReleaseNotes.md
--- a/Engine/Source/Programs/Horde/Docs/ReleaseNotes.md
+++ b/Engine/Source/Programs/Horde/Docs/ReleaseNotes.md
@@ -1,5 +1,36 @@
 # Release Notes
 
+## 2025-10-23
+
+* Make disposing of JobStepLogger more resilient
+  - If a network error occurs during disposal it may prevent job step from being updated
+  - Also make Polly retry on connect timeouts (speculative fix) (47256245)
+* Fix free space calculation. This change allows the use of drives mounted as paths, not just drive letters. (47250181)
+* Merge "Root Cause" Slack workflow into the "Mark Fixed" button (47233499)
+* Batch API for JobStepRef (47213312)
+* Added AllowRetry flag to build graph nodes. (47205038)
+* Attach exception to trace span when agent relay allocations fail (47203856)
+* Use ManagedWorkspace for AutoSDK when PerforceMaterializer is active
+  AutoSDK heavily rely on Horde-based view filtering which is not implemented in PerforceMaterializer.
+  This avoids massive syncs during experimentation. (47203029)
+* Fix issue with multiple claims getting overwritten for built-in OAuth2
+  - Only last claim was ever written
+  - Applies when auth mode is set to Horde accounts (47200703)
+* [Build Health Dashboard] Bring audit log dropdown inline with Horde standard, correctly set defaults, and update options based on backend fetches when live tailing (47179708)
+* Adding isPreflight column to JobSummary data. (47175049)
+* D3 container fixes for label/no-label toggle (47174819)
+* Enable parallel test running for all Horde server related projects (47158966)
+* Enable parallel test execution for HordeServer.Build.Tests and HordeServer.Compute.Tests (47154960)
+* Add support for parallel test execution for database integration tests (MongoDB and Redis) (47150628)
+* Check if a conform is required before asssigning a Job Task lease to an agent (47147172)
+* State filter & minor fix for batch sync failed. (47111953)
+* EpicGames.Perforce: Add support for modtime client option
+  - Maintains the original Perforce timestamp when moved to/from the cache
+  - Requires opting in via "method" setting for the workspace type, like how use of have table works (47099206)
+* Adding step messages to dashboard step view, step details, log view, and history (47051698)
+* Bake report step message/severtity into ref metadata (47041541)
+* Adding job/step messages/severity to job reports (47015259)
+
 ## 2025-10-16
 
 * Refactor issue config to be a root object (46971377)
