[Horde][BuildHealth][Dashboard] Minor change to improve the 'no data' case styling

diff --git a/Engine/Source/Programs/Horde/HordeDashboard/src/components/buildhealth/stepoutcome/StepOutcomeD3Component.tsx b/Engine/Source/Programs/Horde/HordeDashboard/src/components/buildhealth/stepoutcome/StepOutcomeD3Component.tsx
--- a/Engine/Source/Programs/Horde/HordeDashboard/src/components/buildhealth/stepoutcome/StepOutcomeD3Component.tsx
+++ b/Engine/Source/Programs/Horde/HordeDashboard/src/components/buildhealth/stepoutcome/StepOutcomeD3Component.tsx
@@ -149,10 +149,17 @@ function GenerateD3Table(stepOutcomeTable: StepOutcomeTable, handler: StepOutcom
     const { modeColors, _hordeClasses, _detailClasses } = getHordeStyling();
     const theme = getHordeTheme();
     const colorRecords: Record<string, string> = getColorRecords();
-    const summaryColorScale = d3.scaleLinear<string>()
+    const summaryColorScaleBase = d3.scaleLinear<string>()
         .domain([0, 0.5, 1])
         .range([colorRecords["Failure"], colorRecords["Warnings"], colorRecords["Success"]]);
 
+    const summaryColorScale = (value: number): string => {
+        if (value === -1) {
+            return colorRecords["Waiting"];
+        }
+        return summaryColorScaleBase(value);
+    };
+
     const scaffold = useRef<{
         wrapper: d3.Selection<HTMLDivElement, unknown, null, undefined>;
         table: d3.Selection<HTMLTableElement, unknown, null, undefined>;
