Horde: Set env vars for UAT/BuildGraph in PerforceMaterializer

diff --git a/Engine/Source/Programs/Horde/Drivers/JobDriver/Execution/PerforceMaterializer.cs b/Engine/Source/Programs/Horde/Drivers/JobDriver/Execution/PerforceMaterializer.cs
--- a/Engine/Source/Programs/Horde/Drivers/JobDriver/Execution/PerforceMaterializer.cs
+++ b/Engine/Source/Programs/Horde/Drivers/JobDriver/Execution/PerforceMaterializer.cs
@@ -63,7 +63,7 @@ public sealed class PerforceMaterializer : IWorkspaceMaterializer
 	public string Identifier { get; }
 	
 	/// <inheritdoc/>
-	public IReadOnlyDictionary<string, string> EnvironmentVariables { get; }
+	public IReadOnlyDictionary<string, string> EnvironmentVariables { get; private set; }
 	
 	/// <inheritdoc/>
 	public bool IsPerforceWorkspace { get; } = true;
@@ -191,6 +191,17 @@ public sealed class PerforceMaterializer : IWorkspaceMaterializer
 		_logger.LogInformation("Using client {ClientName} (Host: {HostName}, Stream: {StreamName}, Type: {Type}, Root: {Path})",
 			client.Name, client.Host, client.Stream, client.Type, client.Root);
 
+		// Variables expected to be set for UAT/BuildGraph when Perforce is enabled (-P4 flag is set) 
+		EnvironmentVariables = new Dictionary<string, string>()
+		{
+			["uebp_PORT"] = _options.AgentWorkspace.ServerAndPort,
+			["uebp_USER"] = _options.AgentWorkspace.UserName,
+			["uebp_CLIENT"] = client.Name,
+			["uebp_CLIENT_ROOT"] = $"//{client.Name}",
+			["P4CLIENT"] = client.Name,
+			["P4USER"] = _options.AgentWorkspace.UserName,
+		};
+		
 		_perforceWithClient = await PerforceConnection.CreateAsync(_perforceSettings, _logger);
 		await RevertInternalAsync(_perforceWithClient, cancellationToken);
 
