Horde: Improve file count and download speed log lines in PerforceMaterializer

diff --git a/Engine/Source/Programs/Horde/Drivers/JobDriver/Execution/PerforceMaterializer.cs b/Engine/Source/Programs/Horde/Drivers/JobDriver/Execution/PerforceMaterializer.cs
--- a/Engine/Source/Programs/Horde/Drivers/JobDriver/Execution/PerforceMaterializer.cs
+++ b/Engine/Source/Programs/Horde/Drivers/JobDriver/Execution/PerforceMaterializer.cs
@@ -238,8 +238,11 @@ public sealed class PerforceMaterializer : IWorkspaceMaterializer
 
 			long totalFileCount = syncSummaries.Sum(x => x.TotalFileCount);
 			long totalFileSize = syncSummaries.Sum(x => x.TotalFileSize);
-			double totalFileSizeMb = totalFileSize / (1024.0 * 1024.0);
-			_logger.LogInformation("Synced {FileCount} files. {FileSizeMb:F2} MB ({FileSize} bytes)", totalFileCount, totalFileSizeMb, totalFileSize);
+			double totalFileSizeMb = Math.Round(totalFileSize / (1024.0 * 1024.0), 1);
+			double filesPerSecond = Math.Round(totalFileCount / syncTimer.Elapsed.TotalSeconds, 1);
+			double dataPerSecond = Math.Round(totalFileSizeMb / syncTimer.Elapsed.TotalSeconds, 1);
+			_logger.LogInformation("Synced {FileCount} files. {FileSizeMb:F2} MB ({FileSize} bytes). {FilesPerSecond} files/s. {DataPerSecond} MB/s",
+				totalFileCount, totalFileSizeMb, totalFileSize, filesPerSecond, dataPerSecond);
 			syncSpan.SetAttribute("horde.pm.file_count", totalFileCount);
 			syncSpan.SetAttribute("horde.pm.file_size", totalFileSize);
 			status.Progress = $"({syncTimer.Elapsed.TotalSeconds:0.0}s)";
@@ -336,8 +339,7 @@ public sealed class PerforceMaterializer : IWorkspaceMaterializer
 				existingClient.Owner == newClient.Owner &&
 				existingClient.Root == newClient.Root &&
 				existingClient.Host == newClient.Host &&
-				existingClient.Type == newClient.Type &&
-				existingClient.Root == newClient.Root)
+				existingClient.Type == newClient.Type)
 			{
 				_logger.LogInformation("Found existing and matching client!");
 				newClient = existingClient;
