[Horde][BuildHealth][Dashboard] Adding labels to option filters

diff --git a/Engine/Source/Programs/Horde/HordeDashboard/src/components/buildhealth/BuildHealthDataHandler.ts b/Engine/Source/Programs/Horde/HordeDashboard/src/components/buildhealth/BuildHealthDataHandler.ts
--- a/Engine/Source/Programs/Horde/HordeDashboard/src/components/buildhealth/BuildHealthDataHandler.ts
+++ b/Engine/Source/Programs/Horde/HordeDashboard/src/components/buildhealth/BuildHealthDataHandler.ts
@@ -4,8 +4,8 @@ import backend from "horde/backend";
 import { JobData, JobStreamQuery, ProjectData, StreamData } from "horde/backend/Api";
 import { makeObservable, observable, runInAction } from "mobx";
 import { BuildHealthOptionsState } from "./BuildHealthOptions";
-import { StepRefData, TemplateRefData } from "./BuildHealthDataTypes";
-import { decodeTemplateKey, encodeTemplateKeyFromStrings } from "./BuildHealthUtilities";
+import { LabelRefData, StepRefData, TemplateRefData } from "./BuildHealthDataTypes";
+import { decodeTemplateKey, encodeFullyQualifiedLabelName, encodeTemplateKeyFromStrings } from "./BuildHealthUtilities";
 import { BuildHealthOptionsController, IHierarchyValidator } from "./BuildHealthOptionsController";
 
 // #region -- API Data Types & Utilities --
@@ -55,12 +55,20 @@ interface StepCacheEntry {
     steps: StepRefData[];
 }
 
+/**
+ * Represents a cache entry for label ref data.
+ */
+interface LabelCacheEntry {
+    timestamp: number;
+    labels: LabelRefData[];
+}
+
 /**
  * Data handler for all build health data used for dropdowns.
  * @todo Data sharing & caching would ideally occur between the BuildHealth View & Step Outcome.
  */
 export class BuildHealthDataHandler implements IHierarchyValidator {
-    private readonly STREAMS_FILTER = "id,state,streamId,name";
+    private readonly STREAMS_FILTER = "id,state,streamId,name,labels";
     private readonly STREAM_JOBS_BATCH_STEP_FILTER = "batches.steps.id,batches.steps.name,batches.steps.state,batches.steps.outcome,batches.steps.finishTime"
     private readonly STEP_DATA_SAMPLE_SIZE = 5;
 
@@ -71,6 +79,7 @@ export class BuildHealthDataHandler implements IHierarchyValidator {
     private buildHealthOptionsController: BuildHealthOptionsController;
     private cachedStreamToTemplatesData: Map<string, TemplateRefData[]> = new Map<string, TemplateRefData[]>();
     private cachedStreamTempalteIdToSteps: Map<string, StepCacheEntry> = new Map<string, StepCacheEntry>();
+    private cachedStreamTemplateIdToLabels: Map<string, LabelCacheEntry> = new Map<string, LabelCacheEntry>();
 
     // #endregion -- Private Members --
 
@@ -88,6 +97,9 @@ export class BuildHealthDataHandler implements IHierarchyValidator {
     @observable.shallow
     stepData: StepRefData[] = [];
 
+    @observable.shallow
+    labelData: LabelRefData[] = [];
+
     // #endregion -- Public Members --
 
     // #region -- Constructor --
@@ -216,7 +228,7 @@ export class BuildHealthDataHandler implements IHierarchyValidator {
     async getStepData(request: DataHandlerRefreshRequest) {
         const entries = Object.entries(this.buildHealthOptions.enabledTemplates);
         if (request.steps && entries.length > 0) {
-            const enabledIds = new Set(entries.map(([id]) => id));
+            const enabledTemplateKey = new Set(entries.map(([enabledTemplateKey]) => enabledTemplateKey));
 
             await Promise.all(entries.filter(([key, _value]) => !this.cachedStreamTempalteIdToSteps.has(key)).map(async ([key]) => {
                 const templateData = decodeTemplateKey(key);
@@ -230,6 +242,7 @@ export class BuildHealthDataHandler implements IHierarchyValidator {
                 const jobDatas: JobData[] = await backend.getStreamJobs(templateData.streamId, query);
                 const recentJobs = jobDatas.filter(job => job.batches);
                 const stepMap = new Map<string, StepRefData>();
+                const labelMap = new Map<string, LabelRefData>();
 
                 for (const job of recentJobs) {
                     for (const batch of job.batches!) {
@@ -240,24 +253,37 @@ export class BuildHealthDataHandler implements IHierarchyValidator {
                             }
                         }
                     }
+                    for (const label of job.labels!) {
+                        const labelKey = encodeFullyQualifiedLabelName(label.dashboardCategory, label.dashboardName);
+                        if (!labelMap.has(labelKey)) {
+                            labelMap.set(labelKey, { ...label, streamId: templateData.streamId, templateId: decodeTemplateKey(key).templateId })
+                        }
+                    }
                 }
 
                 const distinctSteps = [...stepMap.values()];
+                const distintcLabels = [...labelMap.values()];
 
                 this.cachedStreamTempalteIdToSteps.set(key, { timestamp: Date.now(), steps: distinctSteps });
+                this.cachedStreamTemplateIdToLabels.set(key, { timestamp: Date.now(), labels: distintcLabels });
 
                 return distinctSteps;
             }));
 
             runInAction(() => {
                 this.stepData = Array.from(this.cachedStreamTempalteIdToSteps.entries())
-                    .filter(([templateId]) => enabledIds.has(templateId))
+                    .filter(([templateId]) => enabledTemplateKey.has(templateId))
                     .flatMap(([_, entry]) => entry.steps);
+
+                this.labelData = Array.from(this.cachedStreamTemplateIdToLabels.entries())
+                    .filter(([templateId]) => enabledTemplateKey.has(templateId))
+                    .flatMap(([_, entry]) => entry.labels);
             });
         } else if (request.steps) { // We have no enabled templates, so we should have no step data showing.
 
             runInAction(() => {
                 this.stepData = [];
+                this.labelData = [];
             });
         }
     }
@@ -280,6 +306,24 @@ export class BuildHealthDataHandler implements IHierarchyValidator {
         return returnArray;
     }
 
+    /**
+     * Obtains all the currently known labels for a given stream & template.
+     * @param streamId The stream.
+     * @param templateId The template.
+     * @returns An array representing all the known labels for a given stream & template pair.
+     */
+    getCachedTemplateLabelNames(streamId: string, templateId: string): string[] {
+        let returnArray: string[] = [];
+
+        let encodedTemplateKey: string = encodeTemplateKeyFromStrings(streamId, templateId);
+
+        if (this.cachedStreamTemplateIdToLabels.has(encodedTemplateKey)) {
+            returnArray = this.cachedStreamTemplateIdToLabels.get(encodedTemplateKey)!.labels.map(x => encodeFullyQualifiedLabelName(x.dashboardCategory, x.dashboardName));
+        }
+
+        return returnArray;
+    }
+
     /**
      * Obtains all the StepRefData for a given stream & template.
      * @param streamId The stream.
@@ -318,8 +362,31 @@ export class BuildHealthDataHandler implements IHierarchyValidator {
         if (template === undefined) {
             return false;
         }
+
+        let encodedTemplateKey: string = encodeTemplateKeyFromStrings(streamId, templateId);
+        if (!this.cachedStreamTempalteIdToSteps.get(encodedTemplateKey)?.steps.find(x => x.name.toLocaleLowerCase() == stepName.toLocaleLowerCase())) {
+            return false;
+        }
+
+        return true;
+    }
+
+    /**
+    * @inheritdoc
+    */
+    isValidLabel(streamId: string, templateId: string, fullyQualifiedLabelName: string): boolean {
+        if (!this.cachedStreamToTemplatesData.has(streamId)) {
+            return false;
+        }
+
+        let template = this.cachedStreamToTemplatesData.get(streamId)?.find(x => x.id === templateId)
+
+        if (template === undefined) {
+            return false;
+        }
+
         let encodedTemplateKey: string = encodeTemplateKeyFromStrings(streamId, templateId);
-        if (!this.cachedStreamTempalteIdToSteps.get(encodedTemplateKey)?.steps.find(x => x.name == stepName)) {
+        if (!this.cachedStreamTemplateIdToLabels.get(encodedTemplateKey)?.labels.find(x => encodeFullyQualifiedLabelName(x.dashboardCategory, x.dashboardName).toLocaleLowerCase() === fullyQualifiedLabelName.toLocaleLowerCase())) {
             return false;
         }
 
diff --git a/Engine/Source/Programs/Horde/HordeDashboard/src/components/buildhealth/BuildHealthDataTypes.ts b/Engine/Source/Programs/Horde/HordeDashboard/src/components/buildhealth/BuildHealthDataTypes.ts
--- a/Engine/Source/Programs/Horde/HordeDashboard/src/components/buildhealth/BuildHealthDataTypes.ts
+++ b/Engine/Source/Programs/Horde/HordeDashboard/src/components/buildhealth/BuildHealthDataTypes.ts
@@ -1,11 +1,12 @@
 // Copyright Epic Games, Inc. All Rights Reserved.
 
-import { GetStepResponse, GetTemplateRefResponse } from "horde/backend/Api";
+import { GetLabelStateResponse, GetStepResponse, GetTemplateRefResponse } from "horde/backend/Api";
 
 export const QUERY_SCHEMA_PARAM_NAME: string = "query_schema";
 export const PROJECTS_URL_PARAM_NAME: string = "projects";
 export const STREAMS_URL_PARAM_NAME: string = "streams";
 export const TEMPLATES_URL_PARAM_NAME: string = "templates";
+export const LABELS_URL_PARAM_NAME: string = "labels";
 export const STEPS_URL_PARAM_NAME: string = "stepNames";
 export const STATES_URL_PARAM_NAME: string = "states";
 export const PARAMETER_KEY_PREFIX: string = "parameter_key";
@@ -13,6 +14,7 @@ export const TIME_SPAN_URL_PARAM: string = "lastTimeRange";
 export const INCLUDE_PREFLIGHT_PARAM: string = "includePreflight";
 export const KEY_SEPARATOR: string = "::";
 
+export const LABELS_URL_PARAM_ALL_LABELS: string = "labels_AllKnownValues";
 export const STEPS_URL_PARAM_ALL_STEPS: string = "stepNames_AllKnownValues";
 export const STATES_URL_PARAM_ALL_STATES: string = "stateNames_AllKnownValues";
 
@@ -30,4 +32,12 @@ export type TemplateRefData = GetTemplateRefResponse & {
 export type StepRefData = GetStepResponse & {
     templateId?: string;
     streamId?: string;
+}
+
+/**
+ * Convenience type to fold in extra metadata relevant for the Label. 
+ */
+export type LabelRefData = GetLabelStateResponse & {
+    templateId?: string;
+    streamId?: string;
 }
\ No newline at end of file
diff --git a/Engine/Source/Programs/Horde/HordeDashboard/src/components/buildhealth/BuildHealthOptions.ts b/Engine/Source/Programs/Horde/HordeDashboard/src/components/buildhealth/BuildHealthOptions.ts
--- a/Engine/Source/Programs/Horde/HordeDashboard/src/components/buildhealth/BuildHealthOptions.ts
+++ b/Engine/Source/Programs/Horde/HordeDashboard/src/components/buildhealth/BuildHealthOptions.ts
@@ -1,9 +1,9 @@
 // Copyright Epic Games, Inc. All Rights Reserved.
 
 import { makeObservable, observable, runInAction } from "mobx";
-import { INCLUDE_PREFLIGHT_PARAM, PROJECTS_URL_PARAM_NAME, QUERY_SCHEMA_PARAM_NAME, STATES_URL_PARAM_ALL_STATES, STATES_URL_PARAM_NAME, STEPS_URL_PARAM_ALL_STEPS, STEPS_URL_PARAM_NAME, STREAMS_URL_PARAM_NAME, TEMPLATES_URL_PARAM_NAME, TIME_SPAN_URL_PARAM } from "../buildhealth/BuildHealthDataTypes";
+import { INCLUDE_PREFLIGHT_PARAM, LABELS_URL_PARAM_NAME, PROJECTS_URL_PARAM_NAME, QUERY_SCHEMA_PARAM_NAME, STATES_URL_PARAM_ALL_STATES, STATES_URL_PARAM_NAME, LABELS_URL_PARAM_ALL_LABELS, STEPS_URL_PARAM_ALL_STEPS, STEPS_URL_PARAM_NAME, STREAMS_URL_PARAM_NAME, TEMPLATES_URL_PARAM_NAME, TIME_SPAN_URL_PARAM } from "../buildhealth/BuildHealthDataTypes";
 import { BuildHealthDataHandler } from "./BuildHealthDataHandler";
-import { decodeStateKey, decodeStepKey, decodeTemplateKey, encodeStateKeyFromString, encodeStepKey, encodeStepNameFromStringsPreserveCase, encodeTemplateKey, encodeTemplateKeyFromStrings } from "./BuildHealthUtilities";
+import { decodeFullyQualifiedLabelName, decodeLabelKey, decodeStateKey, decodeStepKey, decodeTemplateKey, encodeFullyQualifiedLabelName, encodeLabelKey, encodeLabelKeyFromStrings, encodeStateKeyFromString, encodeStepKey, encodeStepNameFromStringsPreserveCase, encodeTemplateKey, encodeTemplateKeyFromStrings } from "./BuildHealthUtilities";
 import { Location } from 'react-router-dom';
 import LZString from "lz-string"
 import { BuildHealthOptionsController, IBuildHealthOptionWriter } from "./BuildHealthOptionsController";
@@ -33,6 +33,7 @@ type BuildHealthQueryParams = {
     projectIds: string[];
     streamIds: string[];
     templateIds: string[];
+    labels: string[];
     stepNames: string[];
     stateNames: string[];
     timeSpanKey: string;
@@ -114,6 +115,7 @@ export function toNavigationQuery(optionsState: BuildHealthOptionsState, querySc
     let projects: string = "";
     let streams: string = "";
     let jobs: string = "";
+    let labels: string = "";
     let steps: string = "";
     let states: string = "";
 
@@ -150,13 +152,30 @@ export function toNavigationQuery(optionsState: BuildHealthOptionsState, querySc
     }
 
     // States
-    if (querySchemaVersion === 3) {
+    if (querySchemaVersion >= 3) {
         states = optionsState.stepOutcomeEnabledsStates.length === Object.keys(JobStepState).length ? STATES_URL_PARAM_ALL_STATES : optionsState.stepOutcomeEnabledsStates.join(",");
     }
 
+    // Labels
+    if (querySchemaVersion >= 4) {
+        labels = compressKeys(optionsState.stepOutcomeEnabledLabelKeys,
+            (key: string) => {
+                let { streamId, templateId, dashboardCategory, dashboardName } = decodeLabelKey(key);
+
+                return { prefix: encodeTemplateKeyFromStrings(streamId, templateId), suffix: encodeFullyQualifiedLabelName(dashboardCategory, dashboardName) };
+            },
+            (prefix: string, suffixes: string[]) => {
+                let { streamId, templateId } = decodeTemplateKey(prefix);
+                let summarized = buildHealthDataHandler.getCachedTemplateLabelNames(streamId, templateId)?.length == suffixes.length ? LABELS_URL_PARAM_ALL_LABELS : undefined;
+
+                return summarized;
+            });
+    }
+
     params.set(PROJECTS_URL_PARAM_NAME, LZString.compressToEncodedURIComponent(projects));
     params.set(STREAMS_URL_PARAM_NAME, LZString.compressToEncodedURIComponent(streams));
     params.set(TEMPLATES_URL_PARAM_NAME, LZString.compressToEncodedURIComponent(jobs));
+    params.set(LABELS_URL_PARAM_NAME, LZString.compressToEncodedURIComponent(labels));
     params.set(STEPS_URL_PARAM_NAME, LZString.compressToEncodedURIComponent(steps));
     params.set(STATES_URL_PARAM_NAME, LZString.compressToEncodedURIComponent(states));
     params.set(TIME_SPAN_URL_PARAM, LZString.compressToEncodedURIComponent(lastTimeRange));
@@ -202,6 +221,7 @@ export function parseBuildHealthQueryParams(location: Location): BuildHealthQuer
     let projectIds: string[] = [];
     let streamIds: string[] = [];
     let templateIds: string[] = [];
+    let labels: string[] = [];
     let stepNames: string[] = [];
     let stateNames: string[] = [];
 
@@ -250,10 +270,26 @@ export function parseBuildHealthQueryParams(location: Location): BuildHealthQuer
     }
 
     // States
-    if (querySchemaVersion === 3) {
+    if (querySchemaVersion >= 3) {
         stateNames = getList(STATES_URL_PARAM_NAME);
     }
 
+    // Labels
+    if (querySchemaVersion >= 4) {
+        const rawLabels = searchParams.get(LABELS_URL_PARAM_NAME);
+        if (!rawLabels) {
+            labels = [];
+        }
+        else {
+            let decompressedLabels = LZString.decompressFromEncodedURIComponent(rawLabels);
+            labels = expandKeys(decompressedLabels, (prefix, suffix) => {
+                let { streamId, templateId } = decodeTemplateKey(prefix);
+                let { labelCategory, label } = decodeFullyQualifiedLabelName(suffix);
+                return encodeLabelKeyFromStrings(streamId, templateId, labelCategory, label);
+            });
+        }
+    }
+
     let timeSpanKey: string = (() => {
         const raw = searchParams.get(TIME_SPAN_URL_PARAM);
         return raw ? LZString.decompressFromEncodedURIComponent(raw) ?? "" : "";
@@ -266,7 +302,7 @@ export function parseBuildHealthQueryParams(location: Location): BuildHealthQuer
     // Read from URLSearchParams
     const includePreflight = (searchParams.get(INCLUDE_PREFLIGHT_PARAM) ? LZString.decompressFromEncodedURIComponent(searchParams.get(INCLUDE_PREFLIGHT_PARAM)!) : "false") === "true";
 
-    return { projectIds, streamIds, templateIds: templateIds, stepNames, stateNames, timeSpanKey, includePreflight, querySchemaVersion };
+    return { projectIds, streamIds, templateIds: templateIds, labels, stepNames, stateNames, timeSpanKey, includePreflight, querySchemaVersion };
 }
 
 /**
@@ -407,7 +443,7 @@ export async function loadBuildHealthOptionsFromParams(
             const step = handler.stepData.find(s => s.name.toLowerCase() === stepData.stepName.toLowerCase() && s.streamId == stepData.streamId && s.templateId == stepData.templateId);
             if (step) {
                 optionsWriter.toggleStep(encodeStepKey(step), step.name, true);
-            } else if (stepData.stepName === STEPS_URL_PARAM_ALL_STEPS) {
+            } else if (stepData.stepName.toLocaleLowerCase() === STEPS_URL_PARAM_ALL_STEPS.toLocaleLowerCase()) {
                 let stepNames = handler.getCachedTemplateStepNames(stepData.streamId, stepData.templateId)
                 optionsWriter.toggleAllSteps(stepData.streamId, stepData.templateId, stepNames, true);
             }
@@ -417,6 +453,27 @@ export async function loadBuildHealthOptionsFromParams(
         });
     });
 
+    runInAction(() => {
+        optionsWriter.clearLabels();
+        const decodedLabels = params.labels.map(labelKey => decodeLabelKey(labelKey));
+
+        decodedLabels.forEach(labelData => {
+            const label = handler.labelData.find(l => l.streamId == labelData.streamId
+                && l.templateId == labelData.templateId
+                && l.dashboardCategory?.toLowerCase() === labelData.dashboardCategory.toLowerCase()
+                && l.dashboardName?.toLowerCase() === labelData.dashboardName.toLowerCase());
+            if (label) {
+                optionsWriter.toggleLabel(encodeLabelKey(label), encodeFullyQualifiedLabelName(labelData.dashboardCategory, labelData.dashboardName), true);
+            } else if (labelData.dashboardCategory?.toLocaleLowerCase() === LABELS_URL_PARAM_ALL_LABELS.toLocaleLowerCase()) {
+                let labelNames = handler.getCachedTemplateLabelNames(labelData.streamId, labelData.templateId)
+                optionsWriter.toggleAllLabels(labelData.streamId, labelData.templateId, labelNames, true);
+            }
+            else {
+                console.warn(`Label parameter provided in URL query (${encodeFullyQualifiedLabelName(labelData.dashboardCategory, labelData.dashboardName)}) that could not be matched to underlying data - discarding.`);
+            }
+        });
+    });
+
     options.commitTransactionSession();
 }
 
@@ -428,7 +485,6 @@ export async function loadBuildHealthOptionsFromParams(
 export class BuildHealthUIOptionsState {
     @observable includeDateAnchors: boolean = false;
     @observable warningsAsSummaryFailure: boolean = false;
-    @observable showLabels: boolean = false;
     @observable validStates: Set<string> = new Set<string>();
 
     constructor() {
@@ -461,6 +517,7 @@ export class BuildHealthOptionStateReceipt {
     enabledProjects: ReadonlySet<string>;
     enabledStreams: ReadonlySet<string>;
     enabledTemplates: ReadonlySet<string>;
+    enabledLabels: ReadonlySet<string>;
     enabledSteps: ReadonlySet<string>;
     enabledStates: ReadonlySet<string>;
     enabledPreflight: boolean;
@@ -470,6 +527,7 @@ export class BuildHealthOptionStateReceipt {
         projects: Iterable<string> = [],
         streams: Iterable<string> = [],
         templates: Iterable<string> = [],
+        labels: Iterable<string> = [],
         steps: Iterable<string> = [],
         states: Iterable<string> = [],
         enabledPreflight: boolean,
@@ -478,6 +536,7 @@ export class BuildHealthOptionStateReceipt {
         this.enabledProjects = new Set(projects);
         this.enabledStreams = new Set(streams);
         this.enabledTemplates = new Set(templates);
+        this.enabledLabels = new Set(labels);
         this.enabledSteps = new Set(steps);
         this.enabledStates = new Set(states);
         this.enabledPreflight = enabledPreflight;
@@ -499,6 +558,7 @@ export class BuildHealthOptionStateReceipt {
             BuildHealthOptionStateReceipt.difference(next.enabledProjects, prev.enabledProjects),
             BuildHealthOptionStateReceipt.difference(next.enabledStreams, prev.enabledStreams),
             BuildHealthOptionStateReceipt.difference(next.enabledTemplates, prev.enabledTemplates),
+            BuildHealthOptionStateReceipt.difference(next.enabledLabels, prev.enabledLabels),
             BuildHealthOptionStateReceipt.difference(next.enabledSteps, prev.enabledSteps),
             BuildHealthOptionStateReceipt.difference(next.enabledStates, prev.enabledStates),
             next.enabledPreflight,
@@ -509,6 +569,7 @@ export class BuildHealthOptionStateReceipt {
             BuildHealthOptionStateReceipt.difference(prev.enabledProjects, next.enabledProjects),
             BuildHealthOptionStateReceipt.difference(prev.enabledStreams, next.enabledStreams),
             BuildHealthOptionStateReceipt.difference(prev.enabledTemplates, next.enabledTemplates),
+            BuildHealthOptionStateReceipt.difference(prev.enabledLabels, next.enabledLabels),
             BuildHealthOptionStateReceipt.difference(prev.enabledSteps, next.enabledSteps),
             BuildHealthOptionStateReceipt.difference(prev.enabledStates, next.enabledStates),
             prev.enabledPreflight,
@@ -546,6 +607,7 @@ export class BuildHealthOptionsState {
 
     @observable stepOutcomeEnabledStreamKeys: string[] = [];
     @observable stepOutcomeEnabledJobKeys: string[] = [];
+    @observable stepOutcomeEnabledLabelKeys: string[] = [];
     @observable stepOutcomeEnabledStepKeys: string[] = [];
     @observable stepOutcomeEnabledsStates: string[] = [];
     @observable startDate!: Date;
@@ -558,6 +620,7 @@ export class BuildHealthOptionsState {
     @observable enabledProjects: Record<string, string> = {};
     @observable enabledStreams: Record<string, string> = {};
     @observable enabledTemplates: Record<string, string> = {};
+    @observable enabledLabels: Record<string, string> = {};
     @observable enabledSteps: Record<string, string> = {};
     @observable enabledStates: Record<string, string> = {};
     @observable jobHistoryTimeSpan: TimeSpan = DEFAULT_JOB_HISTORY_TIME_SPAN;
@@ -590,6 +653,7 @@ export class BuildHealthOptionsState {
             Object.keys(this.enabledProjects),
             Object.keys(this.enabledStreams),
             Object.keys(this.enabledTemplates),
+            Object.keys(this.enabledLabels),
             Object.keys(this.enabledSteps),
             Object.values(this.enabledStates),
             this.includePreflight,
diff --git a/Engine/Source/Programs/Horde/HordeDashboard/src/components/buildhealth/BuildHealthOptionsController.ts b/Engine/Source/Programs/Horde/HordeDashboard/src/components/buildhealth/BuildHealthOptionsController.ts
--- a/Engine/Source/Programs/Horde/HordeDashboard/src/components/buildhealth/BuildHealthOptionsController.ts
+++ b/Engine/Source/Programs/Horde/HordeDashboard/src/components/buildhealth/BuildHealthOptionsController.ts
@@ -1,8 +1,7 @@
 import { action, makeObservable } from "mobx";
 import { BuildHealthOptionsState, BuildHealthOptionStateDiff, BuildHealthOptionStateReceipt, BuildHealthUIOptionsState, TimeSpan } from "./BuildHealthOptions";
-import { decodeStateKey, decodeStepKey, decodeTemplateKey, encodeStepNameFromStringsPreserveCase, encodeTemplateKeyFromStrings } from "./BuildHealthUtilities";
+import { decodeFullyQualifiedLabelName, decodeLabelKey, decodeStateKey, decodeStepKey, decodeTemplateKey, encodeFullyQualifiedLabelName, encodeLabelKeyFromStrings, encodeStepNameFromStringsPreserveCase, encodeTemplateKeyFromStrings } from "./BuildHealthUtilities";
 import { BuildHealthDataHandler } from "./BuildHealthDataHandler";
-import { JobStepState } from "horde/backend/Api";
 
 /**
  * An interface that can provide hierarchy checks on stream, template, and step relationships.
@@ -17,13 +16,22 @@ export interface IHierarchyValidator {
     isValidTemplate(streamId: string, templateId: string): boolean;
 
     /**
-         * Determines whether a stream, template, and step combination is valid (in that it exists).
-         * @param streamId The stream.
-         * @param templateId The template.
-         * @param stepName The step name.
-         * @returns True if this stream has a template of the given id, which has a step of the given name, false otherwise.
-         */
+     * Determines whether a stream, template, and step combination is valid (in that it exists).
+     * @param streamId The stream.
+     * @param templateId The template.
+     * @param stepName The step name.
+     * @returns True if this stream has a template of the given id, which has a step of the given name, false otherwise.
+     */
     isValidStep(streamId: string, templateId: string, stepName: string): boolean;
+
+    /**
+     * Determines whether a stream, template, and label combination is valid (in that it exists).
+     * @param streamId The stream.
+     * @param templateId The template.
+     * @param fullyQualifiedLabelName The fully qualified label name.
+     * @returns True if this stream has a template of the given id, which has a label of the given name, false otherwise.
+     */
+    isValidLabel(streamId: string, templateId: string, fullyQualifiedLabelName: string): boolean;
 }
 
 /**
@@ -87,7 +95,7 @@ export interface IBuildHealthOptionWriter {
     toggleSingleStream(id: string, name: string): void;
 
     /**
-     * Toggles all linkedable templates, for all active streams.
+     * Toggles all linkable templates, for all active streams.
      * @param enabled True if we want to enable, false otherwise.
      */
     toggleAllLinkableTemplates(enabled: boolean): void;
@@ -101,7 +109,7 @@ export interface IBuildHealthOptionWriter {
     toggleTemplate(id: string, name: string, enabled: boolean): void;
 
     /**
-     * Toggles all linkedable steps, for all active templates.
+     * Toggles all linkable steps, for all active templates.
      * @param enabled True if we want to enable, false otherwise.
      */
     toggleAllLinkableSteps(enabled: boolean): void;
@@ -116,6 +124,22 @@ export interface IBuildHealthOptionWriter {
      */
     toggleAllSteps(streamId: string, templateId: string, stepNames: string[], enabled: boolean, cascadeLinkedAction?: boolean): void;
 
+    /**
+     * Toggles all linkable labels, for all active templates.
+     * @param enabled True if we want to enable, false otherwise.
+     */
+    toggleAllLinkableLabels(enabled: boolean): void;
+
+    /**
+     * Toggles all provided labels for a given stream & template, given they are valid.
+     * @param streamId The stream.
+     * @param templateId The template.
+     * @param labels The list of labels to toggle.
+     * @param enabled True to enable, false otherwise.
+     * @param cascadeLinkedAction True to enable linked action propagation (default behaviour), false or otherwise.
+     */
+    toggleAllLabels(streamId: string, templateId: string, labels: string[], enabled: boolean, cascadeLinkedAction?: boolean): void;
+
     /**
      * Toggles the state of the provided step.
      * @param id The id of the step to toggle.
@@ -124,6 +148,14 @@ export interface IBuildHealthOptionWriter {
      */
     toggleStep(id: string, name: string, enabled: boolean): void;
 
+    /**
+     * Toggles the state of the provided label.
+     * @param id The id of the label to toggle.
+     * @param name The fully qualified name of the label.
+     * @param enabled Whether it's enabled or not.
+     */
+    toggleLabel(id: string, name: string, enabled: boolean): void;
+
     /**
      * Toggles all provided states.
      * @param enabled True to enable, false otherwise.
@@ -162,6 +194,12 @@ export interface IBuildHealthOptionWriter {
       */
     clearTemplates(cascade?: boolean, owningStreamId?: string): void;
 
+    /**
+     * Clear all the labels tracked by the options.
+     * @param owningHierarchy Whether to clear an item if it belongs to a stream & template. If not provided, will remove any label.
+     */
+    clearLabels(owningHierarchy?: { streamId: string; templateId: string }): void;
+
     /**
      * Clear all the steps tracked by the options.
      * @param owningHierarchy Whether to clear an item if it belongs to a stream & template. If not provided, will remove any step.
@@ -273,6 +311,19 @@ export class BuildHealthOptionsWriter implements IBuildHealthOptionWriter {
         }
     }
 
+    /**
+     * @inheritdoc
+     */
+    @action
+    toggleAllLinkableLabels(enabled: boolean) {
+        if (this.state.linkedSelectionModeEnabled) {
+            Object.keys(this.state.enabledLabels).forEach(labelKey => {
+                let fullyQualifiedLabelName = this.state.enabledLabels[labelKey];
+                this.toggleLinkedLabels(labelKey, fullyQualifiedLabelName, enabled);
+            });
+        }
+    }
+
     /**
      * @inheritdoc
      */
@@ -286,6 +337,17 @@ export class BuildHealthOptionsWriter implements IBuildHealthOptionWriter {
         }
     }
 
+    /**
+     * @inheritdoc
+     */
+    @action
+    toggleLabel(id: string, name: string, enabled: boolean) {
+        this.toggleLabelInternal(id, name, enabled);
+        if (this.state.linkedSelectionModeEnabled) {
+            this.toggleLinkedLabels(id, name, enabled);
+        }
+    }
+
     /**
      * @inheritdoc
      */
@@ -325,6 +387,19 @@ export class BuildHealthOptionsWriter implements IBuildHealthOptionWriter {
         }
     }
 
+    /**
+     * @inheritdoc
+     */
+    @action
+    toggleAllLabels(streamId: string, templateId: string, labels: string[], enabled: boolean, cascadeLinkedAction: boolean = true) {
+        for (let idx: number = 0; idx < labels.length; ++idx) {
+            let label = labels[idx];
+            let decodedFullyQualifiedLabelName = decodeFullyQualifiedLabelName(label);
+            let toggleAllStepId = encodeLabelKeyFromStrings(streamId, templateId, decodedFullyQualifiedLabelName.labelCategory, decodedFullyQualifiedLabelName.label);
+            this.toggleLabelInternal(toggleAllStepId, label, enabled);
+        }
+    }
+
     /**
      * @inheritdoc
      */
@@ -396,12 +471,25 @@ export class BuildHealthOptionsWriter implements IBuildHealthOptionWriter {
                 delete this.state.enabledTemplates[key];
 
                 if (cascade) {
+                    this.clearLabels(owningStreamId ? { streamId: owningStreamId!, templateId: templateId } : undefined);
                     this.clearSteps(owningStreamId ? { streamId: owningStreamId!, templateId: templateId } : undefined);
                 }
             }
         });
     }
 
+    /**
+     * @inheritdoc
+     */
+    clearLabels(owningHierarchy?: { streamId: string; templateId: string }) {
+        Object.keys(this.state.enabledLabels).forEach(key => {
+            let { streamId, templateId, dashboardCategory, dashboardName } = decodeLabelKey(key);
+            if ((!owningHierarchy) || (streamId === owningHierarchy.streamId && templateId === owningHierarchy.templateId)) {
+                delete this.state.enabledLabels[key];
+            }
+        });
+    }
+
     /**
      * @inheritdoc
      */
@@ -431,6 +519,7 @@ export class BuildHealthOptionsWriter implements IBuildHealthOptionWriter {
         } else {
             delete this.state.enabledTemplates[id];
             let { streamId, templateId } = decodeTemplateKey(id);
+            this.clearLabels({ streamId, templateId });
             this.clearSteps({ streamId, templateId });
         }
     }
@@ -447,6 +536,28 @@ export class BuildHealthOptionsWriter implements IBuildHealthOptionWriter {
         });
     }
 
+    private toggleLabelInternal(id: string, name: string, enabled: boolean) {
+        if (enabled) {
+            this.state.enabledLabels[id] = name;
+        } else {
+            delete this.state.enabledLabels[id];
+        }
+    }
+
+    private toggleLinkedLabels(sourceId: string, name: string, enabled: boolean) {
+        let sourceDecodedLabelKey = decodeLabelKey(sourceId);
+        Object.keys(this.state.enabledTemplates).forEach(key => {
+            let targetDecodedTempalteKey = decodeTemplateKey(key);
+            // Set this step on every other enabled template, except the narrow case of the stream-template-step - we've just come from there.
+            if (targetDecodedTempalteKey.templateId !== sourceDecodedLabelKey.templateId || targetDecodedTempalteKey.streamId !== sourceDecodedLabelKey.streamId) {
+                let linkedEncodedStepId = encodeLabelKeyFromStrings(targetDecodedTempalteKey.streamId, targetDecodedTempalteKey.templateId, sourceDecodedLabelKey.dashboardCategory, sourceDecodedLabelKey.dashboardName);
+                if (this.hierarchyValidator?.isValidLabel(targetDecodedTempalteKey.streamId, targetDecodedTempalteKey.templateId, name)) {
+                    this.toggleLabelInternal(linkedEncodedStepId, name, enabled);
+                }
+            }
+        });
+    }
+
     private toggleStepInternal(id: string, name: string, enabled: boolean) {
         if (enabled) {
             this.state.enabledSteps[id] = name;
@@ -491,7 +602,7 @@ export class BuildHealthOptionsController {
     // #region -- Public Members --
 
     optionsChangeVersion = 0;
-    querySchemaVersion: number = 3;
+    querySchemaVersion: number = 4;
 
     // #endregion -- Public Members --
 
@@ -537,15 +648,6 @@ export class BuildHealthOptionsController {
         this.uiState.warningsAsSummaryFailure = warningsAsSummaryFailure;
     }
 
-    /**
-     * Sets whether the step outcome table shows label values.
-     * @param showLabels True to show label step data.
-     */
-    @action
-    setShowLabels(showLabels: boolean) {
-        this.uiState.showLabels = showLabels;
-    }
-
     /**
      * Gets a options transaction session to make changes to.
      * @returns The @see IBuildHealthOptionWriter that will enact options changes.
@@ -597,6 +699,7 @@ export class BuildHealthOptionsController {
         if (this.lastSynchronize < this.optionsChangeVersion) {
             this.state.stepOutcomeEnabledStreamKeys = Object.keys(this.state.enabledStreams);
             this.state.stepOutcomeEnabledJobKeys = Object.keys(this.state.enabledTemplates);
+            this.state.stepOutcomeEnabledLabelKeys = Object.keys(this.state.enabledLabels);
             this.state.stepOutcomeEnabledStepKeys = Object.keys(this.state.enabledSteps);
             this.state.stepOutcomeEnabledsStates = Object.keys(this.state.enabledStates);
             this.state.startDate = new Date(new Date().valueOf() - (this.state.jobHistoryTimeSpan.minutes * 60000));
diff --git a/Engine/Source/Programs/Horde/HordeDashboard/src/components/buildhealth/BuildHealthUtilities.ts b/Engine/Source/Programs/Horde/HordeDashboard/src/components/buildhealth/BuildHealthUtilities.ts
--- a/Engine/Source/Programs/Horde/HordeDashboard/src/components/buildhealth/BuildHealthUtilities.ts
+++ b/Engine/Source/Programs/Horde/HordeDashboard/src/components/buildhealth/BuildHealthUtilities.ts
@@ -1,6 +1,67 @@
 // Copyright Epic Games, Inc. All Rights Reserved.
 
-import { KEY_SEPARATOR, StepRefData, TemplateRefData } from "./BuildHealthDataTypes";
+import { KEY_SEPARATOR, LabelRefData, StepRefData, TemplateRefData } from "./BuildHealthDataTypes";
+
+/**
+ * Produces an encoded fully qualified label name.
+ * @param labelCategory The label category.
+ * @param label The label name.
+ * @returns The fully qualified label name.
+ */
+export function encodeFullyQualifiedLabelName(labelCategory: string | undefined, label: string | undefined) {
+    return `${labelCategory}${KEY_SEPARATOR}${label}`;
+}
+
+/**
+ * Decodes a fully qualified label name to it's constituent parts.
+ * @param fullyQualifiedLabelName 
+ * @returns The constituent parts.
+ */
+export function decodeFullyQualifiedLabelName(fullyQualifiedLabelName: string): { labelCategory: string, label: string } {
+    let [labelCategory, label] = fullyQualifiedLabelName.split(KEY_SEPARATOR);
+    return { labelCategory, label };
+}
+
+/**
+ * Produces an encoded label key for unique identification
+ * @param streamId The streamId to encode.
+ * @param templateId The templateId to encode.
+ * @param labelCategory 
+ * @param labelName 
+ * @returns The key representing the label.
+ */
+export function encodeLabelKeyFromStrings(streamId: string, templateId: string, labelCategory: string | undefined, labelName: string | undefined) {
+    return `${streamId.toLocaleLowerCase()}${KEY_SEPARATOR}${templateId}${KEY_SEPARATOR}${labelCategory?.toLocaleLowerCase()}${KEY_SEPARATOR}${labelName?.toLocaleLowerCase()}`;
+}
+
+/**
+ * Produces an encoded label key for unique identification
+ * @param streamId The streamId to encode.
+ * @param templateId The templateId to encode.
+ * @param labelCategory 
+ * @param labelName 
+ * @returns The key representing the label.
+ */
+export function encodeLabelKeyFromStringsPreserveCase(streamId: string, templateId: string, labelCategory: string | undefined, labelName: string | undefined) {
+    return `${streamId.toLocaleLowerCase()}${KEY_SEPARATOR}${templateId}${KEY_SEPARATOR}${labelCategory}${KEY_SEPARATOR}${labelName}`;
+}
+
+/**
+ * Produces an encoded label key for unique identification.
+ * @returns The key representing the state.
+ */
+export function encodeLabelKey(lableStateReponse: LabelRefData): string {
+    return `${lableStateReponse.streamId}${KEY_SEPARATOR}${lableStateReponse.templateId}${KEY_SEPARATOR}${lableStateReponse.dashboardCategory?.toLocaleLowerCase()}${KEY_SEPARATOR}${lableStateReponse.dashboardName?.toLocaleLowerCase()}`;
+}
+
+/**
+ * Decodes an encoded label key, providing it's constituent parts.
+ * @returns The constituent parts.
+ */
+export function decodeLabelKey(encodedStateKey: string): { streamId: string, templateId: string, dashboardCategory: string, dashboardName: string } {
+    const [streamId, templateId, dashboardCategory, dashboardName] = encodedStateKey.split(KEY_SEPARATOR);
+    return { streamId, templateId, dashboardCategory, dashboardName };
+}
 
 /**
  * Produces an encoded state key for unique identification.
@@ -11,7 +72,7 @@ export function encodeStateKeyFromString(index: number, stateName: string): stri
 }
 
 /**
- * Decodes an encoed state key, providing it's constituent parts.
+ * Decodes an encoded state key, providing it's constituent parts.
  * @returns The constituent parts.
  */
 export function decodeStateKey(encodedStateKey: string): { index: number; stateName: string } {
diff --git a/Engine/Source/Programs/Horde/HordeDashboard/src/components/buildhealth/BuildhealthView.tsx b/Engine/Source/Programs/Horde/HordeDashboard/src/components/buildhealth/BuildhealthView.tsx
--- a/Engine/Source/Programs/Horde/HordeDashboard/src/components/buildhealth/BuildhealthView.tsx
+++ b/Engine/Source/Programs/Horde/HordeDashboard/src/components/buildhealth/BuildhealthView.tsx
@@ -1,7 +1,7 @@
 // Copyright Epic Games, Inc. All Rights Reserved.
 
-import { Checkbox, ComboBox, DefaultButton, DirectionalHint, Dropdown, DropdownMenuItemType, IconButton, IDropdownOption, IDropdownProps, ISelectableOption, Label, mergeStyleSets, PrimaryButton, SearchBox, Stack, TooltipHost } from "@fluentui/react";
-import { GetProjectResponse, GetStreamResponse, JobStepState } from "horde/backend/Api";
+import { Checkbox, ComboBox, DefaultButton, DirectionalHint, Dropdown, DropdownMenuItemType, IconButton, IDropdownOption, IDropdownProps, ISelectableOption, Label, mergeStyleSets, PrimaryButton, SearchBox, SelectableOptionMenuItemType, Stack, TooltipHost } from "@fluentui/react";
+import { GetProjectResponse, GetStreamResponse } from "horde/backend/Api";
 import { observer } from "mobx-react-lite";
 import { useCallback, useEffect, useMemo, useRef, useState } from "react";
 import { useLocation, useNavigate } from "react-router";
@@ -11,9 +11,9 @@ import { Breadcrumbs } from "../Breadcrumbs";
 import { TopNav } from "../TopNav";
 import { StepOutcomeView } from "../buildhealth/stepoutcome/StepOutcome";
 import { BuildHealthDataHandler, DataHandlerRefreshRequest, EMPTY_REFRESH_REQUEST, expandRefreshRequest } from "./BuildHealthDataHandler";
-import { KEY_SEPARATOR, PARAMETER_KEY_PREFIX, StepRefData, TemplateRefData } from "./BuildHealthDataTypes";
+import { KEY_SEPARATOR, LabelRefData, PARAMETER_KEY_PREFIX, StepRefData, TemplateRefData } from "./BuildHealthDataTypes";
 import { ALL_STATES_TUPLE_LIST, BuildHealthOptionsState, BuildHealthOptionStateDiff, BuildHealthUIOptionsState, JobHistoryTimeSpans, loadBuildHealthOptionsFromParams, parseBuildHealthQueryParams, TimeSpan, toNavigationQuery, upgradeBuildHealthQueryParams } from "./BuildHealthOptions";
-import { decodeStepKey, decodeTemplateKey, encodeStateKeyFromString, encodeStepKey, encodeTemplateKey } from "./BuildHealthUtilities";
+import { decodeStepKey, decodeTemplateKey, encodeFullyQualifiedLabelName, encodeLabelKey, encodeStepKey, encodeTemplateKey } from "./BuildHealthUtilities";
 import { stepOutcomeTableClasses } from "./stepoutcome/StepOutcomeSharedUIComponents";
 import { runInAction } from "mobx";
 import { BuildHealthOptionsController } from "./BuildHealthOptionsController";
@@ -355,9 +355,9 @@ const MultiListParameter: React.FC<{ handler: BuildHealthDataHandler, buildHealt
 
             // Select & Clear all for the given filtered list.
             if (option.itemType === DropdownMenuItemType.Header && option.text === SELECT_CLEAR_ALL_SPECIAL_GROUP) {
-                let fitleredOptionList = filteredOptionsRef.current.length == doptions.length ? undefined : filteredOptionsRef.current;
+                let fitleredOptionList = filteredOptionsRef.current.length == doptions.length ? undefined : filteredOptionsRef.current.filter(x => x.itemType == undefined || x.itemType != SelectableOptionMenuItemType.Header);
                 let selectAllLabel = (fitleredOptionList !== undefined && fitleredOptionList.length > 0) ? "Select Filtered" : "Select All";
-
+                let showButton = fitleredOptionList == undefined || fitleredOptionList.length;
                 return (
                     <Stack horizontal tokens={{ childrenGap: 8 }} styles={{
                         root:
@@ -369,20 +369,22 @@ const MultiListParameter: React.FC<{ handler: BuildHealthDataHandler, buildHealt
                         }
                     }}>
                         {/* We have opted to use buttons here as we: 1. don't want to persist the state of a select all option; 2. the semantics of storing a "Select/Clear All" value gets complicated when the action is applied to a filtered set; */}
-                        <DefaultButton title="Selects all options (or currently filtered) within the group." styles={{ root: { flex: 1, height: 20 } }} text={selectAllLabel}
-                            onClick={() => {
-                                if (!option) return;
-
-                                const synthesizedOption: IDropdownOption<BuildHealthOptionData> = {
-                                    ...option,
-                                    selected: true,
-                                };
-
-                                // Only send a suggested list if we actually have filtered something.
-                                let fitleredOptionList = filteredOptionsRef.current.length == doptions.length ? undefined : filteredOptionsRef.current;
-                                option && onSelectAll?.(synthesizedOption, fitleredOptionList)
-                            }}
-                        />
+                        {showButton &&
+                            <DefaultButton title="Selects all options (or currently filtered) within the group." styles={{ root: { flex: 1, height: 20 } }} text={selectAllLabel}
+                                onClick={() => {
+                                    if (!option) return;
+
+                                    const synthesizedOption: IDropdownOption<BuildHealthOptionData> = {
+                                        ...option,
+                                        selected: true,
+                                    };
+
+                                    // Only send a suggested list if we actually have filtered something.
+                                    let fitleredOptionList = filteredOptionsRef.current.length == doptions.length ? undefined : filteredOptionsRef.current;
+                                    option && onSelectAll?.(synthesizedOption, fitleredOptionList)
+                                }}
+                            />
+                        }
                         <DefaultButton title="Clears all options within the group." styles={{ root: { flex: 1, height: 20 } }} text="Clear All"
                             onClick={() => {
                                 if (!option) return;
@@ -406,7 +408,7 @@ const MultiListParameter: React.FC<{ handler: BuildHealthDataHandler, buildHealt
                 // Due to some load testing, when we get 1000s of entires we don't want to bloat the load time with ToolTipHost construction.
                 <div title={`Group: ${option.data?.group || ""}\n${option.text}`} className={dropdownOptions.base}>{option.text}</div>
             );
-        }, []);
+        }, [filteredOptionsRef.current]);
 
     const handleDropdownChange = useCallback((_ev, option) => {
         if (option && option.itemType !== DropdownMenuItemType.Header) {
@@ -593,15 +595,17 @@ const StepDropdown: React.FC<{ handler: BuildHealthDataHandler, buildHealthOptio
     // #region -- Data Memos --
 
     const items = useMemo<BuildHealthOptionData[]>(
-        () =>
-            handler.stepData
+        () => {
+            let transformedStepData = handler.stepData
                 .map((step: StepRefData) => ({
                     id: encodeStepKey(step),
                     text: step.name,
                     group: `${step.streamId} - ${step.templateId}`,
-                }))
-                .sort((a, b) => a.text.localeCompare(b.text)),
-        [handler.stepData]
+                })).sort((a, b) => a.text.localeCompare(b.text));
+
+            return transformedStepData;
+        },
+        [handler.stepData, handler.labelData]
     );
 
     // #endregion -- Data Memos --
@@ -699,6 +703,121 @@ const StepDropdown: React.FC<{ handler: BuildHealthDataHandler, buildHealthOptio
     );
 });
 
+const LabelDropdown: React.FC<{ handler: BuildHealthDataHandler, buildHealthOptions: BuildHealthOptionsController }> = observer(function ConstructProjectDropdown({ handler, buildHealthOptions }) {
+    let selectedLabelKeys: string[] = Object.keys(buildHealthOptions.state.enabledLabels);
+    let request: DataHandlerRefreshRequest = { steps: false };
+
+    // #region -- Data Memos --
+
+    const items = useMemo<BuildHealthOptionData[]>(
+        () => {
+            let transformedLabelData = handler.labelData
+                .map((label: LabelRefData) => ({
+                    id: encodeLabelKey(label),
+                    text: encodeFullyQualifiedLabelName(label.dashboardCategory, label.dashboardName),
+                    group: `${label.streamId} - ${label.templateId}`,
+                })).sort((a, b) => a.text.localeCompare(b.text));
+
+            return transformedLabelData;
+        },
+        [handler.labelData]
+    );
+
+    // #endregion -- Data Memos --
+
+    // #region -- Callbacks --
+
+    const handleChange = useCallback(
+        (option: IDropdownOption<BuildHealthOptionData>) => {
+            const key = option.key as string;
+            const isSelected = !!option.selected;
+            buildHealthOptions
+                .getTransactionSession()
+                .toggleLabel(key, option.text, isSelected);
+        },
+        [buildHealthOptions]
+    );
+
+    const handleSelectAll = useCallback(
+        (
+            option: IDropdownOption<BuildHealthOptionData>,
+            optionSelectAllContext?: IDropdownOption<BuildHealthOptionData>[]
+        ) => {
+            const key = option.key as string;
+            const [, data] = key.split(KEY_SEPARATOR);
+            const [streamId, templateId] = data.split(" - ");
+            const isSelected = !!option.selected;
+
+            // If there is no provided context, simply obtain all the known steps of the (stream, template).
+            const labelNames: string[] =
+                optionSelectAllContext === undefined
+                    ? handler.getCachedTemplateLabelNames(streamId, templateId)
+                    : optionSelectAllContext
+                        .filter((opt) => {
+                            if (!opt.data) return false;
+                            const decoded = decodeStepKey(opt.data.id!);
+                            return (
+                                decoded.streamId === streamId &&
+                                decoded.templateId === templateId
+                            );
+                        })
+                        .map((x) => x.data!.text!);
+
+            // Do not close the transaction session; the user still has the modal open
+            buildHealthOptions
+                .getTransactionSession()
+                .toggleAllLabels(streamId, templateId, labelNames, isSelected);
+        },
+        [buildHealthOptions, handler]
+    );
+
+    const handleLinkSteps = useCallback(
+        () => {
+            // Close the transaction session; the modal is not open so there will be no such transaction completion otherwise.
+            buildHealthOptions.getTransactionSession().toggleAllLinkableLabels(true);
+            buildHealthOptions.commitTransactionSession();
+            handler.requestHierarchicalRefresh(request);
+        },
+        [buildHealthOptions, handler]
+    );
+
+    // #endregion -- Callbacks --
+
+    const labelDropdownParams: BuildHealthOptionList = {
+        label: "Select Label(s)",
+        items: [
+        ],
+        toolTip: "Select an option from the dropdown."
+    };
+
+    labelDropdownParams.items = items;
+
+    return labelDropdownParams.items.length > 0 && (
+        <Stack id={`label_dropdown`}>
+            <Label>Label(s)</Label>
+            <Stack horizontal tokens={{ childrenGap: 8 }} verticalAlign="end">
+                <MultiListParameter handler={handler} buildHealthOptions={buildHealthOptions} param={labelDropdownParams} selectedKeys={selectedLabelKeys} request={request}
+                    onChange={handleChange}
+                    enabledSelectAll={true}
+                    onSelectAll={handleSelectAll} />
+                {
+                    buildHealthOptions.state.linkedSelectionModeEnabled &&
+                    (
+                        <TooltipHost
+                            content={"Link labels: Apply selected labels across selected templates."}
+                            calloutProps={{ directionalHint: DirectionalHint.rightCenter }}>
+                            <IconButton
+                                iconProps={{ iconName: 'Link' }}
+                                styles={{ root: { width: 15, minWidth: 15, padding: 0 } }}
+                                onClick={handleLinkSteps} />
+                        </TooltipHost>
+                    )
+                }
+            </Stack>
+        </Stack>
+    );
+});
+
 const StateDropdown: React.FC<{ handler: BuildHealthDataHandler, buildHealthOptions: BuildHealthOptionsController }> = observer(function ConstructProjectDropdown({ handler, buildHealthOptions }) {
     let selectedStateKeys: string[] = Object.keys(buildHealthOptions.state.enabledStates);
 
@@ -835,7 +954,6 @@ const BuildHealthSidebar: React.FC<{ handler: BuildHealthDataHandler, buildHealt
 
     const [hideDateAnchors, setHideDateAnchors] = useState(() => buildHealthOptions.uiState.includeDateAnchors)
     const [warningsAsSummaryFailure, setWarningsAsSummaryFailure] = useState(() => buildHealthOptions.uiState.warningsAsSummaryFailure)
-    const [showLabels, setShowLabels] = useState(() => buildHealthOptions.uiState.showLabels)
 
     // #ednregion -- Visual Options --
 
@@ -867,11 +985,6 @@ const BuildHealthSidebar: React.FC<{ handler: BuildHealthDataHandler, buildHealt
         buildHealthOptions.setWarningsAsSummaryFailure(isChecked);
     }
 
-    const updateShowLabels = (isChecked) => {
-        setShowLabels(isChecked);
-        buildHealthOptions.setShowLabels(isChecked);
-    }
-
     // #endregion -- Property Callbacks --
 
     // #region -- Property Use Effects --
@@ -920,6 +1033,7 @@ const BuildHealthSidebar: React.FC<{ handler: BuildHealthDataHandler, buildHealt
                         <ProjectDropdownSingle handler={handler} buildHealthOptions={buildHealthOptions} />
                         <StreamDropdown handler={handler} buildHealthOptions={buildHealthOptions} />
                         <TemplateDropdown handler={handler} buildHealthOptions={buildHealthOptions} />
+                        <LabelDropdown handler={handler} buildHealthOptions={buildHealthOptions} />
                         <StepDropdown handler={handler} buildHealthOptions={buildHealthOptions} />
                         <StateDropdown handler={handler} buildHealthOptions={buildHealthOptions} />
                         <TooltipHost content={
@@ -1022,23 +1136,6 @@ const BuildHealthSidebar: React.FC<{ handler: BuildHealthDataHandler, buildHealt
                             onChange={(_, isChecked) => updateWarningsAsSummaryFailure(!!isChecked)}
                         />
                     </TooltipHost>
-                    <TooltipHost content={
-                        <div style={{ maxWidth: 300 }}>
-                            <div>
-                                Shows labels (e.g. Win64) as entries in the table.
-                            </div>
-                        </div>
-                    }
-                        calloutProps={{
-                            directionalHint: DirectionalHint.rightCenter
-                        }}
-                    >
-                        <Checkbox
-                            label="Show Labels"
-                            checked={showLabels}
-                            onChange={(_, isChecked) => updateShowLabels(!!isChecked)}
-                        />
-                    </TooltipHost>
                     <TooltipHost content={
                         <div style={{ maxWidth: 300 }}>
                             <div>
@@ -1091,6 +1188,7 @@ export const BuildHealthView: React.FC = observer(function ConstructBuildHealthV
         streams: buildHealthState.stepOutcomeEnabledStreamKeys,
         jobs: buildHealthState.stepOutcomeEnabledJobKeys,
         steps: buildHealthState.stepOutcomeEnabledStepKeys,
+        labels: buildHealthState.stepOutcomeEnabledLabelKeys,
         jobHistorySpan,
         includePreflights: buildHealthState.includePreflight,
     }),
@@ -1106,12 +1204,10 @@ export const BuildHealthView: React.FC = observer(function ConstructBuildHealthV
     const uiOptions = useMemo(() => ({
         includeDateAnchors: buildHealthUIState.includeDateAnchors,
         warningsAsSummaryFailure: buildHealthUIState.warningsAsSummaryFailure,
-        showLabels: buildHealthUIState.showLabels,
         validStates: buildHealthUIState.validStates
     }),
         // eslint-disable-next-line react-hooks/exhaustive-deps -- All necessary dependencies included for MobX observables
         [
-            buildHealthUIState.showLabels,
             buildHealthUIState.includeDateAnchors,
             buildHealthUIState.warningsAsSummaryFailure,
             buildHealthUIState.validStates
diff --git a/Engine/Source/Programs/Horde/HordeDashboard/src/components/buildhealth/stepoutcome/StepOutcomeD3Component.tsx b/Engine/Source/Programs/Horde/HordeDashboard/src/components/buildhealth/stepoutcome/StepOutcomeD3Component.tsx
--- a/Engine/Source/Programs/Horde/HordeDashboard/src/components/buildhealth/stepoutcome/StepOutcomeD3Component.tsx
+++ b/Engine/Source/Programs/Horde/HordeDashboard/src/components/buildhealth/stepoutcome/StepOutcomeD3Component.tsx
@@ -62,13 +62,11 @@ function exportStepOutcomeDataToCSV(stepOutcomeTable: StepOutcomeTable, handler:
     saveFile(stepOrderedCsvContent, `${filename}-StepOrderedData-${defaultName}.csv`);
     saveFile(changeOrderedCsvContent, `${filename}-ChangeOrderedData-${defaultName}.csv`);
 
-    if (stepOutcomeTable.supportsLabelEntries) {
-        const labelOrderedCsvContent = stepOutcomeTable.toRowOrderedCSV(handler, isLabelOutcome);
-        const labelchangeOrderedCsvContent = stepOutcomeTable.toColOrderedCSV(handler, isLabelOutcome);
+    const labelOrderedCsvContent = stepOutcomeTable.toRowOrderedCSV(handler, isLabelOutcome);
+    const labelchangeOrderedCsvContent = stepOutcomeTable.toColOrderedCSV(handler, isLabelOutcome);
 
-        saveFile(labelOrderedCsvContent, `${filename}-LabelOrderedData-${defaultName}.csv`);
-        saveFile(labelchangeOrderedCsvContent, `${filename}-LabelChangeOrderedData-${defaultName}.csv`);
-    }
+    saveFile(labelOrderedCsvContent, `${filename}-LabelOrderedData-${defaultName}.csv`);
+    saveFile(labelchangeOrderedCsvContent, `${filename}-LabelChangeOrderedData-${defaultName}.csv`);
 }
 
 // #endregion -- Step Outcome Table Download Handlers --
@@ -368,7 +366,7 @@ function GenerateD3Table(stepOutcomeTable: StepOutcomeTable, handler: StepOutcom
             const tr = d3.select(this);
             const stream = rowData.streamId;
             let metadata = stepOutcomeTable.tableStreamMetadata.get(stream) ?? { stepCount: 1, labelCount: 0 };
-            let rowSpan = stepOutcomeTable.viewModelOptions.showLabels ? metadata?.labelCount! + metadata?.stepCount! : metadata?.stepCount;
+            let rowSpan = metadata?.labelCount! + metadata?.stepCount!;
 
             // Only add a TH for the first row of each stream   
             const streamData = !streamTracker.has(stream)
diff --git a/Engine/Source/Programs/Horde/HordeDashboard/src/components/buildhealth/stepoutcome/StepOutcomeDataHandler.ts b/Engine/Source/Programs/Horde/HordeDashboard/src/components/buildhealth/stepoutcome/StepOutcomeDataHandler.ts
--- a/Engine/Source/Programs/Horde/HordeDashboard/src/components/buildhealth/stepoutcome/StepOutcomeDataHandler.ts
+++ b/Engine/Source/Programs/Horde/HordeDashboard/src/components/buildhealth/stepoutcome/StepOutcomeDataHandler.ts
@@ -4,7 +4,7 @@ import { GetJobResponse, IssueQuery, JobStepState, JobStreamQuery, LabelState }
 import { PollBase } from "horde/backend/PollBase";
 import { JobResponseData, LabelOutcomeTableEntry, StepOutcomeTable, StepOutcomeTableEntry, UpdatedChangeSummaryData } from "./StepOutcomeDataTypes";
 import backend from "horde/backend";
-import { decodeStepKey, decodeTemplateKey, encodeStepNameFromStrings } from "../BuildHealthUtilities";
+import { decodeLabelKey, decodeStepKey, decodeTemplateKey, encodeLabelKeyFromStrings, encodeStepNameFromStrings } from "../BuildHealthUtilities";
 import { action, computed, makeObservable, observable, runInAction } from "mobx";
 import { summarizeJob } from "./StepOutcomeUtilities";
 
@@ -84,8 +84,7 @@ export function prettyPrintStepOutcomeFiltersSummary(filters: StepOutcomeFilters
 export interface StepOutcomeViewOptions {
     includeDateAnchors: boolean;
     warningsAsSummaryFailure: boolean;
-    showLabels: boolean;
-    validStates : Set<string>;
+    validStates: Set<string>;
 }
 
 /**
@@ -97,6 +96,7 @@ export interface StepOutcomeFilters {
     jobHistorySpan: { start: Date; end?: Date };
     jobs?: string[];
     steps?: string[];
+    labels?: string[];
     maxJobCount?: number;
     includePreflights: boolean;
 }
@@ -151,7 +151,7 @@ export class StepOutcomeDataHandler extends PollBase {
             this.activeFilter === undefined ||
             !this.activeFilter.streams?.length ||
             !this.activeFilter.jobs?.length ||
-            !this.activeFilter.steps?.length
+            (!this.activeFilter.steps?.length && !this.activeFilter.labels?.length)
         ) {
             this.stepOutcomeTableData.reset();
             this.setUpdated();
@@ -165,6 +165,7 @@ export class StepOutcomeDataHandler extends PollBase {
 
         // @todo UE-309726 - The poll should not need to understand internal data structure; the filter object should be already decoded.
         const stepMap: Map<string, Map<string, Set<string>>> = new Map();
+        const labelMap: Map<string, Map<string, Set<string>>> = new Map();
 
         (this.activeFilter.steps ?? []).forEach(step => {
             const { streamId, templateId, stepName } = decodeStepKey(step);
@@ -182,6 +183,22 @@ export class StepOutcomeDataHandler extends PollBase {
             templateMap.get(templateId)!.add(step.toLocaleLowerCase());
         });
 
+        (this.activeFilter.labels ?? []).forEach(label => {
+            const { streamId, templateId, dashboardCategory, dashboardName } = decodeLabelKey(label);
+
+            if (!labelMap.has(streamId)) {
+                labelMap.set(streamId, new Map());
+            }
+
+            const templateMap = labelMap.get(streamId)!;
+
+            if (!templateMap.has(templateId)) {
+                templateMap.set(templateId, new Set());
+            }
+
+            templateMap.get(templateId)!.add(label.toLocaleLowerCase());
+        });
+
         const streamList = this.activeFilter.streams ?? [];
         const modifiedAfter = this.activeFilter.jobHistorySpan.start.toISOString();
         const jobCount = this.activeFilter.maxJobCount ?? StepOutcomeDataHandler.MAX_JOB_COUNT_DEFAULT;
@@ -276,29 +293,33 @@ export class StepOutcomeDataHandler extends PollBase {
 
                 const batchResponses = jobResponse.batches ?? [];
                 let streamStepSet = stepMap.get(jobResponse.streamId);
+                let streamLabelSet = labelMap.get(jobResponse.streamId);
                 let templateStepSet = streamStepSet ? streamStepSet.get(jobResponse.templateId!) : undefined;
+                let templateLabelSet = streamLabelSet ? streamLabelSet.get(jobResponse.templateId!) : undefined;
 
-                for (const batch of batchResponses) {
-                    for (const step of batch.steps) {
-                        let unifiedKey: string = jobResponse.id + step.id;
-                        if (!jobStepKeyToState.has(unifiedKey)) {
-                            jobStepKeyToState.set(unifiedKey, step.state);
-                        }
-
-                        let fullyQualifiedStepName = encodeStepNameFromStrings(jobResponse.streamId, jobResponse.templateId!, step.name).toLocaleLowerCase();
-
-                        if (
-                            templateStepSet &&
-                            templateStepSet.size >= 0 &&
-                            templateStepSet.has(fullyQualifiedStepName)
-                        ) {
-                            this.stepOutcomeTableData.addEntry(new StepOutcomeTableEntry(step, jobResponse.id, jobResponse.name, jobResponse.streamId, jobResponse.createTime.toString(), jobResponse.change, this.changeData.get(jobResponse.change!)?.dateUtc));
+                if (templateStepSet && templateStepSet.size >= 0) {
+                    for (const batch of batchResponses) {
+                        for (const step of batch.steps) {
+                            let unifiedKey: string = jobResponse.id + step.id;
+                            if (!jobStepKeyToState.has(unifiedKey)) {
+                                jobStepKeyToState.set(unifiedKey, step.state);
+                            }
+
+                            let fullyQualifiedStepName = encodeStepNameFromStrings(jobResponse.streamId, jobResponse.templateId!, step.name).toLocaleLowerCase();
+
+                            if (
+                                templateStepSet &&
+                                templateStepSet.size >= 0 &&
+                                templateStepSet.has(fullyQualifiedStepName)
+                            ) {
+                                this.stepOutcomeTableData.addEntry(new StepOutcomeTableEntry(step, jobResponse.id, jobResponse.name, jobResponse.streamId, jobResponse.createTime.toString(), jobResponse.change, this.changeData.get(jobResponse.change!)?.dateUtc));
+                            }
                         }
                     }
                 }
 
                 // Do not compute label steps if there are no steps for the template.
-                if (templateStepSet && templateStepSet.size >= 0) {
+                if (templateLabelSet && templateLabelSet.size >= 0) {
                     const labelResponses = jobResponse.labels ?? [];
                     for (let idx: number = 0; idx < labelResponses.length; ++idx) {
                         let label = labelResponses[idx];
@@ -309,7 +330,16 @@ export class StepOutcomeDataHandler extends PollBase {
                         if (label.state === LabelState.Unspecified) {
                             synthesizedStepOutcome = jobResponse.abortedByUserInfo !== undefined ? JobStepState.Aborted : JobStepState.Skipped;
                         }
-                        this.stepOutcomeTableData.addEntry(new LabelOutcomeTableEntry(label, idx, jobResponse.id, jobResponse.name, jobResponse.streamId, jobResponse.createTime.toString(), jobResponse.change, this.changeData.get(jobResponse.change!)?.dateUtc, synthesizedStepOutcome));
+
+                        let fullyQualifiedLabelName = encodeLabelKeyFromStrings(jobResponse.streamId, jobResponse.templateId!, label.dashboardCategory, label.dashboardName).toLocaleLowerCase();
+
+                        if (
+                            templateLabelSet &&
+                            templateLabelSet.size >= 0 &&
+                            templateLabelSet.has(fullyQualifiedLabelName)
+                        ) {
+                            this.stepOutcomeTableData.addEntry(new LabelOutcomeTableEntry(label, idx, jobResponse.id, jobResponse.name, jobResponse.streamId, jobResponse.createTime.toString(), jobResponse.change, this.changeData.get(jobResponse.change!)?.dateUtc, synthesizedStepOutcome));
+                        }
                     }
                 }
             }
diff --git a/Engine/Source/Programs/Horde/HordeDashboard/src/components/buildhealth/stepoutcome/StepOutcomeDataTypes.ts b/Engine/Source/Programs/Horde/HordeDashboard/src/components/buildhealth/stepoutcome/StepOutcomeDataTypes.ts
--- a/Engine/Source/Programs/Horde/HordeDashboard/src/components/buildhealth/stepoutcome/StepOutcomeDataTypes.ts
+++ b/Engine/Source/Programs/Horde/HordeDashboard/src/components/buildhealth/stepoutcome/StepOutcomeDataTypes.ts
@@ -561,11 +561,6 @@ export class StepOutcomeTable {
      */
     supportsDateAnchoredChanges: boolean = true;
 
-    /**
-     * Sets whether @see getStepRowHeaders returns label table entries for the given data.
-     */
-    supportsLabelEntries: boolean = false;
-
     /**
      * Change column headers.
      */
@@ -931,9 +926,6 @@ export class StepOutcomeTable {
      */
     *getStepRowHeaders(): Generator<StepNameHeader> {
         for (let idx: number = 0; idx < this.stepNameRowHeaders.length; ++idx) {
-            if (!this.supportsLabelEntries && this.stepNameRowHeaders[idx].isLabel) {
-                continue;
-            }
             yield this.stepNameRowHeaders[idx];
         }
     }
@@ -1029,7 +1021,6 @@ export class StepOutcomeTable {
     @action
     setViewModelOptions(viewModelOptions: StepOutcomeViewOptions) {
         this.viewModelOptions = viewModelOptions;
-        this.supportsLabelEntries = viewModelOptions.showLabels;
         this.supportsDateAnchoredChanges = viewModelOptions.includeDateAnchors;
     }
 
