[Horde][BuildHealth][Server] Increase suspect scoring timeout to 30 seconds

diff --git a/Engine/Source/Programs/Horde/Plugins/Build/HordeServer.Build/BuildConfig.cs b/Engine/Source/Programs/Horde/Plugins/Build/HordeServer.Build/BuildConfig.cs
--- a/Engine/Source/Programs/Horde/Plugins/Build/HordeServer.Build/BuildConfig.cs
+++ b/Engine/Source/Programs/Horde/Plugins/Build/HordeServer.Build/BuildConfig.cs
@@ -109,7 +109,7 @@ namespace HordeServer
 		/// <summary>
 		/// The maximum number of seconds to spend scording suspects
 		/// </summary>
-		public int IssueFindSuspectsMaxScoringSeconds { get; set; } = 20;
+		public int IssueFindSuspectsMaxScoringSeconds { get; set; } = 30;
 
 		private readonly Dictionary<ProjectId, ProjectConfig> _projectLookup = new Dictionary<ProjectId, ProjectConfig>();
 		private readonly Dictionary<StreamId, StreamConfig> _streamLookup = new Dictionary<StreamId, StreamConfig>();
