Horde: Disable analyzers in Dockerfile build for Horde Server

They should be run in the already existing but separate BuildGraph target for better parallelism.

diff --git a/Engine/Source/Programs/Horde/HordeServer/Dockerfile b/Engine/Source/Programs/Horde/HordeServer/Dockerfile
--- a/Engine/Source/Programs/Horde/HordeServer/Dockerfile
+++ b/Engine/Source/Programs/Horde/HordeServer/Dockerfile
@@ -54,7 +54,7 @@ COPY Source/Programs/Horde/HordeServer.Shared/ Source/Programs/Horde/HordeServer
 COPY Source/Programs/AutomationTool/AutomationUtils/Matchers/ Source/Programs/AutomationTool/AutomationUtils/Matchers/
 COPY Source/Programs/UnrealBuildTool/Matchers/ Source/Programs/UnrealBuildTool/Matchers/
 ARG msbuild_args
-RUN dotnet publish Source/Programs/Horde/HordeServer/HordeServer.csproj -c Release -o out -p:WithHordeStorage=false $msbuild_args -p:RunAnalyzersDuringBuild=True
+RUN dotnet publish Source/Programs/Horde/HordeServer/HordeServer.csproj -c Release -o out -p:WithHordeStorage=false $msbuild_args
 
 # Run tests
 COPY Source/Programs/Horde/Drivers/JobDriver/ Source/Programs/Horde/Drivers/JobDriver/
