Avoid keeping Audit of tests that got pruned during clean up

#rnx

diff --git a/Engine/Source/Programs/Horde/Plugins/Build/HordeServer.Build/Jobs/TestData/TestDataCollectionV2.cs b/Engine/Source/Programs/Horde/Plugins/Build/HordeServer.Build/Jobs/TestData/TestDataCollectionV2.cs
--- a/Engine/Source/Programs/Horde/Plugins/Build/HordeServer.Build/Jobs/TestData/TestDataCollectionV2.cs
+++ b/Engine/Source/Programs/Horde/Plugins/Build/HordeServer.Build/Jobs/TestData/TestDataCollectionV2.cs
@@ -1638,6 +1638,18 @@ namespace HordeServer.Jobs.TestData
 			{
 				_logger.LogInformation("Pruned {TestStreamDeleteCount} test session streams", testStreamsDeleted);
 			}
+
+			// pruned all audits related to tests that were NOT kept
+			{
+				FilterDefinition<TestAuditDocument> filter = FilterDefinition<TestAuditDocument>.Empty;
+				FilterDefinitionBuilder<TestAuditDocument> filterBuilder = Builders<TestAuditDocument>.Filter;
+				filter &= filterBuilder.Nin(x => x.TestId, testIds);
+				DeleteResult result = await _testAudits.DeleteManyAsync(filter, cancellationToken);
+				if (result.DeletedCount > 0)
+				{
+					_logger.LogInformation("Pruned {NumDeleted} test audit documents", result.DeletedCount);
+				}
+			}
 		}
 	}
 }
