Testing AllowRetry flag.

diff --git a/Engine/Source/Programs/Horde/StorageTests.xml b/Engine/Source/Programs/Horde/StorageTests.xml
--- a/Engine/Source/Programs/Horde/StorageTests.xml
+++ b/Engine/Source/Programs/Horde/StorageTests.xml
@@ -16,7 +16,7 @@
 
 	<!-- Medium Chunks -->
 	<Agent Name="Medium Agent A" Type="CompileWin64;Win64">
-		<Node Name="Create Medium Chunks" Produces="#MediumChunks">
+		<Node Name="Create Medium Chunks" Produces="#MediumChunks" AllowRetry="False">
 			<RandomData Size="1048576" Count="20000" Different="true" Tag="#MediumChunks"/>
 		</Node>
 	</Agent>
