[Horde][BuildHealth][Dashboard] QoL enhancements:
- Persist view options to query url.
- Disable auto-select all for templates.
- Callout on dropdown moved down to avoid clashing with header dropdowns (when filter box changes).
- Select all now available in stream dropdown.
-- Code infrastructure is now there for templates, but is not connected due to server load issues that come from this.
- Minor fix where clearAll was not being reflected.
- Minor fix where UI properties were precipitating data refresh.
- Possible fix for filter <> selection desync

diff --git a/Engine/Source/Programs/Horde/HordeDashboard/src/components/buildhealth/BuildHealthDataHandler.ts b/Engine/Source/Programs/Horde/HordeDashboard/src/components/buildhealth/BuildHealthDataHandler.ts
--- a/Engine/Source/Programs/Horde/HordeDashboard/src/components/buildhealth/BuildHealthDataHandler.ts
+++ b/Engine/Source/Programs/Horde/HordeDashboard/src/components/buildhealth/BuildHealthDataHandler.ts
@@ -2,11 +2,11 @@
 
 import backend from "horde/backend";
 import { JobData, JobStreamQuery, ProjectData, StreamData } from "horde/backend/Api";
-import { makeObservable, observable, runInAction } from "mobx";
+import { action, makeObservable, observable, runInAction } from "mobx";
 import { BuildHealthOptionsState } from "./BuildHealthOptions";
 import { LabelRefData, StepRefData, TemplateRefData } from "./BuildHealthDataTypes";
 import { decodeTemplateKey, encodeFullyQualifiedLabelName, encodeTemplateKeyFromStrings } from "./BuildHealthUtilities";
-import { BuildHealthOptionsController, IHierarchyValidator } from "./BuildHealthOptionsController";
+import { BuildHealthOptionsController, IHierarchyValidator, StreamDescriptor, TemplateDescriptor } from "./BuildHealthOptionsController";
 
 // #region -- API Data Types & Utilities --
 
@@ -77,6 +77,7 @@ export class BuildHealthDataHandler implements IHierarchyValidator {
     private lastRefresh = -1;
     private buildHealthOptions: BuildHealthOptionsState;
     private buildHealthOptionsController: BuildHealthOptionsController;
+    private cachedProjectToStreamData: Map<string, StreamData[]> = new Map<string, StreamData[]>();
     private cachedStreamToTemplatesData: Map<string, TemplateRefData[]> = new Map<string, TemplateRefData[]>();
     private cachedStreamTempalteIdToSteps: Map<string, StepCacheEntry> = new Map<string, StepCacheEntry>();
     private cachedStreamTemplateIdToLabels: Map<string, LabelCacheEntry> = new Map<string, LabelCacheEntry>();
@@ -126,10 +127,12 @@ export class BuildHealthDataHandler implements IHierarchyValidator {
     /**
      * Resets the handler, retunring it to a project-only data state.
      */
+    @action
     reset() {
-        this.streamsData.length = 0;
-        this.templatesData.length = 0;
-        this.stepData.length = 0;
+        this.streamsData =[];
+        this.templatesData = [];
+        this.stepData = [];
+        this.labelData =[];
     }
 
     /**
@@ -173,6 +176,17 @@ export class BuildHealthDataHandler implements IHierarchyValidator {
 
             this.projectsData.forEach((project: ProjectData) => {
                 const isProjectEnabled = this.buildHealthOptions.enabledProjects.hasOwnProperty(project.id);
+
+                if (!this.cachedProjectToStreamData.has(project.id)) {
+                    this.cachedProjectToStreamData.set(project.id, []);
+                }
+
+                if (project.streams === undefined) {
+                    return;
+                }
+
+                this.cachedProjectToStreamData.get(project.id)?.push(...project.streams!);
+
                 if (!isProjectEnabled) {
                     return;
                 }
@@ -288,6 +302,37 @@ export class BuildHealthDataHandler implements IHierarchyValidator {
         }
     }
 
+    /**
+     *  Obtains all the currently known streams for a given project.
+     * @param projectId The project.
+     * @returns An array representing all the known streams for a given project.
+     */
+    getCachedStreams(projectId: string): StreamDescriptor[] {
+        let returnArray: StreamDescriptor[] = [];
+
+        if (this.cachedProjectToStreamData.has(projectId)) {
+            returnArray = this.cachedProjectToStreamData.get(projectId)!.map(x =>  ({ id: x.id, name: x.name }));
+        }
+
+        return returnArray;
+    }
+
+
+    /**
+     * Obtains all the currently known templates for a given stream.
+     * @param streamId THe stream.
+     * @returns An array representing all the known templates for a given stream.
+     */
+    getCachedTemplates(streamId: string): TemplateDescriptor[] {
+        let returnArray: TemplateDescriptor[] = [];
+
+        if (this.cachedStreamToTemplatesData.has(streamId)) {
+            returnArray = this.cachedStreamToTemplatesData.get(streamId)!.map(x => ({ id: x.id, name: x.name }));
+        }
+
+        return returnArray;
+    }
+
     /**
      * Obtains all the currently known step names for a given stream & template.
      * @param streamId The stream.
diff --git a/Engine/Source/Programs/Horde/HordeDashboard/src/components/buildhealth/BuildHealthDataTypes.ts b/Engine/Source/Programs/Horde/HordeDashboard/src/components/buildhealth/BuildHealthDataTypes.ts
--- a/Engine/Source/Programs/Horde/HordeDashboard/src/components/buildhealth/BuildHealthDataTypes.ts
+++ b/Engine/Source/Programs/Horde/HordeDashboard/src/components/buildhealth/BuildHealthDataTypes.ts
@@ -14,6 +14,9 @@ export const TIME_SPAN_URL_PARAM: string = "lastTimeRange";
 export const INCLUDE_PREFLIGHT_PARAM: string = "includePreflight";
 export const KEY_SEPARATOR: string = "::";
 
+export const UI_SUMMARY_WARNING_AS_ERROR_PARAM: string = "warningAsError";
+export const UI_DATE_ANCHORS_PARAM: string = "dateAnchors";
+
 export const LABELS_URL_PARAM_ALL_LABELS: string = "labels_AllKnownValues";
 export const STEPS_URL_PARAM_ALL_STEPS: string = "stepNames_AllKnownValues";
 export const STATES_URL_PARAM_ALL_STATES: string = "stateNames_AllKnownValues";
diff --git a/Engine/Source/Programs/Horde/HordeDashboard/src/components/buildhealth/BuildHealthOptions.ts b/Engine/Source/Programs/Horde/HordeDashboard/src/components/buildhealth/BuildHealthOptions.ts
--- a/Engine/Source/Programs/Horde/HordeDashboard/src/components/buildhealth/BuildHealthOptions.ts
+++ b/Engine/Source/Programs/Horde/HordeDashboard/src/components/buildhealth/BuildHealthOptions.ts
@@ -1,7 +1,7 @@
 // Copyright Epic Games, Inc. All Rights Reserved.
 
 import { makeObservable, observable, runInAction } from "mobx";
-import { INCLUDE_PREFLIGHT_PARAM, LABELS_URL_PARAM_NAME, PROJECTS_URL_PARAM_NAME, QUERY_SCHEMA_PARAM_NAME, STATES_URL_PARAM_ALL_STATES, STATES_URL_PARAM_NAME, LABELS_URL_PARAM_ALL_LABELS, STEPS_URL_PARAM_ALL_STEPS, STEPS_URL_PARAM_NAME, STREAMS_URL_PARAM_NAME, TEMPLATES_URL_PARAM_NAME, TIME_SPAN_URL_PARAM } from "../buildhealth/BuildHealthDataTypes";
+import { INCLUDE_PREFLIGHT_PARAM, LABELS_URL_PARAM_ALL_LABELS, LABELS_URL_PARAM_NAME, PROJECTS_URL_PARAM_NAME, QUERY_SCHEMA_PARAM_NAME, STATES_URL_PARAM_ALL_STATES, STATES_URL_PARAM_NAME, STEPS_URL_PARAM_ALL_STEPS, STEPS_URL_PARAM_NAME, STREAMS_URL_PARAM_NAME, TEMPLATES_URL_PARAM_NAME, TIME_SPAN_URL_PARAM, UI_DATE_ANCHORS_PARAM, UI_SUMMARY_WARNING_AS_ERROR_PARAM } from "../buildhealth/BuildHealthDataTypes";
 import { BuildHealthDataHandler } from "./BuildHealthDataHandler";
 import { decodeFullyQualifiedLabelName, decodeLabelKey, decodeStateKey, decodeStepKey, decodeTemplateKey, encodeFullyQualifiedLabelName, encodeLabelKey, encodeLabelKeyFromStrings, encodeStateKeyFromString, encodeStepKey, encodeStepNameFromStringsPreserveCase, encodeTemplateKey, encodeTemplateKeyFromStrings } from "./BuildHealthUtilities";
 import { Location } from 'react-router-dom';
@@ -41,6 +41,14 @@ type BuildHealthQueryParams = {
     querySchemaVersion: number;
 };
 
+/**
+ * Data structure used to pass state from query location, to Build Health UI Options.
+ */
+type BuildHealthUIParams = {
+    showDateAnchor: boolean;
+    summaryWarningAsErrors: boolean;
+}
+
 /**
  * Data type used to specify time span from now to consider for jobs.
  */
@@ -86,7 +94,7 @@ export const DEFAULT_JOB_HISTORY_TIME_SPAN: TimeSpan = JobHistoryTimeSpans[1];
  * Obtains a URLSearchParams representation of the current options.
  * @returns The URLSearchParams representation of the current options.
  */
-export function toNavigationQuery(optionsState: BuildHealthOptionsState, querySchemaVersion: number, buildHealthDataHandler: BuildHealthDataHandler): URLSearchParams {
+export function toNavigationQuery(optionsState: BuildHealthOptionsState, uiOptionsState: BuildHealthUIOptionsState, querySchemaVersion: number, buildHealthDataHandler: BuildHealthDataHandler): URLSearchParams {
 
     // Internal helper method to compress keys to the comrpessed representation.
     // e.g. stream-id::template-id::stepname-1,stream-id::template-id::stepname-2 -> stream-id::template-id(stepname-1;stepname-2)
@@ -179,9 +187,13 @@ export function toNavigationQuery(optionsState: BuildHealthOptionsState, querySc
     params.set(STEPS_URL_PARAM_NAME, LZString.compressToEncodedURIComponent(steps));
     params.set(STATES_URL_PARAM_NAME, LZString.compressToEncodedURIComponent(states));
     params.set(TIME_SPAN_URL_PARAM, LZString.compressToEncodedURIComponent(lastTimeRange));
-    params.set(INCLUDE_PREFLIGHT_PARAM, LZString.compressToEncodedURIComponent(String(includePreflight)));
+    params.set(INCLUDE_PREFLIGHT_PARAM, String(includePreflight));
     params.set(QUERY_SCHEMA_PARAM_NAME, querySchemaVersion.toString());
 
+    // Handle UI Option State
+    params.set(UI_DATE_ANCHORS_PARAM, String(uiOptionsState.includeDateAnchors));
+    params.set(UI_SUMMARY_WARNING_AS_ERROR_PARAM, String(uiOptionsState.warningsAsSummaryFailure));
+
     return params;
 }
 
@@ -190,7 +202,7 @@ export function toNavigationQuery(optionsState: BuildHealthOptionsState, querySc
  * @param location Location object.
  * @returns Parsed Build Health Query Params.
  */
-export function parseBuildHealthQueryParams(location: Location): BuildHealthQueryParams {
+export function parseBuildHealthQueryParams(location: Location): { buildHealthQueryParams: BuildHealthQueryParams, buildHealthUIParams: BuildHealthUIParams } {
 
     // Internal helper method to expand keys from the comrpessed representation.
     // e.g. stream-id::template-id(stepname-1;stepname-2) -> stream-id::template-id::stepname-1,stream-id::template-id::stepname-2
@@ -235,6 +247,10 @@ export function parseBuildHealthQueryParams(location: Location): BuildHealthQuer
         return decompressed.split(",").filter(Boolean);
     };
 
+    const getBoolParam = (key: string): boolean => {
+        return searchParams.get(key)?.toLowerCase() === "true";
+    }
+
     const querySchemaVersion = Number(searchParams.get(QUERY_SCHEMA_PARAM_NAME)) || 1;
     projectIds = getList(PROJECTS_URL_PARAM_NAME);
     streamIds = getList(STREAMS_URL_PARAM_NAME);
@@ -299,10 +315,16 @@ export function parseBuildHealthQueryParams(location: Location): BuildHealthQuer
         timeSpanKey = DEFAULT_JOB_HISTORY_TIME_SPAN.key;
     }
 
-    // Read from URLSearchParams
-    const includePreflight = (searchParams.get(INCLUDE_PREFLIGHT_PARAM) ? LZString.decompressFromEncodedURIComponent(searchParams.get(INCLUDE_PREFLIGHT_PARAM)!) : "false") === "true";
+    const includePreflight = getBoolParam(INCLUDE_PREFLIGHT_PARAM);
 
-    return { projectIds, streamIds, templateIds: templateIds, labels, stepNames, stateNames, timeSpanKey, includePreflight, querySchemaVersion };
+    // UI params
+    const uiWarningsAsErrors = getBoolParam(UI_SUMMARY_WARNING_AS_ERROR_PARAM);
+    const uiDateAnchors = getBoolParam(UI_DATE_ANCHORS_PARAM);
+
+    return {
+        buildHealthQueryParams: { projectIds, streamIds, templateIds: templateIds, labels, stepNames, stateNames, timeSpanKey, includePreflight, querySchemaVersion },
+        buildHealthUIParams: { showDateAnchor: uiDateAnchors, summaryWarningAsErrors: uiWarningsAsErrors }
+    };
 }
 
 /**
@@ -357,10 +379,15 @@ export function upgradeBuildHealthQueryParams(params: BuildHealthQueryParams, cu
  */
 export async function loadBuildHealthOptionsFromParams(
     params: BuildHealthQueryParams,
+    uiParams: BuildHealthUIParams,
     options: BuildHealthOptionsController,
     handler: BuildHealthDataHandler) {
     let optionsWriter: IBuildHealthOptionWriter = options.getTransactionSession();
 
+    // Handle UI first
+    options.setWarningsAsSummaryFailure(uiParams.summaryWarningAsErrors);
+    options.setHideDateAnchors(uiParams.showDateAnchor);
+
     runInAction(() => {
         optionsWriter.setIncludePreflight(params.includePreflight);
     })
@@ -485,6 +512,7 @@ export async function loadBuildHealthOptionsFromParams(
 export class BuildHealthUIOptionsState {
     @observable includeDateAnchors: boolean = false;
     @observable warningsAsSummaryFailure: boolean = false;
+    @observable linkedSelectionModeEnabled: boolean = false;
     @observable validStates: Set<string> = new Set<string>();
 
     constructor() {
@@ -625,7 +653,6 @@ export class BuildHealthOptionsState {
     @observable enabledStates: Record<string, string> = {};
     @observable jobHistoryTimeSpan: TimeSpan = DEFAULT_JOB_HISTORY_TIME_SPAN;
     @observable includePreflight: boolean = false;
-    @observable linkedSelectionModeEnabled: boolean = false;
 
     // #endregion -- User Selected Options --
 
diff --git a/Engine/Source/Programs/Horde/HordeDashboard/src/components/buildhealth/BuildHealthOptionsController.ts b/Engine/Source/Programs/Horde/HordeDashboard/src/components/buildhealth/BuildHealthOptionsController.ts
--- a/Engine/Source/Programs/Horde/HordeDashboard/src/components/buildhealth/BuildHealthOptionsController.ts
+++ b/Engine/Source/Programs/Horde/HordeDashboard/src/components/buildhealth/BuildHealthOptionsController.ts
@@ -1,6 +1,6 @@
 import { action, makeObservable } from "mobx";
 import { BuildHealthOptionsState, BuildHealthOptionStateDiff, BuildHealthOptionStateReceipt, BuildHealthUIOptionsState, TimeSpan } from "./BuildHealthOptions";
-import { decodeFullyQualifiedLabelName, decodeLabelKey, decodeStateKey, decodeStepKey, decodeTemplateKey, encodeFullyQualifiedLabelName, encodeLabelKeyFromStrings, encodeStepNameFromStringsPreserveCase, encodeTemplateKeyFromStrings } from "./BuildHealthUtilities";
+import { decodeFullyQualifiedLabelName, decodeLabelKey, decodeStateKey, decodeStepKey, decodeTemplateKey, encodeFullyQualifiedLabelName, encodeLabelKeyFromStrings, encodeStepNameFromStringsPreserveCase, encodeTemplateKey, encodeTemplateKeyFromStrings } from "./BuildHealthUtilities";
 import { BuildHealthDataHandler } from "./BuildHealthDataHandler";
 
 /**
@@ -34,6 +34,36 @@ export interface IHierarchyValidator {
     isValidLabel(streamId: string, templateId: string, fullyQualifiedLabelName: string): boolean;
 }
 
+/**
+ * Data interface to describe a stream.
+ */
+export interface StreamDescriptor {
+    /**
+     * The id of the stream.
+     */
+    id: string;
+
+    /**
+     * The name of the stream.
+     */
+    name: string;
+}
+
+/**
+ * Data interface to describe a template.
+ */
+export interface TemplateDescriptor {
+    /**
+     * The id of the template.
+     */
+    id: string;
+
+    /**
+     * The name of the template.
+     */
+    name: string;
+}
+
 /**
  * Interface that describes all the writeable operations on the build health options.
  */
@@ -56,12 +86,6 @@ export interface IBuildHealthOptionWriter {
      */
     setJobHistoryTimeSpan(time: TimeSpan): void;
 
-    /**
-     * Sets whether to enter into linked selection mode, or not.
-     * @param linkedSelectionModeEnabled True to enter linked selection mode, false otherwise.
-     */
-    setLinkedSelectionMode(linkedSelectionModeEnabled: boolean): void;
-
     /**
      * Toggles the state of the provided project.
      * @param id The id of the project to toggle.
@@ -78,6 +102,13 @@ export interface IBuildHealthOptionWriter {
      */
     toggleSingleProject(id: string, name: string): void;
 
+    /**
+     * Toggles all provided streams.
+     * @param streamDescriptors The list of streams to attempt to toggle. 
+     * @param enabled Whether it's enabled or not.
+     */
+    toggleAllStreams(streamDescriptors: StreamDescriptor[], enabled: boolean): void;
+
     /**
      * Toggles the state of the provided stream.
      * @param id The id of the stream to toggle.
@@ -96,10 +127,18 @@ export interface IBuildHealthOptionWriter {
 
     /**
      * Toggles all linkable templates, for all active streams.
-     * @param enabled True if we want to enable, false otherwise.
+     * @param enabled Whether it's enabled or not.
      */
     toggleAllLinkableTemplates(enabled: boolean): void;
 
+    /**
+     * Toggles all provided templates for a given stream.
+     * @param streamId The stream.
+     * @param templateDescriptors The list of templates to attempt to toggle. 
+     * @param enabled  Whether it's enabled or not.
+     */
+    toggleAllTemplates(streamId: string, templateDescriptors: TemplateDescriptor[], enabled: boolean): void;
+
     /**
      * Toggles the state of the provided template.
      * @param id The id of the template to toggle.
@@ -214,12 +253,14 @@ export interface IBuildHealthOptionWriter {
 
 export class BuildHealthOptionsWriter implements IBuildHealthOptionWriter {
     readonly state: BuildHealthOptionsState;
+    readonly uiState: BuildHealthUIOptionsState;
     readonly hierarchyValidator?: IHierarchyValidator;
     private activeTransactionReceipt: BuildHealthOptionStateReceipt;
 
-    constructor(state: BuildHealthOptionsState, hierarhcyValidator?: IHierarchyValidator) {
+    constructor(state: BuildHealthOptionsState, uiState: BuildHealthUIOptionsState, hierarhcyValidator?: IHierarchyValidator) {
         makeObservable(this);
         this.state = state;
+        this.uiState = uiState;
         this.activeTransactionReceipt = this.state.generateStateReceipt();
         this.hierarchyValidator = hierarhcyValidator;
     }
@@ -265,6 +306,17 @@ export class BuildHealthOptionsWriter implements IBuildHealthOptionWriter {
         }
     }
 
+    /**
+     * @inheritdoc
+     */
+    @action
+    toggleAllStreams(streamDescriptors: StreamDescriptor[], enabled: boolean): void {
+        for (let idx: number = 0; idx < streamDescriptors.length; ++idx) {
+            let streamDescriptor: StreamDescriptor = streamDescriptors[idx];
+            this.toggleStream(streamDescriptor.id, streamDescriptor.name, enabled);
+        }
+    }
+
     /**
      * @inheritdoc
      */
@@ -287,13 +339,29 @@ export class BuildHealthOptionsWriter implements IBuildHealthOptionWriter {
         }
     }
 
+    /**
+     * @inheritdoc
+     */
+    @action
+    toggleAllTemplates(streamId: string, templateDescriptors: StreamDescriptor[], enabled: boolean, cascadeLinkedAction?: boolean): void {
+        for (let idx: number = 0; idx < templateDescriptors.length; ++idx) {
+            let templateDescriptor: StreamDescriptor = templateDescriptors[idx];
+            let encodedTemplateKey = encodeTemplateKeyFromStrings(streamId, templateDescriptor.id);
+            if (cascadeLinkedAction) {
+                this.toggleTemplate(encodedTemplateKey, templateDescriptor.name, enabled);
+            } else {
+                this.toggleTemplateInternal(encodedTemplateKey, templateDescriptor.name, enabled);
+            }
+        }
+    }
+
     /**
      * @inheritdoc
      */
     @action
     toggleTemplate(id: string, name: string, enabled: boolean) {
         this.toggleTemplateInternal(id, name, enabled);
-        if (this.state.linkedSelectionModeEnabled) {
+        if (this.uiState.linkedSelectionModeEnabled) {
             this.toggleLinkedTemplates(id, name, enabled);
         }
     }
@@ -303,7 +371,7 @@ export class BuildHealthOptionsWriter implements IBuildHealthOptionWriter {
      */
     @action
     toggleAllLinkableTemplates(enabled: boolean) {
-        if (this.state.linkedSelectionModeEnabled) {
+        if (this.uiState.linkedSelectionModeEnabled) {
             Object.keys(this.state.enabledTemplates).forEach(templateKey => {
                 let templateName = this.state.enabledTemplates[templateKey];
                 this.toggleLinkedTemplates(templateKey, templateName, enabled);
@@ -316,7 +384,7 @@ export class BuildHealthOptionsWriter implements IBuildHealthOptionWriter {
      */
     @action
     toggleAllLinkableLabels(enabled: boolean) {
-        if (this.state.linkedSelectionModeEnabled) {
+        if (this.uiState.linkedSelectionModeEnabled) {
             Object.keys(this.state.enabledLabels).forEach(labelKey => {
                 let fullyQualifiedLabelName = this.state.enabledLabels[labelKey];
                 this.toggleLinkedLabels(labelKey, fullyQualifiedLabelName, enabled);
@@ -329,7 +397,7 @@ export class BuildHealthOptionsWriter implements IBuildHealthOptionWriter {
      */
     @action
     toggleAllLinkableSteps(enabled: boolean) {
-        if (this.state.linkedSelectionModeEnabled) {
+        if (this.uiState.linkedSelectionModeEnabled) {
             Object.keys(this.state.enabledSteps).forEach(stepKey => {
                 let stepName = this.state.enabledSteps[stepKey];
                 this.toggleLinkedSteps(stepKey, stepName, enabled);
@@ -343,7 +411,7 @@ export class BuildHealthOptionsWriter implements IBuildHealthOptionWriter {
     @action
     toggleLabel(id: string, name: string, enabled: boolean) {
         this.toggleLabelInternal(id, name, enabled);
-        if (this.state.linkedSelectionModeEnabled) {
+        if (this.uiState.linkedSelectionModeEnabled) {
             this.toggleLinkedLabels(id, name, enabled);
         }
     }
@@ -354,7 +422,7 @@ export class BuildHealthOptionsWriter implements IBuildHealthOptionWriter {
     @action
     toggleStep(id: string, name: string, enabled: boolean) {
         this.toggleStepInternal(id, name, enabled);
-        if (this.state.linkedSelectionModeEnabled) {
+        if (this.uiState.linkedSelectionModeEnabled) {
             this.toggleLinkedSteps(id, name, enabled);
         }
     }
@@ -427,14 +495,6 @@ export class BuildHealthOptionsWriter implements IBuildHealthOptionWriter {
         this.state.includePreflight = includePreflight;
     }
 
-    /**
-     * @inheritdoc
-     */
-    @action
-    setLinkedSelectionMode(linkedSelectionModeEnabled: boolean) {
-        this.state.linkedSelectionModeEnabled = linkedSelectionModeEnabled;
-    }
-
     /**
      * @inheritdoc
      */
@@ -602,6 +662,7 @@ export class BuildHealthOptionsController {
     // #region -- Public Members --
 
     optionsChangeVersion = 0;
+    uiOptionsChangeVersion = 0;
     querySchemaVersion: number = 4;
 
     // #endregion -- Public Members --
@@ -618,6 +679,10 @@ export class BuildHealthOptionsController {
         this.optionsChangeVersion++;
     }
 
+    private setUIOptionsChanged() {
+        this.uiOptionsChangeVersion++;
+    }
+
     // #endregion -- Constructor --
 
     // #region -- Public API --
@@ -630,13 +695,24 @@ export class BuildHealthOptionsController {
         this.hierarchyValidator = validator;
     }
 
+    /**
+     * Sets whether the linked selection mode is on when editing.
+     * @param linkedSelectionModeEnabled True to enable linked selection, false otherwise.
+     */
+    @action
+    setLinkedSelectionMode(linkedSelectionModeEnabled: boolean) {
+        this.uiState.linkedSelectionModeEnabled = linkedSelectionModeEnabled;
+        this.setUIOptionsChanged();
+    }
+
     /**
      * Sets whether the date anchors should be used when visualizing the table.
-     * @param includeDateAnchors 
+     * @param includeDateAnchors True to show date anchors, false otherwise.
      */
     @action
     setHideDateAnchors(includeDateAnchors: boolean) {
         this.uiState.includeDateAnchors = includeDateAnchors;
+        this.setUIOptionsChanged();
     }
 
     /**
@@ -646,6 +722,7 @@ export class BuildHealthOptionsController {
     @action
     setWarningsAsSummaryFailure(warningsAsSummaryFailure: boolean) {
         this.uiState.warningsAsSummaryFailure = warningsAsSummaryFailure;
+        this.setUIOptionsChanged();
     }
 
     /**
@@ -655,7 +732,7 @@ export class BuildHealthOptionsController {
     @action
     getTransactionSession(): IBuildHealthOptionWriter {
         if (!this.writer) {
-            this.writer = new BuildHealthOptionsWriter(this.state, this.hierarchyValidator);
+            this.writer = new BuildHealthOptionsWriter(this.state, this.uiState, this.hierarchyValidator);
         }
 
         return this.writer;
diff --git a/Engine/Source/Programs/Horde/HordeDashboard/src/components/buildhealth/BuildhealthView.tsx b/Engine/Source/Programs/Horde/HordeDashboard/src/components/buildhealth/BuildhealthView.tsx
--- a/Engine/Source/Programs/Horde/HordeDashboard/src/components/buildhealth/BuildhealthView.tsx
+++ b/Engine/Source/Programs/Horde/HordeDashboard/src/components/buildhealth/BuildhealthView.tsx
@@ -3,7 +3,7 @@
 import { Checkbox, ComboBox, DefaultButton, DirectionalHint, Dropdown, DropdownMenuItemType, IconButton, IDropdownOption, IDropdownProps, ISelectableOption, Label, mergeStyleSets, PrimaryButton, SearchBox, SelectableOptionMenuItemType, Stack, TooltipHost } from "@fluentui/react";
 import { GetProjectResponse, GetStreamResponse } from "horde/backend/Api";
 import { observer } from "mobx-react-lite";
-import { useCallback, useEffect, useMemo, useRef, useState } from "react";
+import { useCallback, useEffect, useLayoutEffect, useMemo, useRef, useState } from "react";
 import { useLocation, useNavigate } from "react-router";
 import { useWindowSize } from "../../base/utilities/hooks";
 import { getHordeStyling } from "../../styles/Styles";
@@ -16,7 +16,7 @@ import { ALL_STATES_TUPLE_LIST, BuildHealthOptionsState, BuildHealthOptionStateD
 import { decodeStepKey, decodeTemplateKey, encodeFullyQualifiedLabelName, encodeLabelKey, encodeStepKey, encodeTemplateKey } from "./BuildHealthUtilities";
 import { stepOutcomeTableClasses } from "./stepoutcome/StepOutcomeSharedUIComponents";
 import { runInAction } from "mobx";
-import { BuildHealthOptionsController } from "./BuildHealthOptionsController";
+import { BuildHealthOptionsController, StreamDescriptor, TemplateDescriptor } from "./BuildHealthOptionsController";
 import { StepOutcomeLegend } from "./stepoutcome/StepOutcomeLegend";
 import { StepOutcomeDataHandler } from "./stepoutcome/StepOutcomeDataHandler";
 
@@ -163,7 +163,8 @@ const SearchableDropdown: React.FC<ISearchableDropdownProps> = (props) => {
             ref={dropdownRef}
             {...dropdownProps}
             calloutProps={{
-                directionalHint: DirectionalHint.rightCenter,
+                directionalHint: DirectionalHint.rightBottomEdge,
+                alignTargetEdge: true,
                 styles: { calloutMain: { width: 500, minWidth: 500 } },
             }}
             onRenderList={(props, defaultRender) => {
@@ -249,7 +250,8 @@ const BasicListParameter: React.FC<{ handler: BuildHealthDataHandler, buildHealt
         dropdownWidth={200}
         styles={{ dropdown: { width: 300 } }}
         calloutProps={{
-            directionalHint: DirectionalHint.rightCenter
+            directionalHint: DirectionalHint.rightBottomEdge,
+            alignTargetEdge: true,
         }}
         onChange={(_ev, option) => {
             let castedOption = option as IDropdownOption<BuildHealthOptionData>;
@@ -274,21 +276,11 @@ const MultiListParameter: React.FC<{ handler: BuildHealthDataHandler, buildHealt
     const [filterText, setFilterText] = useState("");
     const changeVersion = buildHealthOptions.optionsChangeVersion;
     const key = `${PARAMETER_KEY_PREFIX}${KEY_SEPARATOR}${param.label}`;
-    const gset: Set<string> = new Set();
     const selectedSet = useMemo(() => new Set(selectedKeys), [selectedKeys]);
     const groupRefs = useRef<(HTMLDivElement | null)[]>([]);
 
     useEffect(() => { groupRefs.current = []; }, [param.items]);
 
-    param.items.forEach(item => {
-        if (!item.group) {
-            item.group = "__nogroup";
-        }
-        if (item.group) {
-            gset.add(item.group);
-        }
-    });
-
     // #region -- Data Memos --
 
     const doptions = useMemo(() => {
@@ -341,7 +333,7 @@ const MultiListParameter: React.FC<{ handler: BuildHealthDataHandler, buildHealt
 
     const filteredOptionsRef = useRef<IDropdownOption<BuildHealthOptionData>[]>([]);
 
-    useEffect(() => {
+    useLayoutEffect(() => {
         filteredOptionsRef.current = filteredOptions;
     }, [filteredOptions]);
 
@@ -357,7 +349,7 @@ const MultiListParameter: React.FC<{ handler: BuildHealthDataHandler, buildHealt
             if (option.itemType === DropdownMenuItemType.Header && option.text === SELECT_CLEAR_ALL_SPECIAL_GROUP) {
                 let fitleredOptionList = filteredOptionsRef.current.length == doptions.length ? undefined : filteredOptionsRef.current.filter(x => x.itemType == undefined || x.itemType != SelectableOptionMenuItemType.Header);
                 let selectAllLabel = (fitleredOptionList !== undefined && fitleredOptionList.length > 0) ? "Select Filtered" : "Select All";
-                let showButton = fitleredOptionList == undefined || fitleredOptionList.length;
+                let showButton = fitleredOptionList == undefined || fitleredOptionList.length > 0;
                 return (
                     <Stack horizontal tokens={{ childrenGap: 8 }} styles={{
                         root:
@@ -507,18 +499,52 @@ const StreamDropdown: React.FC<{ handler: BuildHealthDataHandler, buildHealthOpt
         return listParam;
     }).sort((a, b) => a.text.localeCompare(b.text));
 
+    const handleSelectAll = useCallback(
+        (
+            option: IDropdownOption<BuildHealthOptionData>,
+            optionSelectAllContext?: IDropdownOption<BuildHealthOptionData>[]
+        ) => {
+            const key = option.key as string;
+            const [, projectId] = key.split(KEY_SEPARATOR);
+            const isSelected = !!option.selected;
+
+            // If there is no provided context, simply obtain all the known streams of the (project).
+            const streamIds: StreamDescriptor[] =
+                optionSelectAllContext === undefined
+                    ? handler.getCachedStreams(projectId)
+                    : optionSelectAllContext
+                        .filter((opt) => {
+                            if (!opt.data) return false;
+
+                            return (
+                                opt.data.group === projectId
+                            );
+                        })
+                        .map((x) => ({ id: x.data!.id, name: x.data!.text! }));
+
+            // Do not close the transaction session; the user still has the modal open
+            buildHealthOptions
+                .getTransactionSession()
+                .toggleAllStreams(streamIds, isSelected);
+        },
+        [buildHealthOptions, handler]
+    );
+
     let selectedStreamKeys: string[] = Object.keys(buildHealthOptions.state.enabledStreams);
     let request: DataHandlerRefreshRequest = expandRefreshRequest({ streams: true });
 
     return streamDropdownParams.items.length > 0 && (
         <Stack>
             <Label>Stream(s)</Label>
-            <MultiListParameter handler={handler} buildHealthOptions={buildHealthOptions} param={streamDropdownParams} selectedKeys={selectedStreamKeys} request={request} onChange={(option: IDropdownOption<BuildHealthOptionData>) => {
-                const key = option.key as string;
-                const isSelected = !!option.selected;
-
-                buildHealthOptions.getTransactionSession().toggleStream(key, option.text, isSelected);
-            }} />
+            <MultiListParameter handler={handler} buildHealthOptions={buildHealthOptions} param={streamDropdownParams} selectedKeys={selectedStreamKeys} request={request}
+                enabledSelectAll={true}
+                onSelectAll={handleSelectAll}
+                onChange={(option: IDropdownOption<BuildHealthOptionData>) => {
+                    const key = option.key as string;
+                    const isSelected = !!option.selected;
+
+                    buildHealthOptions.getTransactionSession().toggleStream(key, option.text, isSelected);
+                }} />
         </Stack>
     );
 });
@@ -540,6 +566,40 @@ const TemplateDropdown: React.FC<{ handler: BuildHealthDataHandler, buildHealthO
         return listParam;
     }).sort((a, b) => a.text.localeCompare(b.text));
 
+    // Note: We are keeping our select all logic here, but not currently using it. 
+    //       We would like to be prepared for the use case, however, a select all on templates will be excessive load on the server.
+    //       This should be revisited after UE-349992
+    const handleSelectAll = useCallback(
+        (
+            option: IDropdownOption<BuildHealthOptionData>,
+            optionSelectAllContext?: IDropdownOption<BuildHealthOptionData>[]
+        ) => {
+            const key = option.key as string;
+            const [, streamId] = key.split(KEY_SEPARATOR);
+            const isSelected = !!option.selected;
+
+            // If there is no provided context, simply obtain all the known templates of the (stream).
+            const templateIds: TemplateDescriptor[] =
+                optionSelectAllContext === undefined
+                    ? handler.getCachedTemplates(streamId)
+                    : optionSelectAllContext
+                        .filter((opt) => {
+                            if (!opt.data) return false;
+                            const decoded = decodeTemplateKey(opt.data.id!);
+                            return (
+                                decoded.streamId === streamId
+                            );
+                        })
+                        .map((x) => ({ id: x.data!.id, name: x.data!.text! }));
+
+            // Do not close the transaction session; the user still has the modal open
+            buildHealthOptions
+                .getTransactionSession()
+                .toggleAllTemplates(streamId, templateIds, isSelected);
+        },
+        [buildHealthOptions, handler]
+    );
+
     let selectedTemplateKeys: string[] = Object.keys(buildHealthOptions.state.enabledTemplates);
     let request: DataHandlerRefreshRequest = expandRefreshRequest({ jobs: true });
 
@@ -552,21 +612,9 @@ const TemplateDropdown: React.FC<{ handler: BuildHealthDataHandler, buildHealthO
                         const key = option.key as string;
                         const isSelected = !!option.selected;
                         buildHealthOptions.getTransactionSession().toggleTemplate(key, option.text, isSelected);
-                    }}
-                    onHierarchicalRefreshCompleted={(receipts?: BuildHealthOptionStateDiff) => {
-                        // If we have added templates, issue another request to toggle all steps & close the transaction.
-                        // We do not want to propagate linked actions in this special case, as it can be quite destructive.
-                        if (receipts !== null && receipts!.added && receipts!.added.enabledTemplates.size > 0) {
-                            for (const key of receipts!.added.enabledTemplates!) {
-                                const templateData = decodeTemplateKey(key);
-                                let stepNames = handler.getCachedTemplateStepNames(templateData.streamId, templateData.templateId);
-                                buildHealthOptions.getTransactionSession().toggleAllSteps(templateData.streamId, templateData.templateId, stepNames, true, false);
-                                buildHealthOptions.commitTransactionSession();
-                            }
-                        }
                     }} />
                 {
-                    buildHealthOptions.state.linkedSelectionModeEnabled &&
+                    buildHealthOptions.uiState.linkedSelectionModeEnabled &&
                     (
                         <TooltipHost
                             content={"Link templates: Apply selected templates across selected streams."}
@@ -605,7 +653,7 @@ const StepDropdown: React.FC<{ handler: BuildHealthDataHandler, buildHealthOptio
 
             return transformedStepData;
         },
-        [handler.stepData, handler.labelData]
+        [handler.stepData]
     );
 
     // #endregion -- Data Memos --
@@ -686,7 +734,7 @@ const StepDropdown: React.FC<{ handler: BuildHealthDataHandler, buildHealthOptio
                     enabledSelectAll={true}
                     onSelectAll={handleSelectAll} />
                 {
-                    buildHealthOptions.state.linkedSelectionModeEnabled &&
+                    buildHealthOptions.uiState.linkedSelectionModeEnabled &&
                     (
                         <TooltipHost
                             content={"Link steps: Apply selected steps across selected templates."}
@@ -801,7 +849,7 @@ const LabelDropdown: React.FC<{ handler: BuildHealthDataHandler, buildHealthOpti
                     enabledSelectAll={true}
                     onSelectAll={handleSelectAll} />
                 {
-                    buildHealthOptions.state.linkedSelectionModeEnabled &&
+                    buildHealthOptions.uiState.linkedSelectionModeEnabled &&
                     (
                         <TooltipHost
                             content={"Link labels: Apply selected labels across selected templates."}
@@ -923,19 +971,19 @@ const BuildHealthUrlSync: React.FC<{ buildHealthController: BuildHealthOptionsCo
     useEffect(() => {
         const initializeFromUrl = async () => {
             const params = parseBuildHealthQueryParams(location);
-            const upgradedParams = upgradeBuildHealthQueryParams(params, buildHealthController.querySchemaVersion);
-            await loadBuildHealthOptionsFromParams(upgradedParams, buildHealthController, handler);
+            const upgradedParams = upgradeBuildHealthQueryParams(params.buildHealthQueryParams, buildHealthController.querySchemaVersion);
+            await loadBuildHealthOptionsFromParams(upgradedParams, params.buildHealthUIParams, buildHealthController, handler);
         };
 
         initializeFromUrl();
     }, []);
 
     useEffect(() => {
-        const desiredSearch = toNavigationQuery(buildHealthController.state, buildHealthController.querySchemaVersion, handler).toString();
+        const desiredSearch = toNavigationQuery(buildHealthController.state, buildHealthController.uiState, buildHealthController.querySchemaVersion, handler).toString();
         if (desiredSearch && location.search !== `?${desiredSearch}`) {
             navigate(`${location.pathname}?${desiredSearch}`, { replace: true });
         }
-    }, [navigate, buildHealthController.optionsChangeVersion, location.pathname, location.search]);
+    }, [navigate, buildHealthController.optionsChangeVersion, buildHealthController.uiOptionsChangeVersion, location.pathname, location.search]);
 
     return null;
 };
@@ -948,7 +996,7 @@ const BuildHealthSidebar: React.FC<{ handler: BuildHealthDataHandler, buildHealt
     const { hordeClasses, modeColors } = getHordeStyling();
     const [timeSpanKey, setTimeSpanKey] = useState(() => buildHealthOptions.state.jobHistoryTimeSpan.key);
     const [includePreflight, setIncludePreflight] = useState(() => buildHealthOptions.state.includePreflight)
-    const [linkedSelectionMode, setLinkedSelectionMode] = useState(() => buildHealthOptions.state.linkedSelectionModeEnabled)
+    const [linkedSelectionMode, setLinkedSelectionMode] = useState(() => buildHealthOptions.uiState.linkedSelectionModeEnabled)
 
     // #region -- Visual Options -- 
 
@@ -963,16 +1011,18 @@ const BuildHealthSidebar: React.FC<{ handler: BuildHealthDataHandler, buildHealt
         const select = option as TimeSpan;
         setTimeSpanKey(select.key);
         buildHealthOptions.getTransactionSession().setJobHistoryTimeSpan(select);
+        buildHealthOptions.commitTransactionSession();
     };
 
     const updateIncludePreflight = (isChecked) => {
         setIncludePreflight(isChecked);
         buildHealthOptions.getTransactionSession().setIncludePreflight(isChecked);
+        buildHealthOptions.commitTransactionSession();
     }
 
     const updateLinkedSelectionMode = (isChecked) => {
         setLinkedSelectionMode(isChecked);
-        buildHealthOptions.getTransactionSession().setLinkedSelectionMode(isChecked);
+        buildHealthOptions.setLinkedSelectionMode(isChecked);
     }
 
     const updateIncludeDateAnchors = (isChecked) => {
diff --git a/Engine/Source/Programs/Horde/HordeDashboard/src/components/buildhealth/stepoutcome/StepOutcomeDataTypes.ts b/Engine/Source/Programs/Horde/HordeDashboard/src/components/buildhealth/stepoutcome/StepOutcomeDataTypes.ts
--- a/Engine/Source/Programs/Horde/HordeDashboard/src/components/buildhealth/stepoutcome/StepOutcomeDataTypes.ts
+++ b/Engine/Source/Programs/Horde/HordeDashboard/src/components/buildhealth/stepoutcome/StepOutcomeDataTypes.ts
@@ -1024,6 +1024,32 @@ export class StepOutcomeTable {
         this.supportsDateAnchoredChanges = viewModelOptions.includeDateAnchors;
     }
 
+    /**
+     * The number of step results stored in the table.
+     */
+    get StepResultCount(): number {
+        let stepCount: number = 0;
+
+        for (const [_streamId, metadata] of this.tableStreamMetadata.entries()) {
+            stepCount += metadata.stepCount;
+        }
+
+        return stepCount;
+    }
+
+    /**
+     * The number of label results stored in the table.
+     */
+    get LabelResultCount(): number {
+        let labelCount: number = 0;
+
+        for (const [_streamId, metadata] of this.tableStreamMetadata.entries()) {
+            labelCount += metadata.labelCount;
+        }
+
+        return labelCount;
+    }
+
     // #endregion -- Public API --
 }
 
diff --git a/Engine/Source/Programs/Horde/HordeDashboard/src/components/buildhealth/stepoutcome/StepOutcomeLegend.tsx b/Engine/Source/Programs/Horde/HordeDashboard/src/components/buildhealth/stepoutcome/StepOutcomeLegend.tsx
--- a/Engine/Source/Programs/Horde/HordeDashboard/src/components/buildhealth/stepoutcome/StepOutcomeLegend.tsx
+++ b/Engine/Source/Programs/Horde/HordeDashboard/src/components/buildhealth/stepoutcome/StepOutcomeLegend.tsx
@@ -168,7 +168,7 @@ export const StepOutcomeLegend: React.FC = () => {
                                 width: 20,
                                 height: 20,
                                 backgroundColor: getCellStyle(fakeSuccessMultiRunEntry).bgColor,
-                            }}><TooltipHost content="The step & change has multiple entires." ><div style={{ width: "100%", height: "100%" }} /></TooltipHost></td>
+                            }}><TooltipHost content="The step & change has multiple entires. Most recent result is always shown." ><div style={{ width: "100%", height: "100%" }} /></TooltipHost></td>
                         </tr>
                     </tbody>
                 </table>
diff --git a/Engine/Source/Programs/Horde/HordeDashboard/tests/BuildHealthOptions.test.ts b/Engine/Source/Programs/Horde/HordeDashboard/tests/BuildHealthOptions.test.ts
--- a/Engine/Source/Programs/Horde/HordeDashboard/tests/BuildHealthOptions.test.ts
+++ b/Engine/Source/Programs/Horde/HordeDashboard/tests/BuildHealthOptions.test.ts
@@ -14,7 +14,7 @@ describe('parseBuildHealthQueryParams', () => {
 		const compressedTemplates = LZString.compressToEncodedURIComponent('stream1::template1,stream2::template2');
 		const compressedSteps = LZString.compressToEncodedURIComponent('stream1::template1::step1,stream2::template2::step2');
 		const compressedTimeSpan = LZString.compressToEncodedURIComponent('time_1_week');
-		const compressedPreflight = LZString.compressToEncodedURIComponent('true');
+		const compressedPreflight = 'true';
 
 		// Do *NOT* use the param const strings; we want to catch those contract failures as well.
 		mockParams.set('projects', compressedProjects);
@@ -31,13 +31,13 @@ describe('parseBuildHealthQueryParams', () => {
 		const result = parseBuildHealthQueryParams(location);
 
 		// Assert
-		expect(result.projectIds).toEqual(['projA', 'projB']);
-		expect(result.streamIds).toEqual(['stream1', 'stream2']);
-		expect(result.templateIds).toEqual(['stream1::template1', 'stream2::template2']);
-		expect(result.stepNames).toEqual(['stream1::template1::step1', 'stream2::template2::step2']);
-		expect(result.timeSpanKey).toBe('time_1_week');
-		expect(result.includePreflight).toBe(true);
-		expect(result.querySchemaVersion).toBe(1);
+		expect(result.buildHealthQueryParams.projectIds).toEqual(['projA', 'projB']);
+		expect(result.buildHealthQueryParams.streamIds).toEqual(['stream1', 'stream2']);
+		expect(result.buildHealthQueryParams.templateIds).toEqual(['stream1::template1', 'stream2::template2']);
+		expect(result.buildHealthQueryParams.stepNames).toEqual(['stream1::template1::step1', 'stream2::template2::step2']);
+		expect(result.buildHealthQueryParams.timeSpanKey).toBe('time_1_week');
+		expect(result.buildHealthQueryParams.includePreflight).toBe(true);
+		expect(result.buildHealthQueryParams.querySchemaVersion).toBe(1);
 	});
 });
 
@@ -52,7 +52,7 @@ describe('upgradeBuildHealthQueryParams', () => {
 		const compressedTemplates = LZString.compressToEncodedURIComponent('stream1::template1,stream2::template2,stream2::template3');
 		const compressedSteps = LZString.compressToEncodedURIComponent('stream1::template1::step1,stream2::template2::step2');
 		const compressedTimeSpan = LZString.compressToEncodedURIComponent('time_1_week');
-		const compressedPreflight = LZString.compressToEncodedURIComponent('true');
+		const compressedPreflight = 'true';
 
 		// Do *NOT* use the param const strings; we want to catch those contract failures as well.
 		mockParams.set('projects', compressedProjects);
@@ -67,7 +67,7 @@ describe('upgradeBuildHealthQueryParams', () => {
 
 		// Act
 		const result = parseBuildHealthQueryParams(location);
-		const upgradedParams = upgradeBuildHealthQueryParams(result, 2);
+		const upgradedParams = upgradeBuildHealthQueryParams(result.buildHealthQueryParams, 2);
 
 		// Assert
 		expect(upgradedParams.projectIds).toEqual(['projA', 'projB']);
diff --git a/Engine/Source/Programs/Horde/HordeDashboard/tests/StepOutcomeDataTypes.test.ts b/Engine/Source/Programs/Horde/HordeDashboard/tests/StepOutcomeDataTypes.test.ts
--- a/Engine/Source/Programs/Horde/HordeDashboard/tests/StepOutcomeDataTypes.test.ts
+++ b/Engine/Source/Programs/Horde/HordeDashboard/tests/StepOutcomeDataTypes.test.ts
@@ -1,24 +1,20 @@
-import { describe, it, expect } from 'vitest';
+import './helpers/FoundationMocks.test.setup';
+import { describe, it, expect} from 'vitest';
 import { encodeStepName } from '../src/components/buildhealth/stepoutcome/StepOutcomeDataTypes';
 
-
 // Mock object type
 type StepOutcomeTableEntry = {
     streamId: string;
-    stepResponse: {
-        name: string;
-    };
+    name: string;
 };
 
-const KEY_SEPARATOR = "_";
-
 describe('encodeStepName', () => {
     it('Should encode the step name correctly', () => {
-        const mockEntry: StepOutcomeTableEntry = {
+        const mockEntry : StepOutcomeTableEntry = {
             streamId: 'STREAM1',
-            stepResponse: { name: 'StepA' }
+            name: 'StepA',
         };
-        
+
         const result = encodeStepName(mockEntry as any);
 
         expect(result).toBe('stream1::stepa'); // streamId + KEY_SEPARATOR + stepResponse.name, lowercased
diff --git a/Engine/Source/Programs/Horde/HordeDashboard/tests/helpers/FoundationMocks.test.setup.ts b/Engine/Source/Programs/Horde/HordeDashboard/tests/helpers/FoundationMocks.test.setup.ts
new file mode 100644
--- /dev/null
+++ b/Engine/Source/Programs/Horde/HordeDashboard/tests/helpers/FoundationMocks.test.setup.ts
@@ -0,0 +1,49 @@
+import { vi } from 'vitest';
+
+// Mock UserInactivity BEFORE anything imports it
+vi.mock('../src/backend/UserInactivity', () => {
+    return {
+        UserInactivity: class {
+            start = vi.fn();
+            stop = vi.fn();
+            eventHandler = vi.fn();
+        },
+    };
+});
+
+// Mock Dashboard BEFORE importing anything that uses it
+vi.mock('../src/backend/Dashboard', () => {
+    return {
+        default: {
+            get darktheme() { return false; },
+            set darktheme(value: boolean) { },
+            getStatusColors: vi.fn(() => new Map([
+                ["Success", "#00FF00"],
+                ["Warning", "#FFFF00"],
+                ["Failure", "#FF0000"],
+                ["Skipped", "#AAAAAA"],
+                ["Running", "#0000FF"],
+            ])),
+        },
+        StatusColor: {
+            Success: "Success",
+            Warning: "Warning",
+            Failure: "Failure",
+            Skipped: "Skipped",
+            Running: "Running",
+        },
+    };
+});
+
+// Mock global browser APIs immediately
+globalThis.window = {
+    addEventListener: vi.fn(),
+    removeEventListener: vi.fn(),
+} as any;
+
+globalThis.localStorage = {
+    getItem: vi.fn(),
+    setItem: vi.fn(),
+    removeItem: vi.fn(),
+    clear: vi.fn(),
+} as any;
\ No newline at end of file
