Horde: Move EnableJobTaskConformCheck from BuildServerConfig to BuildConfig

diff --git a/Engine/Source/Programs/Horde/Plugins/Build/HordeServer.Build.Tests/BuildTestSetup.cs b/Engine/Source/Programs/Horde/Plugins/Build/HordeServer.Build.Tests/BuildTestSetup.cs
--- a/Engine/Source/Programs/Horde/Plugins/Build/HordeServer.Build.Tests/BuildTestSetup.cs
+++ b/Engine/Source/Programs/Horde/Plugins/Build/HordeServer.Build.Tests/BuildTestSetup.cs
@@ -123,7 +123,6 @@ namespace HordeServer.Tests
 
 			IConfiguration config = new ConfigurationBuilder().Build();
 			services.Configure<ServerSettings>(ConfigureSettings);
-			services.Configure<BuildServerConfig>(x => x.EnableJobTaskConformCheck = true);
 			services.AddSingleton<IConfiguration>(config);
 
 			services.AddHttpClient<RpcService>();
diff --git a/Engine/Source/Programs/Horde/Plugins/Build/HordeServer.Build.Tests/Jobs/JobTaskSourceTests.cs b/Engine/Source/Programs/Horde/Plugins/Build/HordeServer.Build.Tests/Jobs/JobTaskSourceTests.cs
--- a/Engine/Source/Programs/Horde/Plugins/Build/HordeServer.Build.Tests/Jobs/JobTaskSourceTests.cs
+++ b/Engine/Source/Programs/Horde/Plugins/Build/HordeServer.Build.Tests/Jobs/JobTaskSourceTests.cs
@@ -183,7 +183,6 @@ namespace HordeServer.Tests.Jobs
 		public async Task AgentAssignLeaseWithConformCheckAsync(long diskFreeSpaceInMb, bool expectingConform)
 		{
 			// TODO Make workspace conform size an input parameter rather than be hardcoded in Fixture.PopulateAsync
-			// TODO Make setting EnableJobTaskConformCheck part of the test rather than set feature flag globally for all in BuildTestSetup.cs
 
 			Fixture fixture = await CreateFixtureAsync();
 			IPool pool = await CreatePoolAsync(new PoolConfig { Name = Fixture.PoolName, EnableAutoscaling = false, MinAgents = 0, NumReserveAgents = 0 });
@@ -191,6 +190,11 @@ namespace HordeServer.Tests.Jobs
 			IAgent agent = fixture.Agent1;
 			agent = Deref(await agent.TryUpdateAsync(new UpdateAgentOptions { Enabled = true, ExplicitPools = [pool.Id] }));
 
+			await UpdateConfigAsync(config =>
+			{
+				config.Plugins.GetBuildConfig().EnableJobTaskConformCheck = true;
+			});
+
 			List<string> properties = [$"{KnownPropertyNames.DiskFreeSpace}={diskFreeSpaceInMb * 1024 * 1024}"];
 			agent = Deref(await AgentService.CreateSessionAsync(agent, new RpcAgentCapabilities(properties), null));
 
diff --git a/Engine/Source/Programs/Horde/Plugins/Build/HordeServer.Build/BuildConfig.cs b/Engine/Source/Programs/Horde/Plugins/Build/HordeServer.Build/BuildConfig.cs
--- a/Engine/Source/Programs/Horde/Plugins/Build/HordeServer.Build/BuildConfig.cs
+++ b/Engine/Source/Programs/Horde/Plugins/Build/HordeServer.Build/BuildConfig.cs
@@ -85,6 +85,11 @@ namespace HordeServer
 		/// </summary>
 		public bool EnableConformTasks { get; set; } = true;
 
+		/// <summary>
+		/// Whether to check if a conform is required on an agent before assigning a lease for a job task source.
+		/// </summary>
+		public bool EnableJobTaskConformCheck { get; set; }
+
 		/// <summary>
 		/// Commit tag to use for marking issues as fixed
 		/// </summary>
diff --git a/Engine/Source/Programs/Horde/Plugins/Build/HordeServer.Build/BuildServerConfig.cs b/Engine/Source/Programs/Horde/Plugins/Build/HordeServer.Build/BuildServerConfig.cs
--- a/Engine/Source/Programs/Horde/Plugins/Build/HordeServer.Build/BuildServerConfig.cs
+++ b/Engine/Source/Programs/Horde/Plugins/Build/HordeServer.Build/BuildServerConfig.cs
@@ -35,11 +35,6 @@ namespace HordeServer
 		/// </summary>
 		public bool EnableConformTasks { get; set; } = true;
 
-		/// <summary>
-		/// Whether to check if a conform is required on an agent before assigning a lease for a job task source.
-		/// </summary>
-		public bool EnableJobTaskConformCheck { get; set; }
-
 		/// <summary>
 		/// Url of P4 Swarm installation
 		/// </summary>
diff --git a/Engine/Source/Programs/Horde/Plugins/Build/HordeServer.Build/Jobs/JobTaskSource.cs b/Engine/Source/Programs/Horde/Plugins/Build/HordeServer.Build/Jobs/JobTaskSource.cs
--- a/Engine/Source/Programs/Horde/Plugins/Build/HordeServer.Build/Jobs/JobTaskSource.cs
+++ b/Engine/Source/Programs/Horde/Plugins/Build/HordeServer.Build/Jobs/JobTaskSource.cs
@@ -767,7 +767,7 @@ namespace HordeServer.Jobs
 
 		private async Task<bool> AutoConformAgentAsync(IAgent agent, CancellationToken cancellationToken)
 		{
-			if (!_settings.CurrentValue.EnableJobTaskConformCheck)
+			if (!_buildConfig.CurrentValue.EnableJobTaskConformCheck)
 			{
 				return false;
 			}
