Introduce a transcient storage for keep ad-hoc transcient properties in memory while part of the UI has been destroy
Use it for Summary card expand and stream compare

#rnx

diff --git a/Engine/Source/Programs/Horde/HordeDashboard/plugins/test-automation/testAutomationPhaseView.tsx b/Engine/Source/Programs/Horde/HordeDashboard/plugins/test-automation/testAutomationPhaseView.tsx
--- a/Engine/Source/Programs/Horde/HordeDashboard/plugins/test-automation/testAutomationPhaseView.tsx
+++ b/Engine/Source/Programs/Horde/HordeDashboard/plugins/test-automation/testAutomationPhaseView.tsx
@@ -355,7 +355,10 @@ export const TestPhasesView: React.FC<{ testDataId: string, handler: TestDataHan
     useEffect(() => {
         handler.queryTestData(testDataId).then((fetchedSessionDetails) => {
             if (fetchedSessionDetails) {
-                handler.selectStream(fetchedSessionDetails.streamId);
+                if(!handler.filterState.stream) {
+                    // Select stream only if none selected.
+                    handler.selectStream(fetchedSessionDetails.streamId);
+                }
                 setSessionDetails(fetchedSessionDetails);
                 const testId = fetchedSessionDetails.test?.id ?? testDataId;
                 updateFilter(phaseFiltersContext.filterTestPhases(testId, fetchedSessionDetails.phases));
diff --git a/Engine/Source/Programs/Horde/HordeDashboard/plugins/test-automation/testAutomationSummary.tsx b/Engine/Source/Programs/Horde/HordeDashboard/plugins/test-automation/testAutomationSummary.tsx
--- a/Engine/Source/Programs/Horde/HordeDashboard/plugins/test-automation/testAutomationSummary.tsx
+++ b/Engine/Source/Programs/Horde/HordeDashboard/plugins/test-automation/testAutomationSummary.tsx
@@ -11,6 +11,7 @@ import { getHordeStyling } from "horde/styles/Styles";
 import { projectStore } from 'horde/backend/ProjectStore';
 import { SessionStatusBar, SessionValues, StatusBar, StatusBarStack, MultiOptionChooser, getTestSessionStatusColor, getStreamOptions, getStatusColors, statusTexts } from "./testAutomationCommon";
 import { TestHistoryGraph } from "./testHistoryGraph";
+import { userStorage } from "./testAutomationUserStorage";
 
 type SessionCallback = (session: TestSessionResult) => void;
 type HistorySessionCallback = (test: TestNameRef, onClick: SessionCallback) => void;
@@ -64,7 +65,7 @@ const StreamSelector: React.FC<{ streams: string[], onClick: (stream: string) =>
 export const TestHistoryView: React.FC<{ test: TestNameRef, handler: TestDataHandler, onClick: SessionCallback  }> = observer(({ test, handler, onClick }) => {
 
     const { hordeClasses } = getHordeStyling();
-    const [streams, setStreams] = useState<{current: string, toCompare?: string}>({current: handler.stream!});
+    const [streams, setStreams] = useState<{current: string, toCompare?: string}>({current: handler.stream!, toCompare: userStorage.getItem('historyToCompare', () => undefined)});
     const [selectedMeta, setSelectedMeta] = useState<Set<string>>(new Set());
 
     handler.subscribeToSubQueryLoading();
@@ -73,7 +74,10 @@ export const TestHistoryView: React.FC<{ test: TestNameRef, handler: TestDataHan
         setStreams({...streams, toCompare: undefined});
         handler.query(targetStream, true)
             .catch((reason) => console.error(reason))
-            .finally(() => setStreams({...streams, toCompare: targetStream}));
+            .finally(() => {
+                setStreams({...streams, toCompare: targetStream});
+                userStorage.setItem('historyToCompare', targetStream);
+            });
     }, [streams]);
 
     const availableStreams = handler.availableStreams.filter((s) => s !== streams.current);
@@ -198,15 +202,11 @@ const TestHistoryModal: React.FC<{ test: TestNameRef, handler: TestDataHandler,
         </Modal>
 }
 
-// Manifest to persist the exapanded summary cards when phase view widget is opened.
-// The phase view being lower in the view stack, the expanded properties value of the card widget get lost at that time.
-const _persistExpandedCards: Map<TestNameRef, boolean> = new Map();
-
 const TestSummaryCard: React.FC<{ test: TestNameRef, sessions: [MetadataRef, TestSessionResult][], metaKeyMask: string[] | undefined, handler: TestDataHandler, onDetailsClick: SessionCallback, onHistoryClick: HistorySessionCallback }> = memo(({ test, sessions, metaKeyMask, handler, onDetailsClick, onHistoryClick }) => {
-    const [expanded, setExpanded] = useState(_persistExpandedCards.get(test) ?? false);
+    const [expanded, setExpanded] = useState(userStorage.getItem('summaryExpandedCards', () => new Map<string, boolean>()).get(test) ?? false);
 
     const expandCard = useCallback((bExpand) => {
-        _persistExpandedCards.set(test, bExpand);
+        userStorage.getItem('summaryExpandedCards').set(test, bExpand);
         setExpanded(bExpand);
     }, [test]);
  
diff --git a/Engine/Source/Programs/Horde/HordeDashboard/plugins/test-automation/testAutomationUserStorage.ts b/Engine/Source/Programs/Horde/HordeDashboard/plugins/test-automation/testAutomationUserStorage.ts
new file mode 100644
--- /dev/null
+++ b/Engine/Source/Programs/Horde/HordeDashboard/plugins/test-automation/testAutomationUserStorage.ts
@@ -0,0 +1,23 @@
+// Copyright Epic Games, Inc. All Rights Reserved.
+
+class TransientStorage {
+    constructor() {
+        this.blackboard = new  Map<string, any>();
+    }
+
+    getItem(key: string, initializer?: () => any): any {
+        if (!this.blackboard.has(key) && !!initializer) {
+            this.blackboard.set(key, initializer());
+        }
+
+        return this.blackboard.get(key);
+    }
+
+    setItem(key: string, value: any) {
+        this.blackboard.set(key, value);
+    }
+
+    private blackboard: Map<string, any>;
+}
+
+export const userStorage = new TransientStorage();
