Override test descriptive metadata with Audit input

#rnx

diff --git a/Engine/Source/Programs/Horde/Plugins/Build/HordeServer.Build.Tests/Jobs/TestData/TestDataTests.cs b/Engine/Source/Programs/Horde/Plugins/Build/HordeServer.Build.Tests/Jobs/TestData/TestDataTests.cs
--- a/Engine/Source/Programs/Horde/Plugins/Build/HordeServer.Build.Tests/Jobs/TestData/TestDataTests.cs
+++ b/Engine/Source/Programs/Horde/Plugins/Build/HordeServer.Build.Tests/Jobs/TestData/TestDataTests.cs
@@ -613,6 +613,8 @@ namespace HordeServer.Tests.Jobs.TestData
 			List<GetTestSessionStreamResponse> streams = streamResult.Value;
 			Assert.AreNotEqual(0, streams.Count);
 			Assert.AreEqual(2, streams[0].Tests.Count);
+			Assert.IsNull(streams[0].Tests[0].Intent);
+			Assert.AreEqual("UE", streams[0].Tests[0].Team);
 
 			GetTestAuditRequest testRequestAll = new GetTestAuditRequest() { TestIds = streams[0].Tests.Select(t => t.Id).ToList() };
 			ActionResult<List<GetTestAuditResponse>> testauditResult = await TestDataController.GetTestAuditsV2Async(testRequestAll);
@@ -684,6 +686,7 @@ namespace HordeServer.Tests.Jobs.TestData
 				TestId = testId1,
 				Notes = "Some audit notes updated",
 				Intent = "Test intent",
+				Team = "TeamTest",
 				Owner = userIdString,
 				Customers = [userIdString],
 				Notifications = new Dictionary<string, string>(){ { _mainStreamId.ToString(), "test-workflow1" } , { _releaseStreamId.ToString(), "test-workflow2" } }
@@ -695,6 +698,7 @@ namespace HordeServer.Tests.Jobs.TestData
 			Assert.AreEqual(testId1, testauditPost.TestId);
 			Assert.AreEqual("Some audit notes updated", testauditPost.Notes);
 			Assert.AreEqual("Test intent", testauditPost.Intent);
+			Assert.AreEqual("TeamTest", testauditPost.Team);
 
 			// Check history output with time range for one test
 			testauditResult = await TestDataController.GetTestAuditsV2Async(testRequestOne, DateTimeOffset.UtcNow.AddHours(-1));
@@ -721,6 +725,16 @@ namespace HordeServer.Tests.Jobs.TestData
 			Assert.IsTrue(new List<string>(["test-workflow1", "test-workflow2"]).All(w => testauditPost.Notifications?.Any(n => n.Value == w) ?? false));
 			Assert.IsNull(testauditPost.LastAuditDate);
 			Assert.AreEqual("Test intent", testauditPost.Intent);
+			Assert.AreEqual("TeamTest", testauditPost.Team);
+
+			// Check test stream update
+			streamRequest = new GetStreamTestsRequest() { StreamIds = streamIds.ToList() };
+			streamResult = await TestDataController.GetTestStreamsV2Async(streamRequest);
+			Assert.IsNotNull(streamResult);
+			Assert.IsNotNull(streamResult.Value);
+			streams = streamResult.Value;
+			Assert.AreEqual(2, streams.Count);
+			Assert.AreEqual("TeamTest", streams[0].Tests[0].Team);
 		}
 
 		private readonly string[] _simpleTestDataLines =
diff --git a/Engine/Source/Programs/Horde/Plugins/Build/HordeServer.Build/Jobs/TestData/TestDataCollectionV2.cs b/Engine/Source/Programs/Horde/Plugins/Build/HordeServer.Build/Jobs/TestData/TestDataCollectionV2.cs
--- a/Engine/Source/Programs/Horde/Plugins/Build/HordeServer.Build/Jobs/TestData/TestDataCollectionV2.cs
+++ b/Engine/Source/Programs/Horde/Plugins/Build/HordeServer.Build/Jobs/TestData/TestDataCollectionV2.cs
@@ -1327,10 +1327,25 @@ namespace HordeServer.Jobs.TestData
 				updates.Add(updateBuilder.Set(x => x.Name, name));
 				if (description != null)
 				{
-					updates.Add(updateBuilder.Set(x => x.Team, description.Team));
-					updates.Add(updateBuilder.Set(x => x.Intent, description.Intent));
-					updates.Add(updateBuilder.Set(x => x.Harness, description.Harness));
-					updates.Add(updateBuilder.Set(x => x.Profile, description.Profile));
+					IReadOnlyList<ITestAudit> queryResult =  await FindTestAuditsAsync([test.Id], null, cancellationToken);
+					bool hasAudit = queryResult.Any();
+					// update only if no user override was done through audit
+					if (!hasAudit || queryResult[0].Team == null)
+					{
+						updates.Add(updateBuilder.Set(x => x.Team, description.Team));
+					}
+					if (!hasAudit || queryResult[0].Intent == null)
+					{
+						updates.Add(updateBuilder.Set(x => x.Intent, description.Intent));
+					}
+					if (!hasAudit || queryResult[0].Harness == null)
+					{
+						updates.Add(updateBuilder.Set(x => x.Harness, description.Harness));
+					}
+					if (!hasAudit || queryResult[0].Profile == null)
+					{
+						updates.Add(updateBuilder.Set(x => x.Profile, description.Profile));
+					}
 				}
 				await _tests.FindOneAndUpdateAsync(x => x.Id == test.Id, updateBuilder.Combine(updates), null, cancellationToken);
 			}
@@ -1435,6 +1450,31 @@ namespace HordeServer.Jobs.TestData
 
 			await _testAudits.InsertOneAsync(audit, null, cancellationToken);
 
+			// update test with user override
+			if (audit.Team != null || audit.Intent != null || audit.Harness != null || audit.Profile != null)
+			{
+				UpdateDefinitionBuilder<TestNameDocument> testUpdateBuilder = Builders<TestNameDocument>.Update;
+				List<UpdateDefinition<TestNameDocument>> testUpdates = new List<UpdateDefinition<TestNameDocument>>();
+				testUpdates.Add(testUpdateBuilder.Set(x => x.LastSeenUtc, DateTime.UtcNow));
+				if (audit.Team != null)
+				{
+					testUpdates.Add(testUpdateBuilder.Set(x => x.Team, audit.Team));
+				}
+				if (audit.Intent != null)
+				{
+					testUpdates.Add(testUpdateBuilder.Set(x => x.Intent, audit.Intent));
+				}
+				if (audit.Harness != null)
+				{
+					testUpdates.Add(testUpdateBuilder.Set(x => x.Harness, audit.Harness));
+				}
+				if (audit.Profile != null)
+				{
+					testUpdates.Add(testUpdateBuilder.Set(x => x.Profile, audit.Profile));
+				}
+				await _tests.FindOneAndUpdateAsync(x => x.Id == testId, testUpdateBuilder.Combine(testUpdates), null, cancellationToken);
+			}
+
 			return audit;
 		}
 
diff --git a/Engine/Source/Programs/Horde/Plugins/Build/HordeServer.Build/Jobs/TestData/TestDataController.cs b/Engine/Source/Programs/Horde/Plugins/Build/HordeServer.Build/Jobs/TestData/TestDataController.cs
--- a/Engine/Source/Programs/Horde/Plugins/Build/HordeServer.Build/Jobs/TestData/TestDataController.cs
+++ b/Engine/Source/Programs/Horde/Plugins/Build/HordeServer.Build/Jobs/TestData/TestDataController.cs
@@ -610,7 +610,11 @@ namespace HordeServer.Jobs.TestData
 					{
 						Id = test.Id.ToString(),
 						Name = test.Name,
-						Key = test.Key
+						Key = test.Key,
+						Team = test.Team,
+						Harness = test.Harness,
+						Profile =  test.Profile
+						// we purposefully skip Intent has it is not going to be a search-able field for now
 					}).ToList(),
 					TestMetadata = filteredMetadata.Select(meta => new GetTestMetadataResponse
 					{
