[Horde][BuildHealth][Dashboard] Bugfix for backspaces causes a desync with the rendered state.

diff --git a/Engine/Source/Programs/Horde/HordeDashboard/src/components/buildhealth/BuildhealthView.tsx b/Engine/Source/Programs/Horde/HordeDashboard/src/components/buildhealth/BuildhealthView.tsx
--- a/Engine/Source/Programs/Horde/HordeDashboard/src/components/buildhealth/BuildhealthView.tsx
+++ b/Engine/Source/Programs/Horde/HordeDashboard/src/components/buildhealth/BuildhealthView.tsx
@@ -68,7 +68,6 @@ interface ISearchableDropdownProps extends IDropdownProps {
 
     // Delegate that is run when search text field is changed.
     onSearchValueChanged?: (searchValue: string) => void;
-
     // Refs to all groups in the dropdown.
     groupRefs?: React.RefObject<(HTMLDivElement | null)[]>;
 }
@@ -165,7 +164,14 @@ const SearchableDropdown: React.FC<ISearchableDropdownProps> = (props) => {
             calloutProps={{
                 directionalHint: DirectionalHint.rightBottomEdge,
                 alignTargetEdge: true,
-                styles: { calloutMain: { width: 500, minWidth: 500 } },
+                styles: {
+                    calloutMain: {
+                        width: 500,
+                        minWidth: 500,
+                        height: 1000,
+                        maxHeight: 1000,
+                    }
+                },
             }}
             onRenderList={(props, defaultRender) => {
                 if (!props) {
@@ -182,7 +188,7 @@ const SearchableDropdown: React.FC<ISearchableDropdownProps> = (props) => {
                             <SearchBox
                                 placeholder="Filter options"
                                 autoFocus
-                                onChange={(_ev, newVal) => onSearchValueChanged?.(newVal || '')}
+                                onChange={(_ev, newVal) => onSearchValueChanged?.(newVal !== undefined ? newVal : '')}
                                 onClear={() => onSearchValueChanged?.('')}
                                 onKeyDown={(ev) => {
                                     if (ev.key === 'Enter' || ev.key === 'Escape') ev.stopPropagation();
@@ -321,7 +327,6 @@ const MultiListParameter: React.FC<{ handler: BuildHealthDataHandler, buildHealt
     }, [param.items, selectedSet, enabledSelectAll]);
 
     // #endregion -- Data Memos --
-
     const filteredOptions = useMemo(() => {
         const lowerFilter = filterText.toLowerCase();
         return doptions.filter(o =>
@@ -436,7 +441,21 @@ const MultiListParameter: React.FC<{ handler: BuildHealthDataHandler, buildHealt
             selectedKeys={selectedKeys}
             onChange={handleDropdownChange}
             onDismiss={handleMenuDismiss}
-            onSearchValueChanged={setFilterText}
+            onSearchValueChanged={(newValue) => {
+                // This stems from a very awkward bug where we need to force the underlying dropdown to resynchronize 
+                // on deletions by sending a empty filter - forcing the whole collection to be re-populated.
+                // We also need to force a complete animation frame in order to propagate this through, resulting in the
+                // underyling Dropdown component internal index synchronization.
+                // Without doing so, upon "deletion" changes the onChange will receive the wrong DropdownOption.
+                if (newValue.length < filterText.length) {
+                    setFilterText("");
+                    requestAnimationFrame(() => {
+                        setFilterText(newValue);
+                    });
+                } else {
+                    setFilterText(newValue);
+                }
+            }}
             onRenderOption={onRenderOption}
             styles={{
                 dropdown: { width: 300, maxWidth: 900 },
