Add filter-by-tag functionality to the test automation hub.  This adds a tag filter to the Test Automation Hub (V2) landing page and to the Phase View page.

diff --git a/Engine/Source/Programs/Horde/HordeDashboard/plugins/test-automation/api.ts b/Engine/Source/Programs/Horde/HordeDashboard/plugins/test-automation/api.ts
--- a/Engine/Source/Programs/Horde/HordeDashboard/plugins/test-automation/api.ts
+++ b/Engine/Source/Programs/Horde/HordeDashboard/plugins/test-automation/api.ts
@@ -319,6 +319,9 @@ export type GetTestPhaseSessionResponse = {
 
     /// The error finerprint from the event stream
     errorFingerprint?: string;
+
+    /// The tag ids associated with the phase session
+    tagIds?: string[]
 }
 
 export type TestSessionSummaryResponse = {
@@ -373,6 +376,9 @@ export type PhaseSessionDetailsResponse = {
     
     /// The event stream path that details the phase events
     eventStreamPath?: string;
+
+    /// The tag associated with the phase session
+    tags?: string[]
 }
 
 export type DeviceDetailsResponse = {
diff --git a/Engine/Source/Programs/Horde/HordeDashboard/plugins/test-automation/testAutomationCommon.tsx b/Engine/Source/Programs/Horde/HordeDashboard/plugins/test-automation/testAutomationCommon.tsx
--- a/Engine/Source/Programs/Horde/HordeDashboard/plugins/test-automation/testAutomationCommon.tsx
+++ b/Engine/Source/Programs/Horde/HordeDashboard/plugins/test-automation/testAutomationCommon.tsx
@@ -9,6 +9,8 @@ import { PhaseHistoryGraph } from "./phaseHistoryGraph";
 import { projectStore } from "horde/backend/ProjectStore";
 import { getHordeStyling } from "horde/styles/Styles";
 
+import * as d3 from "d3";
+
 export type PhaseSessionCallback = (session: PhaseSessionResult) => void;
 
 export const styles = {
@@ -501,6 +503,7 @@ class PhaseFiltersContext {
     constructor() {
         this._outcomeOptions = getPhaseOutcomeOptions();
         this._names = new Map();
+        this._tags = new Map();
     }
 
     get outcomeOptions() {
@@ -510,7 +513,8 @@ class PhaseFiltersContext {
     filterTestPhases(testId: string, phases: TestPhaseStatus[]) {
         const outcomes = this.getFilteredOutcomes();
         const names = this.getFilteredNames(testId);
-        return phases.filter(phase => phase.isMatch(names, outcomes));
+        const tags = this.getFilteredTags(testId);
+        return phases.filter(phase => phase.isMatch(names, outcomes, tags));
     }
 
     getFilteredOutcomes() {
@@ -542,11 +546,30 @@ class PhaseFiltersContext {
     }
 
     setFilteredNames(testId: string, names: string[]) {
-        this._names.set(testId, names.map(n => n.toLowerCase()));
+        const filteredNames = names.filter(i => !i.startsWith("#"));
+        this._names.set(testId, filteredNames.map(n => n.toLowerCase()));
+        if (filteredNames.length !== names.length) {
+            this.setFilteredTags(testId, names.filter(i => i.startsWith("#")).map(i => i.substring(1).toLowerCase()));
+        } else {
+            this.setFilteredTags(testId, []);
+        }
+    }
+
+    getFilteredTags(testId: string) {
+        return this._tags.get(testId);
+    }
+
+    setFilteredTags(testId: string, tags: string[]) {
+        if (tags.length === 0) {
+            this._tags.delete(testId);
+        } else {
+            this._tags.set(testId, tags);
+        }
     }
 
     private _outcomeOptions: ICheckListOption[];
     private _names: Map<string, string[]>;
+    private _tags: Map<string, string[]>;
 }
 export const phaseFiltersContext = new PhaseFiltersContext();
 
@@ -794,3 +817,29 @@ export const PhaseArtifactsModal: React.FC<{ phase: TestPhaseStatus, artifactPat
             </Modal>
         </Stack>
 }
+
+class TagColors {
+    
+    setTagColors(tagNames: string[]) {
+        // Strip off the leading "#" tag indicator if it's present
+        const cleanTagNames = tagNames.map(t => t.startsWith("#") ? t.substring(1) : t);
+
+        // Only reset the colors if the requested tag array is different than what we've already got
+        if ((tagNames.length > 0) && (JSON.stringify(this._tagColors.domain()) !== JSON.stringify(cleanTagNames))) {
+            this._tagColors.domain(Array.from(cleanTagNames));
+        }
+    }
+
+    getTagColor(tag: string) {
+        // Strip off the leading "#" tag indicator if it's present
+        if (tag.startsWith("#")) {
+            tag = tag.substring(1);
+        }
+        return this._tagColors(tag) as string;
+    }
+
+    private _tagColorRange = dashboard.darktheme ? [...d3.schemeDark2, ...d3.schemeCategory10] : [...d3.schemeSet2,...d3.schemeTableau10];
+    private _tagColors = d3.scaleOrdinal().range(this._tagColorRange);
+}
+export const tagColors = new TagColors();
+
diff --git a/Engine/Source/Programs/Horde/HordeDashboard/plugins/test-automation/testAutomationPhaseView.tsx b/Engine/Source/Programs/Horde/HordeDashboard/plugins/test-automation/testAutomationPhaseView.tsx
--- a/Engine/Source/Programs/Horde/HordeDashboard/plugins/test-automation/testAutomationPhaseView.tsx
+++ b/Engine/Source/Programs/Horde/HordeDashboard/plugins/test-automation/testAutomationPhaseView.tsx
@@ -9,9 +9,8 @@ import dashboard from "horde/backend/Dashboard";
 import { ChangeButton } from "horde/components/ChangeButton";
 import { getShortNiceTime, msecToElapsed } from "horde/base/utilities/timeUtils";
 import { getHordeTheme } from "horde/styles/theme";
-import { DashboardPreference, GetArtifactResponse } from "horde/backend/Api";
-import { isBright } from "horde/base/utilities/colors";
-import { CheckListOption, getPhaseSessionStatusColor, phaseFiltersContext, PhaseHistoryGraphWidget, PhasesStatusPie, styles } from "./testAutomationCommon";
+import { GetArtifactResponse } from "horde/backend/Api";
+import { CheckListOption, getPhaseSessionStatusColor, phaseFiltersContext, PhaseHistoryGraphWidget, PhasesStatusPie, styles, tagColors } from "./testAutomationCommon";
 import { projectStore } from "horde/backend/ProjectStore";
 import { JobArtifactsModal } from "horde/components/artifacts/ArtifactsModal";
 import { getJobParams, JobContext } from "./phaseCompare";
@@ -130,6 +129,11 @@ const TestPhaseHeader: React.FC<{ phase: TestPhaseStatus, expanded: boolean, tog
                             <FontIcon style={{ fontSize: 12, color: color }} iconName="Square" />
                         </Stack>
                         <Text style={{fontSize: 11}}>{phase.outcome}</Text>
+                        {!!phase.tags &&
+                            <Stack horizontal verticalAlign="center" tokens={{childrenGap: 6}} style={{paddingLeft: 6}}>
+                                {phase.tags.map(tag => <Stack style={{backgroundColor: tagColors.getTagColor(tag), borderRadius: 5, padding: '0px 4px', opacity: 0.8}}><Text style={{fontSize: 11, whiteSpace: 'nowrap'}}>#{tag}</Text></Stack>)}
+                            </Stack>
+                        }
                     </Stack>
                 </Stack>
                 <Stack horizontalAlign="end" grow style={{paddingRight: 12}}>
@@ -208,10 +212,11 @@ const onRenderSuggestionsItem = (item: ITag) => {
 
 const onRenderItem = (props: IPickerItemProps<ITag>) => {
     const item = props.item;
+    const isTag = item.name.startsWith("#");
     return <Stack key={`picker_item_${item.name}`} horizontal verticalAlign="center" verticalFill={true} style={{ marginRight: 2 }}>
                 <TooltipHost content={item.name} delay={0} closeDelay={500}>
                     <FontIcon
-                        style={{ fontSize: 13, padding: 4, paddingTop: 5, cursor: 'pointer', backgroundColor: styles.defaultFilter.backgroundColor, borderRadius: 3, color: 'white'}}
+                        style={{ fontSize: 13, padding: 4, paddingTop: 5, cursor: 'pointer', backgroundColor: isTag? tagColors.getTagColor(getTextFromTag(item)) : styles.defaultFilter.backgroundColor, borderRadius: 3, color: 'white'}}
                         iconName="Filter"
                         title="Remove"
                         onClick={props.onRemoveItem} />
@@ -223,22 +228,29 @@ const getTextFromTag = (item: ITag) => item.name;
 
 const PhaseFilter: React.FC<{phases: TestPhaseStatus[], selected?: string[], onChange: (items: string[]) => void, onFocus?: () => void, onBlur?: () => void}> = ({ phases, selected, onChange, onFocus, onBlur }) => {
     const [filters, setFilters] = useState<ITag[]>([]);
+    const phaseTokens = useRef<ITag[]>();
     const phaseTags = useRef<ITag[]>();
 
     useEffect(() => {
         const tokens: ITag[] = (new Set(phases.map(p => (p.name.split('.'))).flat())).keys().map(token => ({key: token.toLowerCase(), name: token})).toArray().sort((a, b) => a.key.length - b.key.length);
-        const tags: ITag[] = phases.map(p => ({key: `@${p.name.toLowerCase()}`, name: `@${p.name}`}));
-        phaseTags.current = [...tokens, ...tags];
+        const names: ITag[] = phases.map(p => ({key: `@${p.name.toLowerCase()}`, name: `@${p.name}`}));
+        const tags: ITag[] = (new Set(phases.map(p => p.tags ?? []).flat())).keys().map(t => ({key: `#${t.toLowerCase()}`, name: `#${t}`})).toArray().sort((a, b) => a.key.length - b.key.length);
+        phaseTokens.current = [...tokens, ...names];
+        phaseTags.current = tags;
         selected && setFilters(selected.map(i => ({key: i, name: i})));
     }, [phases]);
     
     const filterSelectedTests = useCallback((filterText: string): ITag[] => {
         if (filterText.length < 2) return [];
         const lowerText = filterText.toLowerCase();
-        const closestTags = phaseTags.current?.values().filter((tag) => (tag.key as string).indexOf(lowerText) >= 0).take(10).toArray() ?? [];
-        const isTagExist = closestTags.find((tag) => tag.key === lowerText);
-        if (!isTagExist) closestTags.splice(0, 0, {key: lowerText, name: filterText});
-        return closestTags;
+        if (filterText.startsWith("#")) {
+            const closestTags = phaseTags.current?.values().filter((tag) => (tag.key as string).indexOf(lowerText) >= 0).take(10).toArray() ?? [];
+            return closestTags;
+        }
+        const closestTokens = phaseTokens.current?.values().filter((tag) => (tag.key as string).indexOf(lowerText) >= 0).take(10).toArray() ?? [];
+        const isTokenExist = closestTokens.find((tag) => tag.key === lowerText);
+        if (!isTokenExist) closestTokens.splice(0, 0, {key: lowerText, name: filterText});
+        return closestTokens;
     }, [phases]);
 
     return <TagPicker
@@ -254,14 +266,14 @@ const PhaseFilter: React.FC<{phases: TestPhaseStatus[], selected?: string[], onC
                 selectedItems={filters}
                 onResolveSuggestions={filterSelectedTests}
                 getTextFromItem={getTextFromTag}
-                onChange={(tags) => {
-                    let items = tags ?? [];
-                    const uniqueTags = new Set(items.map(i => i.key));
-                    items = uniqueTags.keys().map(k => items.find(i => i.key === k)!).toArray();
+                onChange={(tokens) => {
+                    let items = tokens ?? [];
+                    const uniqueTokens = new Set(items.map(i => i.key));
+                    items = uniqueTokens.keys().map(k => items.find(i => i.key === k)!).toArray();
                     setFilters(items);
                     onChange(items.map(t => t.name));
                 }}
-                inputProps={{ placeholder: "Filter names", onFocus: onFocus, onBlur: onBlur }}
+                inputProps={{ placeholder: "Filter names and #tags", onFocus: onFocus, onBlur: onBlur }}
             />
 }
 
@@ -352,6 +364,13 @@ export const TestPhasesView: React.FC<{ testDataId: string, handler: TestDataHan
         }
     }, [])
 
+    useEffect(() => {
+        if (!!sessionDetails) {
+            const tagNames: string[] = (new Set(sessionDetails.phases.map(p => p.tags ?? []).flat())).keys().toArray().sort((a, b) => a.length - b.length);
+            tagColors.setTagColors(tagNames);
+        }
+    }, [sessionDetails]);
+
     useEffect(() => {
         handler.queryTestData(testDataId).then((fetchedSessionDetails) => {
             if (fetchedSessionDetails) {
@@ -567,10 +586,10 @@ export const TestPhasesView: React.FC<{ testDataId: string, handler: TestDataHan
                             </Stack>
 
                             <Stack verticalAlign='center' style={{ height: '100%' }}>
-                                <TooltipHost content='Filter phases by name'>
+                                <TooltipHost content='Filter phases by name and tag'>
                                     <PhaseFilter
                                         phases={sessionDetails.phases}
-                                        selected={phaseFiltersContext.getFilteredNames(testId)}
+                                        selected={[...phaseFiltersContext.getFilteredNames(testId), ...phaseFiltersContext.getFilteredTags(testId)?.map(t => `#${t}`) ?? []]}
                                         onChange={(items => {
                                             phaseFiltersContext.setFilteredNames(testId, items);
                                             updateFilter(phaseFiltersContext.filterTestPhases(testId, sessionDetails.phases));
diff --git a/Engine/Source/Programs/Horde/HordeDashboard/plugins/test-automation/testAutomationSummary.tsx b/Engine/Source/Programs/Horde/HordeDashboard/plugins/test-automation/testAutomationSummary.tsx
--- a/Engine/Source/Programs/Horde/HordeDashboard/plugins/test-automation/testAutomationSummary.tsx
+++ b/Engine/Source/Programs/Horde/HordeDashboard/plugins/test-automation/testAutomationSummary.tsx
@@ -3,7 +3,7 @@
 import { DefaultButton, FontIcon, IconButton, mergeStyleSets, Modal, Stack, Text, HoverCard, HoverCardType, IPlainCardProps, DirectionalHint, Spinner, SpinnerSize, IContextualMenuItem, IContextualMenuProps, IComboBoxOption } from "@fluentui/react";
 import { useState, useEffect, useCallback, memo } from "react";
 import { observer } from "mobx-react-lite";
-import { TestDataHandler, TestNameRef, MetadataRef, TestSessionStatus, TestSessionResult, TestStatus } from "./testData";
+import { TestDataHandler, TestNameRef, MetadataRef, TestSessionStatus, TestSessionResult, TestStatus, TestTagRef } from "./testData";
 import { TestOutcome } from './api';
 import dashboard, { StatusColor } from "horde/backend/Dashboard";
 import { getShortNiceTime } from "horde/base/utilities/timeUtils";
@@ -356,7 +356,7 @@ const TestSummaryCard: React.FC<{ test: TestNameRef, sessions: [MetadataRef, Tes
             </Stack>
 });
 
-const filterSessionsWithMetaMask = (testStatus: TestStatus, filteredTests: TestNameRef[], handler: TestDataHandler): [[TestNameRef, [MetadataRef, TestSessionStatus][]][], [TestNameRef, [MetadataRef, TestSessionStatus][]][], Map<string, string[]>] => {
+const filterSessionsWithMetaMask = (testStatus: TestStatus, filteredTests: TestNameRef[], selectedTags: TestTagRef[] | undefined, handler: TestDataHandler): [[TestNameRef, [MetadataRef, TestSessionStatus][]][], [TestNameRef, [MetadataRef, TestSessionStatus][]][], Map<string, string[]>] => {
     const filteredMetadata = handler.filteredMetadata;
 
     const testMeta: Map<string, [TestNameRef, MetadataRef[]][]> = new Map();
@@ -383,6 +383,11 @@ const filterSessionsWithMetaMask = (testStatus: TestStatus, filteredTests: TestN
             const session = metaStatus.getLastSession();
             if (!session) return;
 
+            if (!!selectedTags && !metaStatus.includeTags(selectedTags)) {
+                // skip if it filters by tags and no tag match
+                return;
+            }
+
             const outcome = session.outcome;
             if ( outcome === TestOutcome.Failure || outcome === TestOutcome.Unspecified ) {
                 totalFailureCount += session.phasesFailedCount + session.phasesUnspecifiedCount;
@@ -442,7 +447,9 @@ export const TestSummary: React.FC<{ handler: TestDataHandler }> = ({ handler })
         return null;
     }
 
-    const [failingTestSessions, passingTestSessions, testMetaKeys] = filterSessionsWithMetaMask(testStatus, filteredTests, handler);
+    const selectedTags = handler.selectedTags;
+
+    const [failingTestSessions, passingTestSessions, testMetaKeys] = filterSessionsWithMetaMask(testStatus, filteredTests, selectedTags, handler);
 
     const testSummaryTemplate =([test, sessions]) => {
 
diff --git a/Engine/Source/Programs/Horde/HordeDashboard/plugins/test-automation/testAutomationView.tsx b/Engine/Source/Programs/Horde/HordeDashboard/plugins/test-automation/testAutomationView.tsx
--- a/Engine/Source/Programs/Horde/HordeDashboard/plugins/test-automation/testAutomationView.tsx
+++ b/Engine/Source/Programs/Horde/HordeDashboard/plugins/test-automation/testAutomationView.tsx
@@ -10,7 +10,7 @@ import React, { useEffect, useState, useCallback, useMemo } from "react";
 import { projectStore } from 'horde/backend/ProjectStore';
 import { TestDataHandler, TestDataVersionRegistrar, TestSessionDetails } from "./testData";
 import { TestSummary } from "./testAutomationSummary";
-import { getStatusColors, getStreamOptions, getTestSessionStatusColor, MultiOptionChooser, SessionStatusBar } from "./testAutomationCommon";
+import { getStatusColors, getStreamOptions, getTestSessionStatusColor, MultiOptionChooser, SessionStatusBar, tagColors } from "./testAutomationCommon";
 import { TestPhasesView } from "./testAutomationPhaseView";
 import { JobStepTestDataItem } from "./api";
 
@@ -45,18 +45,20 @@ const onRenderSuggestionsItem = (item: ITag) => {
 
 const onRenderItem = (props: IPickerItemProps<ITag>) => {
     const item = props.item;
+    const isTag = item.name.startsWith("#");
     return <Stack style={{ marginTop: 2, marginLeft: 3 }} key={`picker_item_${item.name}`}>
-        <PrimaryButton
-            iconProps={{ iconName: "Cancel", styles: { root: { fontSize: 12, margin: 0, padding: 0 } } }}
-            styles={{label: { textOverflow: 'ellipsis', overflow: 'hidden', whiteSpace: 'nowrap', minWidth: 0, maxWidth: 150, margin: 0, paddingLeft: 3 }}}
-            style={{ padding: 2, paddingLeft: 5, paddingRight: 5, fontSize: 12, height: "unset"}}
-            text={item.name}
-            title={item.name}
-            onClick={props.onRemoveItem} />
-    </Stack>;
+                <PrimaryButton
+                    iconProps={{ iconName: "Cancel", styles: { root: { fontSize: 12, margin: 0, padding: 0 } } }}
+                    styles={{label: { textOverflow: 'ellipsis', overflow: 'hidden', whiteSpace: 'nowrap', minWidth: 0, maxWidth: 150, margin: 0, paddingLeft: 3 }}}
+                    style={{ padding: 2, paddingLeft: 5, paddingRight: 5, fontSize: 12, height: "unset", backgroundColor: isTag ? tagColors.getTagColor(getTextFromTag(item)) : undefined }}
+                    text={item.name}
+                    title={item.name}
+                    onClick={props.onRemoveItem} />
+            </Stack>;
 }
 
 const getTextFromTag = (item: ITag) => item.name;
+const getTagNameFromTag = (item: ITag) => item.name.substring(1);
 
 const TestChooser: React.FC = observer(() => {
     
@@ -75,14 +77,14 @@ const TestChooser: React.FC = observer(() => {
         }
     }, [handler.searchUpdated]);
 
-    const testTags: ITag[] = streamTests.map(t => ({key: `@${t.toLowerCase()}`, name: `@${t}`}));
+    const testNames: ITag[] = streamTests.map(t => ({key: `@${t.toLowerCase()}`, name: `@${t}`}));
     const filterSelectedTests = useCallback((filterText: string, _: ITag[]): ITag[] => {
         if (filterText.length < 2) return [];
         const lowerText = filterText.toLowerCase();
-        const closestTags = testTags.values().filter((tag) => (tag.key as string).indexOf(lowerText) >= 0).take(10).toArray();
-        const isTagExist = closestTags.find((tag) => tag.key === lowerText);
-        if (!isTagExist) closestTags.splice(0, 0, {key: lowerText, name: filterText});
-        return closestTags;
+        const closestNames = testNames.values().filter((tag) => (tag.key as string).indexOf(lowerText) >= 0).take(10).toArray();
+        const isTagExist = closestNames.find((tag) => tag.key === lowerText);
+        if (!isTagExist) closestNames.splice(0, 0, {key: lowerText, name: filterText});
+        return closestNames;
     }, [handler.queryLoading, handler.searchUpdated]);
 
     const selectedItems: ITag[] = testsSelection.keys().toArray().map((key) => ({key: key.toLowerCase(), name: key} as ITag));
@@ -115,6 +117,62 @@ const TestChooser: React.FC = observer(() => {
         </Stack>
 });
 
+const TagChooser: React.FC = observer(() => {
+    
+    handler.subscribeToFilter();
+    handler.subscribeToSearch();
+
+    const ctags: Set<string> = new Set<string>(handler.filterState.tags ?? []);
+    const streamTags = handler.allStreamTestTags;
+
+    const [tagsSelection, setTagsSelection] = useState<Set<string>>(ctags);
+
+    useEffect(() => {
+        // keep input selection in sync
+        if (tagsSelection.size !== ctags.size || tagsSelection.keys().some(s => !ctags.has(s))) {
+            setTagsSelection(ctags);
+        }
+    }, [handler.searchUpdated]);
+
+    const testTags: ITag[] = streamTags.map(t => ({key: t.toLowerCase(), name: `#${t}`}));
+    const filterSelectedTags = useCallback((filterText: string, _: ITag[]): ITag[] => {
+        if (filterText.startsWith("#")) filterText = filterText.substring(1);
+        if (filterText.length < 2) return [];
+        const lowerText = filterText.toLowerCase();
+        const closestTags = testTags.values().filter((tag) => (tag.key as string).indexOf(lowerText) >= 0).take(10).toArray();
+        return closestTags;
+    }, [handler.queryLoading, handler.searchUpdated]);
+
+    const selectedItems: ITag[] = tagsSelection.keys().toArray().map((key) => ({key: key.toLowerCase(), name: `#${key}`} as ITag));
+
+    return <Stack style={{ paddingTop: 12, paddingBottom: 4 }} horizontal>
+            <Label style={{ paddingTop: 5, paddingRight: 5 }}>Tags</Label>
+            <Stack tokens={{ childrenGap: 6 }} style={{ paddingLeft: 12 }} grow>
+                <TagPicker
+                    disabled={!handler.stream}
+                    styles={{ input: { height: 28, width: 40 }, itemsWrapper: { marginRight: 3, marginBottom: 2 } }}
+                    onRenderItem={onRenderItem}
+                    onRenderSuggestionsItem={onRenderSuggestionsItem}
+                    removeButtonAriaLabel="Remove"
+                    selectionAriaLabel="Selected tags"
+                    selectedItems={selectedItems}
+                    onResolveSuggestions={filterSelectedTags}
+                    getTextFromItem={getTagNameFromTag}
+                    onChange={(tags) => {
+                        if (!!tags) {
+                            const keys = tags!.map(getTagNameFromTag);
+                            keys.forEach((key) => handler.addTag(key));
+                            tagsSelection.difference(new Set(keys)).values().forEach(key => handler.removeTag(key));
+                        } else {
+                            tagsSelection.values().forEach(key => handler.removeTag(key));
+                        }
+                        setTagsSelection(new Set(handler.filterState.tags ?? []));
+                    }}
+                />                
+            </Stack>
+        </Stack>
+});
+
 const MetaKeyChooser: React.FC<{ onUpdateKey: (key: string) => void }> = observer(({ onUpdateKey }) => {
 
     handler.subscribeToFilter();
@@ -223,6 +281,9 @@ const TestAutomationSidebarLeft: React.FC = observer(() => {
                 <Stack key="chooser_test">
                     <TestChooser />
                 </Stack>
+                <Stack key="chooser_tag">
+                    <TagChooser />
+                </Stack>
                 {!!metaKeys.length && metaKeys.map((key) => <Stack key={`chooser_meta_${key}`}><MetaValueChooser metakey={key} onRemove={updateMetaKey} /></Stack>)}
                 <Stack>
                     <MetaKeyChooser key="chooser_meta" onUpdateKey={updateMetaKey} />
@@ -366,6 +427,11 @@ export const TestAutomationView: React.FC = () => {
         initialized && handler.syncSearchParam(searchParams.toString());
     }, [searchParams]);
 
+    useEffect(() => {
+        const allTags = handler.allStreamTestTags;
+        tagColors.setTagColors(allTags);
+    }, [handler.stream]);
+
     const { hordeClasses, modeColors } = getHordeStyling();
 
     return <Stack className={hordeClasses.horde} key="key_test_automation_hub">
diff --git a/Engine/Source/Programs/Horde/HordeDashboard/plugins/test-automation/testData.ts b/Engine/Source/Programs/Horde/HordeDashboard/plugins/test-automation/testData.ts
--- a/Engine/Source/Programs/Horde/HordeDashboard/plugins/test-automation/testData.ts
+++ b/Engine/Source/Programs/Horde/HordeDashboard/plugins/test-automation/testData.ts
@@ -87,6 +87,7 @@ export class PhaseSessionResult {
     stepId?: string;
     hasWarning?: boolean;
     errorFingerprint?: string;
+    tagIds?: TagId[];
 }
 
 export class TestSessionResult {
@@ -128,6 +129,7 @@ export class TestSessionResult {
     jobId?: string;
     stepId?: string;
     testDataId?: string;
+    tagIds?: TagId[];
 
     private _artifactFactory?: ArtifactFactory;
 }
@@ -182,6 +184,11 @@ export class TestSessionStatus {
         this.phaseHistories.get(phaseId)?.sort((a, b) => b.commitOrder - a.commitOrder);
     }
 
+    includeTags(tags: TestTagRef[]): boolean {
+        const sessionTagIds = this.getLastSession()?.tagIds;
+        return !!sessionTagIds && tags.some(tag => sessionTagIds.some(tagId => tag.id === tagId));
+    }
+
     history: TestSessionResult[];
     phaseHistories: Map<PhaseKey, PhaseSessionResult[]>;
     test: TestNameRef;
@@ -231,10 +238,20 @@ export class TestPhaseRef {
         return toMatch.charAt(0) === '@'? this.lowerName === toMatch.substring(1) : this.lowerName.includes(toMatch);
     }
 
+    updateTags(tags: TestTagRef[]) {
+        const newTags = new Set<TestTagRef>(tags);
+        this.tagRefs = this.tagRefs? this.tagRefs.union(newTags) : newTags;
+    }
+
+    includesTags(toMatch: string[]): boolean {
+        return toMatch.every(tag => this.tagRefs?.values().some(r => r.isMatch(tag)));
+    }
+
     id: string;
     name: string;
     key: string;
     version: number;
+    tagRefs?: Set<TestTagRef>;
 
     private lowerName: string;
 }
@@ -285,8 +302,10 @@ export abstract class TestPhaseStatus {
         this.session = phaseSessions?.find(s => s.sessionId === this.testSession!.id);
     }
 
-    isMatch(names: string[], outcomes?: Set<TestPhaseOutcome>) {
+    isMatch(names: string[], outcomes?: Set<TestPhaseOutcome>, tags?: string[]) {
         if (!!outcomes && !outcomes.has(this.outcome)) return false;
+        const lowerTags = this.tags?.map(t => t.toLowerCase());
+        if (!!tags && (!lowerTags || !lowerTags.some((t) => tags.includes(t)))) return false;
         if (names.length === 0) return true;
 
         const lowerName = this.name.toLowerCase();
@@ -334,6 +353,7 @@ export abstract class TestPhaseStatus {
     session?: PhaseSessionResult;
     componentRef?: any;
     hasWarning?: boolean;
+    tags?: string[];
 
     private _name: string;
     private _start: Date;
@@ -412,6 +432,7 @@ export class TestSessionDetails {
     test?: TestNameRef;
     meta?: MetadataRef;
     session?: TestSessionResult;
+    tags?: string[];
     preflightCommitId?: string;
 
     private _start: Date;
@@ -443,7 +464,7 @@ export class TestNameRef {
     }
 
     includesTags(toMatch: string[]): boolean {
-        return toMatch.every(tag => this.tagRefs?.values().some(r => r.isMatch(tag)));
+        return toMatch.some(tag => this.tagRefs?.values().some(r => r.isMatch(tag)));
     }
 
     id: TestId;
@@ -745,6 +766,18 @@ export class TestStreamRef {
         return matches;
     }
 
+    getMatchingTags(toMatchTags?: string[]): Set<TestTagRef> {
+        const lowerToMatchTags = toMatchTags?.map((i) => i.toLowerCase());
+        const matches: Set<TestTagRef> = new Set();
+        this.tags.forEach((tag) => {
+            if (!lowerToMatchTags || lowerToMatchTags.some((t) => tag.isMatch(t))) {
+                matches.add(tag);
+            }
+        });
+
+        return matches;
+    }
+
     id: StreamId;
     tests: TestNameRef[];
     meta: MetadataRef[];
@@ -1025,6 +1058,43 @@ export class TestDataHandler {
         this.setFilterUpdated();
     }
 
+    addTag(tag: string) {
+
+        const state = this.filterState;
+
+        if (state.tags) {
+            if (state.tags.includes(tag)) {
+                return;
+            }
+        }
+
+        if (!state.tags) {
+            state.tags = [];
+        }
+
+        state.tags.push(tag);
+
+        this.setFilterUpdated();
+    }
+
+    removeTag(tag: string) {
+
+        const state = this.filterState;
+
+        const idx = state.tags?.indexOf(tag);
+        if (!isNumber(idx) || idx! < 0) {
+            return;
+        }
+
+        state.tags!.splice(idx!, 1);
+
+        if (!state.tags!.length) {
+            state.tags = undefined;
+        }
+
+        this.setFilterUpdated();
+    }
+
     selectStream(streamId?: StreamId) {
 
         const state = this.filterState;
@@ -1148,6 +1218,27 @@ export class TestDataHandler {
         return testStreamRef!.getMatchingMeta(state.metadata).keys().toArray().sort((a, b) => a.orderedValues.join().localeCompare(b.orderedValues.join()));
     }
 
+    /// Tags that match current tag filters, return undefined if no tag is selected
+    get selectedTags(): TestTagRef[] | undefined {
+
+        const state = this.filterState;
+        if (!state.stream) {
+            return undefined;
+        }
+
+        const testStreamRef = this.testStreams.get(state.stream);
+        if (!testStreamRef) {
+            return undefined;
+        }
+
+        if (!state.tags) {
+            return undefined;
+        }
+
+        const lowerTags = state.tags.map(t => t.toLowerCase());
+        return testStreamRef!.getMatchingTags(lowerTags).keys().toArray().sort((a, b) => a.name.localeCompare(b.name));
+    }
+
     private queueHistoryReplace() {
         this.nextHistoryReplace = true;
     }
@@ -1177,6 +1268,7 @@ export class TestDataHandler {
         state.stream = state.stream ? state.stream : undefined;
         state.metadata = state.metadata && state.metadata.size > 0 ? state.metadata : undefined
         state.tests = state.tests?.sort((a, b) => a.localeCompare(b));
+        state.tags = state.tags?.sort((a, b) => a.localeCompare(b));
 
         let needUpdate = false;
         const search = this.search;
@@ -1196,6 +1288,11 @@ export class TestDataHandler {
             search.set("tests", null, state.tests);
         }
 
+        needUpdate = needUpdate || search.get('tags')?.toString() !== state.tags?.toString();
+        if (needUpdate) {
+            search.set("tags", null, state.tags);
+        }
+
         const metaKeys = new Set<string>(search.getGroup("meta")?.keys());
         metaKeys.union(new Set(state.metadata?.keys())).keys().forEach(key => {
             needUpdate = needUpdate || search.get('meta', key)?.toString() !== state.metadata?.get(key)?.toString();
@@ -1219,11 +1316,13 @@ export class TestDataHandler {
         const stream = search.get("stream") ?? undefined;
         const metadata = search.getGroup("meta");
         const tests = search.get("tests") ?? undefined;
+        const tags = search.get("tags") ?? undefined;
         const weeks = search.get("weeks") ?? undefined;
 
         state.stream = stream && stream.length > 0? stream[0] : undefined;
         state.metadata = metadata && metadata.size > 0 ? metadata : undefined;
         state.tests = tests;
+        state.tags = tags;
         state.weeks = weeks && weeks.length > 0 ? parseInt(weeks[0]) : defaultQueryWeeks;
         if (typeof (state.weeks) !== "number") {
             state.weeks = defaultQueryWeeks;
diff --git a/Engine/Source/Programs/Horde/HordeDashboard/plugins/test-automation/testDataV2Fetcher.ts b/Engine/Source/Programs/Horde/HordeDashboard/plugins/test-automation/testDataV2Fetcher.ts
--- a/Engine/Source/Programs/Horde/HordeDashboard/plugins/test-automation/testDataV2Fetcher.ts
+++ b/Engine/Source/Programs/Horde/HordeDashboard/plugins/test-automation/testDataV2Fetcher.ts
@@ -55,6 +55,7 @@ export class TestDataV2 implements TestVersionFetcher {
 
         const streamId = streamRef.id;
         const tests: Map<string, TestNameRef> = new Map(streamRef.tests.filter((t) => t.version >= this.version).map((t) => [t.id, t]));
+        const tags: Map<string, TestTagRef> = new Map(streamRef.tags.filter((t) => t.version >= this.version).map((t) => [t.id, t]));
         const meta: Map<string, MetadataRef> = new Map(streamRef.meta.filter((m) => m.version >= this.version).map((m) => [m.id, m]));
 
         const minQueryDate = new Date(new Date().valueOf() - (minuteInWeek * weeks * 60000));
@@ -69,13 +70,17 @@ export class TestDataV2 implements TestVersionFetcher {
                 const testSession = this.createTestSession(session);
                 const testNameRef = tests.get(testSession.nameRef);
                 const metadataRef = meta.get(testSession.metadataId);
+                const tagRefs = testSession.tagIds?.map((id) => tags.get(id)).filter((ref) => !!ref);
+                if (!!tagRefs && !!testNameRef) {
+                    testNameRef.updateTags(tagRefs);
+                }
                 if (!!testNameRef && !!metadataRef) {
                     status.addSession(testNameRef, metadataRef, testSession);
                     if (onAdd) onAdd(testSession);
-                }
+                }                
             });
         }
-
+        
     }
 
     async fetchTestData(testDataId: string): Promise<TestSessionDetails | undefined> {
@@ -97,6 +102,7 @@ export class TestDataV2 implements TestVersionFetcher {
     async fetchPhaseSessions(streamRef: TestStreamRef, status: TestMetaStatus, phaseKey: string, weeks: number): Promise<boolean> {
 
         const streamId = streamRef.id;
+        const tags: Map<string, TestTagRef> = new Map(streamRef.tags.filter((t) => t.version >= this.version).map((t) => [t.id, t]));
         const meta: Map<string, MetadataRef> = new Map(streamRef.meta.filter((m) => m.version >= this.version).map((m) => [m.id, m]));
 
         if (!status.test.phases) {
@@ -104,17 +110,21 @@ export class TestDataV2 implements TestVersionFetcher {
             status.test.phases = new Map(responses.map(i => i.phases.map(p => this.createPhaseRef(p))).flat().map(p => [p.key, p]));
         }
         
-        const phaseId = status.test.phases.get(phaseKey)?.id;
-        if (phaseId) {
+        const phaseRef = status.test.phases.get(phaseKey);
+        if (phaseRef) {
             const minQueryDate = new Date(new Date().valueOf() - (minuteInWeek * weeks * 60000));
             const maxQueryDate = new Date();
     
-            const responses = await getTestPhaseSessionsV2([streamId], [phaseId], minQueryDate.toISOString(), maxQueryDate.toISOString());
+            const responses = await getTestPhaseSessionsV2([streamId], [phaseRef.id], minQueryDate.toISOString(), maxQueryDate.toISOString());
             responses.forEach((session) => {
                 if (session.streamId !== streamId) {
                     return;
                 }
                 const phaseSession = this.createPhaseSession(session);
+                const tagRefs = phaseSession.tagIds?.map((id) => tags.get(id)).filter((ref) => !!ref);
+                if (!!tagRefs) {
+                    phaseRef.updateTags(tagRefs);
+                }
                 const metadataRef = meta.get(phaseSession.metadataId);
                 if (!!metadataRef) {
                     status.addPhaseSession(metadataRef, phaseKey, phaseSession);
@@ -130,7 +140,7 @@ export class TestDataV2 implements TestVersionFetcher {
 
     private testNameRefMap: Map<string, TestNameRef> = new Map();
     private metadataRefMap: Map<string, MetadataRef> = new Map();
-    private testTageRefMap: Map<string, TestTagRef> = new Map();
+    private testTagRefMap: Map<string, TestTagRef> = new Map();
 
     private createTestNameRef(test: GetTestNameResponse): TestNameRef {
         const globalKey = test.id;
@@ -156,10 +166,10 @@ export class TestDataV2 implements TestVersionFetcher {
 
     private createTagRef(item: GetTestTagResponse): TestTagRef {
         const globalKey = item.id;
-        let tag = this.testTageRefMap.get(globalKey);
+        let tag = this.testTagRefMap.get(globalKey);
         if (!tag) {
             tag = new TestTagRef(item.id, item.name, this.version);
-            this.testTageRefMap.set(globalKey, tag);
+            this.testTagRefMap.set(globalKey, tag);
         }
 
         return tag;
@@ -177,6 +187,7 @@ export class TestDataV2 implements TestVersionFetcher {
         result.phasesUnspecifiedCount = test.phasesUndefinedCount;
         result.jobId = test.jobId;
         result.stepId = test.stepId;
+        result.tagIds = test.tagIds;
 
         return result;
     }
@@ -188,6 +199,7 @@ export class TestDataV2 implements TestVersionFetcher {
         result.errorFingerprint = phase.errorFingerprint;
         result.jobId = phase.jobId;
         result.stepId = phase.stepId;
+        result.tagIds = phase.tagIds;
 
         return result;
     }
@@ -210,6 +222,7 @@ export class TestDataV2 implements TestVersionFetcher {
                                 : testDetails.summary.phasesTotalCount > 0 ?
                                     TestOutcome.Success : TestOutcome.Skipped;
         result.preflightCommitId = testData.preflightCommitId;
+        result.tags = testDetails.tags;
 
         // search for known metadata
         const metaFilter = new Map(Object.entries(testDetails.metadata).map(([key, value]) => [key, [value]]));
@@ -226,6 +239,7 @@ export class TestDataV2 implements TestVersionFetcher {
     private createPhaseStatus(phase: PhaseSessionDetailsResponse, deviceRefs: Map<string, DeviceRef>, parent: TestSessionDetails) {
         const devices = phase.deviceKeys.map(key => deviceRefs.get(key) ?? new DeviceRef(key));
         const result = new TestPhaseStatusV2(phase, devices, parent);
+        result.tags = phase.tags;
 
         return result;
     }
diff --git a/Engine/Source/Programs/Horde/HordeDashboard/plugins/test-automation/testHistoryGraph.ts b/Engine/Source/Programs/Horde/HordeDashboard/plugins/test-automation/testHistoryGraph.ts
--- a/Engine/Source/Programs/Horde/HordeDashboard/plugins/test-automation/testHistoryGraph.ts
+++ b/Engine/Source/Programs/Horde/HordeDashboard/plugins/test-automation/testHistoryGraph.ts
@@ -36,11 +36,18 @@ export class TestHistoryGraph {
         this.metaRefs = this.handler.filteredMetadata.filter((m) => this.status.includesMetadata(m) && (!selectedMetaKeys || selectedMetaKeys.includes(m.id)));
         const commonMetaKeys = this.commonMetaKeys = this.commonMetaKeys ?? MetadataRef.identifyCommonKeys(this.metaRefs);
 
+        const selectedTagRefs = this.handler.selectedTags;
+
         this.sessions = [];
         this.metaInfo = new Map();
 
         this.metaRefs.forEach((meta, index) => {
             const metaLastSession = this.status.getLastSession(meta);
+            if (!!metaLastSession && !!selectedTagRefs
+                && !this.status.sessions.get(meta)!.includeTags(selectedTagRefs)) {
+                // skip if it filters by tags and no tag match
+                return;
+            }
             this.metaStatus.set(meta, !metaLastSession? TestOutcome.Unspecified : metaLastSession.outcome);
 
             const elements: string[] = meta.getValuesExcept(commonMetaKeys);
