[Horde][BuildHealth][Server] Add additional temporary logging for long suspect scoring

diff --git a/Engine/Source/Programs/Horde/Plugins/Build/HordeServer.Build/BuildConfig.cs b/Engine/Source/Programs/Horde/Plugins/Build/HordeServer.Build/BuildConfig.cs
--- a/Engine/Source/Programs/Horde/Plugins/Build/HordeServer.Build/BuildConfig.cs
+++ b/Engine/Source/Programs/Horde/Plugins/Build/HordeServer.Build/BuildConfig.cs
@@ -114,7 +114,7 @@ namespace HordeServer
 		/// <summary>
 		/// The maximum number of seconds to spend scording suspects
 		/// </summary>
-		public int IssueFindSuspectsMaxScoringSeconds { get; set; } = 30;
+		public int IssueFindSuspectsMaxScoringSeconds { get; set; } = 60;
 
 		private readonly Dictionary<ProjectId, ProjectConfig> _projectLookup = new Dictionary<ProjectId, ProjectConfig>();
 		private readonly Dictionary<StreamId, StreamConfig> _streamLookup = new Dictionary<StreamId, StreamConfig>();
diff --git a/Engine/Source/Programs/Horde/Plugins/Build/HordeServer.Build/Issues/IssueService.cs b/Engine/Source/Programs/Horde/Plugins/Build/HordeServer.Build/Issues/IssueService.cs
--- a/Engine/Source/Programs/Horde/Plugins/Build/HordeServer.Build/Issues/IssueService.cs
+++ b/Engine/Source/Programs/Horde/Plugins/Build/HordeServer.Build/Issues/IssueService.cs
@@ -1507,6 +1507,15 @@ namespace HordeServer.Issues
 							suspects.Add(suspect);
 						}
 					}
+
+					// Temporary logging to gather information on how long these are taking
+					TimeSpan totalScoringDuration = DateTime.Now - scoringStart;
+					int warningDuration = 20;
+					if (totalScoringDuration.TotalSeconds >= warningDuration)
+					{
+						LoggerUtils.MultiLog([(_logger, LogLevel.Warning), (issueLogger, LogLevel.Information)], "Scoring all suspect commits took {TotalScoringSeconds} seconds", totalScoringDuration.TotalSeconds);
+					}
+
 					issueLogger?.LogDebug("Max rank of {MaxRank} achieved during suspect finding", maxRank);
 				}
 			}
