Horde: Updating schema documentation

#rnx

diff --git a/Engine/Source/Programs/Horde/Docs/Config/Schema/Globals.md b/Engine/Source/Programs/Horde/Docs/Config/Schema/Globals.md
--- a/Engine/Source/Programs/Horde/Docs/Config/Schema/Globals.md
+++ b/Engine/Source/Programs/Horde/Docs/Config/Schema/Globals.md
@@ -90,6 +90,7 @@ Name | Description
 `artifactTypes` | [ArtifactTypeConfig](#artifacttypeconfig)`[]`<br>Configuration for different artifact types
 `projects` | [ProjectConfig](Projects.md)`[]`<br>List of projects
 `enableConformTasks` | `boolean`<br>Whether to allow conform tasks to run
+`enableJobTaskConformCheck` | `boolean`<br>Whether to check if a conform is required on an agent before assigning a lease for a job task source.
 `issueFixedTag` | `string`<br>Commit tag to use for marking issues as fixed
 `issues` | [IssueConfig](Issues.md)<br>Build health issue configurations
 `issueFindSuspectsMaxSuccessDays` | `integer`<br>The maximum days since node success to scan for suspects
diff --git a/Engine/Source/Programs/Horde/Docs/Deployment/ServerSettings.md b/Engine/Source/Programs/Horde/Docs/Deployment/ServerSettings.md
--- a/Engine/Source/Programs/Horde/Docs/Deployment/ServerSettings.md
+++ b/Engine/Source/Programs/Horde/Docs/Deployment/ServerSettings.md
@@ -184,7 +184,6 @@ Name | Description
 `useLocalPerforceEnv` | `boolean`<br>Whether to use the local Perforce environment
 `perforceConnectionPoolSize` | `integer`<br>Number of pooled perforce connections to keep
 `enableConformTasks` | `boolean`<br>Whether to enable the conform task source.
-`enableJobTaskConformCheck` | `boolean`<br>Whether to check if a conform is required on an agent before assigning a lease for a job task source.
 `p4SwarmUrl` | `string`<br>Url of P4 Swarm installation
 `robomergeUrl` | `string`<br>Url of Robomergem installation
 `commitsViewerUrl` | `string`<br>Url of Commits Viewer
