Horde: Avoid logging too large config exception stacktraces

Causes log lines to get dropped by log processor.

diff --git a/Engine/Source/Programs/Horde/Plugins/Build/HordeServer.Build/Notifications/Sinks/SlackNotificationSink.cs b/Engine/Source/Programs/Horde/Plugins/Build/HordeServer.Build/Notifications/Sinks/SlackNotificationSink.cs
--- a/Engine/Source/Programs/Horde/Plugins/Build/HordeServer.Build/Notifications/Sinks/SlackNotificationSink.cs
+++ b/Engine/Source/Programs/Horde/Plugins/Build/HordeServer.Build/Notifications/Sinks/SlackNotificationSink.cs
@@ -2301,7 +2301,12 @@ namespace HordeServer.Notifications.Sinks
 			Exception? ex = info.Exception;
 			if (ex != null)
 			{
-				_logger.LogInformation(ex, "Sending config update failure notification: {Message}", ex.Message);
+				// Exception stacktraces can be really deep, which can lead to parsing issues on the log receiving end (OpenTelemetry, Datadog, etc)
+				_logger.LogInformation(
+					"Sending config update failure notification: {Message}, Exception: {ExceptionType} Stacktrace: {Stacktrace}", 
+					ex.Message,
+					ex.GetType().FullName,
+					ex.StackTrace?[..500]);
 
 				List<string> details = new List<string>();
 				if (ex is ConfigException configEx)
