Horde: Remove index hint for listing jobs per stream

After upgrading to DocumentDB 5, this index hint does more harm than good.
Running manual explain on this exact query in production shows significant speedups.

diff --git a/Engine/Source/Programs/Horde/Plugins/Build/HordeServer.Build/Jobs/JobCollection.cs b/Engine/Source/Programs/Horde/Plugins/Build/HordeServer.Build/Jobs/JobCollection.cs
--- a/Engine/Source/Programs/Horde/Plugins/Build/HordeServer.Build/Jobs/JobCollection.cs
+++ b/Engine/Source/Programs/Horde/Plugins/Build/HordeServer.Build/Jobs/JobCollection.cs
@@ -1240,8 +1240,6 @@ namespace HordeServer.Jobs
 			span.SetAttribute("Index", index);
 			span.SetAttribute("Count", count);
 
-			string indexHint = _streamThenTemplateThenCreationTimeIndex.Name;
-
 			// This find call uses an index hint. Modifying the parameter passed to FindAsync can affect execution time a lot as the query planner is forced to use the specified index.
 			return await FindInternalAsync(
 				new FindJobOptions(
@@ -1252,7 +1250,6 @@ namespace HordeServer.Jobs
 					MaxCreateTime: maxCreateTime),
 				index: index,
 				count: count,
-				indexHint: indexHint,
 				consistentRead: consistentRead, 
 				cancellationToken: cancellationToken);
 		}
