[Horde][BuildHealth][Dashboard]Label styling:
- Updated modal noun for "Step Link", "View Step", and "Step" so that it now properly uses "Step" | "Label".
- Updated tooltip onmouseover to use label where applicable
- Specific label styling fo step name column

diff --git a/Engine/Source/Programs/Horde/HordeDashboard/src/components/buildhealth/stepoutcome/StepOutcome.tsx b/Engine/Source/Programs/Horde/HordeDashboard/src/components/buildhealth/stepoutcome/StepOutcome.tsx
--- a/Engine/Source/Programs/Horde/HordeDashboard/src/components/buildhealth/stepoutcome/StepOutcome.tsx
+++ b/Engine/Source/Programs/Horde/HordeDashboard/src/components/buildhealth/stepoutcome/StepOutcome.tsx
@@ -56,9 +56,12 @@ const StepOutcomeFocusModal: React.FC<StepOutcomeFocusModalProps> = ({ item, onD
     if (item === undefined) { return; }
 
     let stepLink: string | null = null;
+    let stepNoun: string | null = null;
     if (isStepOutcome(item)) {
+        stepNoun = "Step";
         stepLink = `/job/${item.jobId}?step=${item.id}`;
     } else if (isLabelOutcome(item)) {
+        stepNoun = "Label";
         stepLink = `/job/${item.jobId}?label=${item.id}`;
     }
 
@@ -68,13 +71,13 @@ const StepOutcomeFocusModal: React.FC<StepOutcomeFocusModalProps> = ({ item, onD
         ['Stream', item.streamId],
         ['Code Change', item.change],
         ['Job', item.jobName],
-        ['Step', item.name],
+        [`${stepNoun}`, item.name],
         ['Duration', isStepOutcome(item) ? item.getDurationString() : NOT_APPLICABLE_LABEL],
         ['State', item.state],
         // For simplicity, we will compeltely override the outcome of this step with the job's outcome, as the step was aborted.
         item.state === AbortedStr ? ['Job Outcome', handler.getJobResponseData(item.jobId)?.summarizedOutcome ?? "-"] : isTerminalState(item) ? ['Outcome', item.outcome] : null,
         ['Job Link', <a href={`/job/${item.jobId}`} target="_blank" rel="noopen er noreferrer">View Job</a>],
-        ['Step Link', <a href={`${stepLink}`} target="_blank" rel="noopener noreferrer">View Step</a>],
+        [`${stepNoun} Link`, <a href={`${stepLink}`} target="_blank" rel="noopener noreferrer">View {stepNoun}</a>],
     ].filter((d): d is (string | JSX.Element)[] => Boolean(d));
 
     let shouldRenderDuplicateColumn = item.duplicateEntries && item.duplicateEntries.length > 0;
@@ -163,9 +166,12 @@ const StepOutcomeFocusModal: React.FC<StepOutcomeFocusModalProps> = ({ item, onD
                                 <Stack tokens={{ childrenGap: 10 }} styles={{ root: { paddingLeft: 16 } }}>
                                     {item.duplicateEntries.map((relatedEntry, index) => {
                                         let relatedStepLink: string | null = null;
+                                        let stepNoun: string | null = null;
                                         if (isStepOutcome(relatedEntry)) {
+                                            stepNoun = "Step";
                                             relatedStepLink = `/job/${relatedEntry.jobId}?step=${relatedEntry.id}`;
                                         } else if (isLabelOutcome(item)) {
+                                            stepNoun = "Label";
                                             relatedStepLink = `/job/${relatedEntry.jobId}?label=${relatedEntry.id}`;
                                         }
 
@@ -180,7 +186,7 @@ const StepOutcomeFocusModal: React.FC<StepOutcomeFocusModalProps> = ({ item, onD
                                                     ['Date (Step)', relatedEntry.startTime ? new Date(relatedEntry.startTime).toLocaleString() : 'N/A'],
                                                     ['Job', relatedEntry.jobName],
                                                     ['Job Link', <a href={`/job/${relatedEntry.jobId}`} target="_blank" rel="noopener noreferrer">View Related Job</a>],
-                                                    ['Step Link', <a href={`${relatedStepLink}`} target="_blank" rel="noopener noreferrer">View Related Step</a>],
+                                                    [`${stepNoun} Link`, <a href={`${relatedStepLink}`} target="_blank" rel="noopener noreferrer">View Related {stepNoun}</a>],
                                                     ['State', relatedEntry.state],
                                                 ].map(([label, value], j) => (
                                                     <Stack key={j} horizontal tokens={{ childrenGap: 4 }}>
diff --git a/Engine/Source/Programs/Horde/HordeDashboard/src/components/buildhealth/stepoutcome/StepOutcomeD3Component.tsx b/Engine/Source/Programs/Horde/HordeDashboard/src/components/buildhealth/stepoutcome/StepOutcomeD3Component.tsx
--- a/Engine/Source/Programs/Horde/HordeDashboard/src/components/buildhealth/stepoutcome/StepOutcomeD3Component.tsx
+++ b/Engine/Source/Programs/Horde/HordeDashboard/src/components/buildhealth/stepoutcome/StepOutcomeD3Component.tsx
@@ -459,11 +459,12 @@ function GenerateD3Table(stepOutcomeTable: StepOutcomeTable, handler: StepOutcom
                 .data([rowData])
                 .join(
                     enter => enter.append("td").attr("class", "step-name").text(d => d.stepName),
-                    update => update.text(d => d.stepName),
+                    update => update,
                     exit => exit.remove()
                 );
 
-            cells.style("position", "sticky")
+            cells.html(d => d.isLabel ? `<span style="background-color: ${theme.semanticColors.bodyBackgroundHovered};">${d.stepName}</span>` : `<span>${d.stepName}</span>`)
+                .style("position", "sticky")
                 .style("left", "32px")
                 .style("z-index", "1")
                 .style("background-color", modeColors.content)
@@ -473,7 +474,6 @@ function GenerateD3Table(stepOutcomeTable: StepOutcomeTable, handler: StepOutcom
                 .style("padding-left", "8px")
                 .style("text-overflow", "ellipsis")
                 .style("box-sizing", "border-box")
-                .text(d => `${d.stepName}`)
                 .on("mouseover", function (_event, stepNameHeader) {
                     d3.select(this).style("background-color", theme.semanticColors.bodyBackgroundHovered);
 
@@ -580,7 +580,12 @@ function GenerateD3Table(stepOutcomeTable: StepOutcomeTable, handler: StepOutcom
                                     tooltipContent.appendChild(jobDiv);
 
                                     const stepDiv = document.createElement("div");
-                                    stepDiv.innerHTML = `<strong>Step:</strong><p class="${tooltipClasses.tooltipLine}">${entry?.name ?? "-"}</p>`;
+                                    if (isLabelOutcome(entry)) {
+                                        stepDiv.innerHTML = `<strong>Label:</strong><p class="${tooltipClasses.tooltipLine}">${entry?.name ?? "-"}</p>`;
+                                    }
+                                    else {
+                                        stepDiv.innerHTML = `<strong>Step:</strong><p class="${tooltipClasses.tooltipLine}">${entry?.name ?? "-"}</p>`;
+                                    }
                                     tooltipContent.appendChild(stepDiv);
 
                                     attachTargettedToolTip(this as HTMLElement, scaffold.current!.tooltip, tooltipContent);
