[Horde][BuildHealth][Dashboard] Feedback from posthumous review.

diff --git a/Engine/Source/Programs/Horde/HordeDashboard/src/components/buildhealth/BuildHealthDataHandler.ts b/Engine/Source/Programs/Horde/HordeDashboard/src/components/buildhealth/BuildHealthDataHandler.ts
--- a/Engine/Source/Programs/Horde/HordeDashboard/src/components/buildhealth/BuildHealthDataHandler.ts
+++ b/Engine/Source/Programs/Horde/HordeDashboard/src/components/buildhealth/BuildHealthDataHandler.ts
@@ -320,7 +320,7 @@ export class BuildHealthDataHandler implements IHierarchyValidator {
 
     /**
      * Obtains all the currently known templates for a given stream.
-     * @param streamId THe stream.
+     * @param streamId The stream.
      * @returns An array representing all the known templates for a given stream.
      */
     getCachedTemplates(streamId: string): TemplateDescriptor[] {
diff --git a/Engine/Source/Programs/Horde/HordeDashboard/src/components/buildhealth/BuildHealthOptions.ts b/Engine/Source/Programs/Horde/HordeDashboard/src/components/buildhealth/BuildHealthOptions.ts
--- a/Engine/Source/Programs/Horde/HordeDashboard/src/components/buildhealth/BuildHealthOptions.ts
+++ b/Engine/Source/Programs/Horde/HordeDashboard/src/components/buildhealth/BuildHealthOptions.ts
@@ -375,7 +375,7 @@ export function upgradeBuildHealthQueryParams(params: BuildHealthQueryParams, cu
  * Processes provided Build Health Query Params for the Build Health View.
  * @param params The params to process.
  * @param options The Build Health Option controller to use to update options from the params.
- * @param handler  THe Build Health Data Handler to use to obtain data based on the params.
+ * @param handler The Build Health Data Handler to use to obtain data based on the params.
  */
 export async function loadBuildHealthOptionsFromParams(
     params: BuildHealthQueryParams,
