[Horde][BuildHealth][Dashboard] Fixup misleading variable name.

diff --git a/Engine/Source/Programs/Horde/HordeDashboard/src/components/buildhealth/BuildHealthOptionsController.ts b/Engine/Source/Programs/Horde/HordeDashboard/src/components/buildhealth/BuildHealthOptionsController.ts
--- a/Engine/Source/Programs/Horde/HordeDashboard/src/components/buildhealth/BuildHealthOptionsController.ts
+++ b/Engine/Source/Programs/Horde/HordeDashboard/src/components/buildhealth/BuildHealthOptionsController.ts
@@ -463,8 +463,8 @@ export class BuildHealthOptionsWriter implements IBuildHealthOptionWriter {
         for (let idx: number = 0; idx < labels.length; ++idx) {
             let label = labels[idx];
             let decodedFullyQualifiedLabelName = decodeFullyQualifiedLabelName(label);
-            let toggleAllStepId = encodeLabelKeyFromStrings(streamId, templateId, decodedFullyQualifiedLabelName.labelCategory, decodedFullyQualifiedLabelName.label);
-            this.toggleLabelInternal(toggleAllStepId, label, enabled);
+            let toggleAllLabelId = encodeLabelKeyFromStrings(streamId, templateId, decodedFullyQualifiedLabelName.labelCategory, decodedFullyQualifiedLabelName.label);
+            this.toggleLabelInternal(toggleAllLabelId, label, enabled);
         }
     }
 
