[Horde][BuildHealth][Server] Link to step history panel in Slack thread when there is no regression range

diff --git a/Engine/Source/Programs/Horde/Plugins/Build/HordeServer.Build/Notifications/Sinks/SlackNotificationSink.cs b/Engine/Source/Programs/Horde/Plugins/Build/HordeServer.Build/Notifications/Sinks/SlackNotificationSink.cs
--- a/Engine/Source/Programs/Horde/Plugins/Build/HordeServer.Build/Notifications/Sinks/SlackNotificationSink.cs
+++ b/Engine/Source/Programs/Horde/Plugins/Build/HordeServer.Build/Notifications/Sinks/SlackNotificationSink.cs
@@ -1276,6 +1276,10 @@ namespace HordeServer.Notifications.Sinks
 							{
 								suspectMessage += $"<{swarmLink}|View changes in Swarm>";
 							}
+							else if (span.LastSuccess == null)
+							{
+								suspectMessage += $"No previous successes for span. See <{GetStepUrl(span.FirstFailure.JobId, span.FirstFailure.StepId)}|step history panel> and click on change buttons for CL ranges.";
+							}
 
 							Uri? commitsViewerLink = GetCommitsViewerLinkForSpan(span);
 							if (commitsViewerLink != null)
