[Horde][BuildHealth][Dashboard][Server] Add prevention suggestion wording to root cause summary placeholder

diff --git a/Engine/Source/Programs/Horde/HordeDashboard/src/components/IssueViewV2.tsx b/Engine/Source/Programs/Horde/HordeDashboard/src/components/IssueViewV2.tsx
--- a/Engine/Source/Programs/Horde/HordeDashboard/src/components/IssueViewV2.tsx
+++ b/Engine/Source/Programs/Horde/HordeDashboard/src/components/IssueViewV2.tsx
@@ -2471,7 +2471,7 @@ const MarkFixedModal: React.FC<{ onClose: () => void, autoFix?:boolean } > = ({
                      label="Root Cause Category"
                      options={rcaCategoryOptions}
                      onChange={(event, option) => {
-                        if (option?.key === "rca_clear") {                           
+                        if (option?.key === "rca_clear") {
                            option = undefined;
                         }
                         setState({ ...state, rootCauseCategory: option?.text as string | undefined })
@@ -2479,7 +2479,7 @@ const MarkFixedModal: React.FC<{ onClose: () => void, autoFix?:boolean } > = ({
                   />
                </Stack>
                <Stack>
-                  <TextField disabled={disableRoot} spellCheck={false} autoComplete="off" label="Root Cause Summary" placeholder="Please enter a brief summary describing the root cause of this issue." value={state.rootCauseSummary} style={{ minHeight: 120, maxHeight: 120 }} multiline resizable={false} onChange={(ev, newValue) => {
+                  <TextField disabled={disableRoot} spellCheck={false} autoComplete="off" label="Root Cause Summary" placeholder="Please enter a brief summary describing the root cause of this issue and any suggestions to prevent it." value={state.rootCauseSummary} style={{ minHeight: 120, maxHeight: 120 }} multiline resizable={false} onChange={(ev, newValue) => {
                      ev.preventDefault();
                      setState({ ...state, rootCauseSummary: newValue ? newValue : "" })
                   }} />
diff --git a/Engine/Source/Programs/Horde/Plugins/Build/HordeServer.Build/Notifications/Sinks/SlackNotificationSink.cs b/Engine/Source/Programs/Horde/Plugins/Build/HordeServer.Build/Notifications/Sinks/SlackNotificationSink.cs
--- a/Engine/Source/Programs/Horde/Plugins/Build/HordeServer.Build/Notifications/Sinks/SlackNotificationSink.cs
+++ b/Engine/Source/Programs/Horde/Plugins/Build/HordeServer.Build/Notifications/Sinks/SlackNotificationSink.cs
@@ -3539,7 +3539,7 @@ namespace HordeServer.Notifications.Sinks
 					selectBlock.BlockId = "rootcause_category";
 					selectBlock.Optional = true;
 
-					InputBlock summaryBlock = view.AddInput("Root Cause Summary:", new PlainTextInputElement("rootcause_summary_action", multiLine: true, placeholder: "Please enter a brief summary describing the root cause of this issue", initialValue: issue?.RootCauseSummary));
+					InputBlock summaryBlock = view.AddInput("Root Cause Summary:", new PlainTextInputElement("rootcause_summary_action", multiLine: true, placeholder: "Please enter a brief summary describing the root cause of this issue and any suggestions to prevent it", initialValue: issue?.RootCauseSummary));
 					summaryBlock.BlockId = "rootcause_summary";
 					summaryBlock.Optional = true;
 
