Horde: Update release notes

#rnx

diff --git a/Engine/Source/Programs/Horde/Docs/ReleaseNotes.md b/Engine/Source/Programs/Horde/Docs/ReleaseNotes.md
--- a/Engine/Source/Programs/Horde/Docs/ReleaseNotes.md
+++ b/Engine/Source/Programs/Horde/Docs/ReleaseNotes.md
@@ -1,5 +1,40 @@
 # Release Notes
 
+## 2025-10-30
+
+* Add prevention suggestion wording to root cause summary placeholder (47573311)
+* Link to step history panel in Slack thread when there is no regression range (47545130)
+* Step Outcome view Label styling:
+  - Updated modal noun for "Step Link", "View Step", and "Step" so that it now properly uses "Step" | "Label".
+  - Updated tooltip onmouseover to use label where applicable
+  - Specific label styling fo step name column (47525598)
+* New workspace option, maxFileScanConcurrency, to control file scanning and writing independently (47506990)
+* Remove index hint for listing jobs per stream
+  After upgrading to DocumentDB 5, this index hint does more harm than good.
+  Running manual explain on this exact query in production shows significant speedups. (47493366)
+* Avoid logging too large config exception stacktraces
+  Causes log lines to get dropped by log processor. (47490436)
+* Add filter-by-tag functionality to the test automation hub. This adds a tag filter to the Test Automation Hub (V2) landing page and to the Phase View page. (47440781)
+* Bugfix for backspaces causes a desync with the rendered state. (47434123)
+* Move EnableJobTaskConformCheck from BuildServerConfig to BuildConfig (47418622)
+* Step Outcome view QoL enhancements:
+  - Persist view options to query url.
+  - Disable auto-select all for templates.
+  - Callout on dropdown moved down to avoid clashing with header dropdowns (when filter box changes).
+  - Select all now available in stream dropdown.
+    -- Code infrastructure is now there for templates, but is not connected due to server load issues that come from this.
+  - Minor fix where clearAll was not being reflected.
+  - Minor fix where UI properties were precipitating data refresh.
+  - Possible fix for filter <> selection desync (47401550)
+* Disable analyzers in Dockerfile build for Horde Server
+  They should be run in the already existing but separate BuildGraph target for better parallelism. (47321321)
+* Adding labels to option filters (47317161)
+* Improve file count and download speed log lines in PerforceMaterializer (47313084)
+* Make dirty checking more robust in PerforceMaterializer
+  Prevents unnecessary p4 flushes. (47311776)
+* Set env vars for UAT/BuildGraph in PerforceMaterializer (47267210)
+* Minor change to improve the 'no data' case styling (47266542)
+
 ## 2025-10-23
 
 * Make disposing of JobStepLogger more resilient
