Horde: Send cancellation request to Job Driver process

When then Job Handler receives a request to cancel the lease give the Job Driver process a chance to exit to update the step state and logs

diff --git a/Engine/Source/Programs/Horde/Drivers/JobDriver.Tests/AbortTestExecutor.cs b/Engine/Source/Programs/Horde/Drivers/JobDriver.Tests/AbortTestExecutor.cs
new file mode 100644
--- /dev/null
+++ b/Engine/Source/Programs/Horde/Drivers/JobDriver.Tests/AbortTestExecutor.cs
@@ -0,0 +1,40 @@
+// Copyright Epic Games, Inc. All Rights Reserved.
+
+using System;
+using System.Threading;
+using System.Threading.Tasks;
+using HordeCommon.Rpc.Tasks;
+using JobDriver.Execution;
+using Microsoft.Extensions.Logging;
+using Microsoft.Extensions.Logging.Abstractions;
+
+namespace JobDriver.Tests
+{
+	class AbortTestExecutor : JobExecutor
+	{
+		public const string Name = "Abort";
+
+		private readonly Func<ILogger, CancellationToken, Task> _func;
+
+		public AbortTestExecutor(Func<ILogger, CancellationToken, Task> func)
+			: base(new JobExecutorOptions(null!, null!, [], default, default, default, new RpcJobOptions { Executor = Name, Container = new RpcJobContainerOptions()}), SimpleTestExecutor.NoOpTracer, NullLogger.Instance)
+		{
+			_func = func;
+		}
+
+		protected override Task ExecuteInternalAsync(ILogger logger, CancellationToken cancellationToken)
+		{
+			return _func(logger, cancellationToken);
+		}
+
+		protected override Task<bool> SetupAsync(JobStepInfo step, ILogger logger, CancellationToken cancellationToken)
+		{
+			throw new NotImplementedException();
+		}
+
+		protected override Task<bool> ExecuteAsync(JobStepInfo step, ILogger logger, CancellationToken cancellationToken)
+		{
+			throw new NotImplementedException();
+		}
+	}
+}
diff --git a/Engine/Source/Programs/Horde/Drivers/JobDriver.Tests/WorkerServiceTest.cs b/Engine/Source/Programs/Horde/Drivers/JobDriver.Tests/WorkerServiceTest.cs
--- a/Engine/Source/Programs/Horde/Drivers/JobDriver.Tests/WorkerServiceTest.cs
+++ b/Engine/Source/Programs/Horde/Drivers/JobDriver.Tests/WorkerServiceTest.cs
@@ -82,8 +82,10 @@ namespace JobDriver.Tests
 					return JobStepOutcome.Success;
 				});
 
-				await Assert.ThrowsExceptionAsync<TaskCanceledException>(() => executor.ExecuteStepAsync(
-					null!, NullLogger.Instance, cancelSource.Token, stepCancelSource.Token));
+				(JobStepOutcome stepOutcome, JobStepState stepState) = await executor.ExecuteStepAsync(null!, NullLogger.Instance,
+					cancelSource.Token, stepCancelSource.Token);
+				Assert.AreEqual(JobStepOutcome.Failure, stepOutcome);
+				Assert.AreEqual(JobStepState.Aborted, stepState);
 			}
 
 			{
@@ -154,6 +156,30 @@ namespace JobDriver.Tests
 			Assert.AreEqual(JobStepState.Completed, (JobStepState)client.UpdateStepRequests[2].State);
 		}
 
+		[TestMethod]
+		[DataRow(true, DisplayName = "Cancel job executor when abort is signaled")]
+		[DataRow(false, DisplayName = "Complete job executor when abort is not signaled")]
+		public async Task AbortRequestedTestAsync(bool abort)
+		{
+			using AbortTestExecutor executor = new((logger, cancelToken) => Task.CompletedTask);
+
+			if (abort)
+			{
+				executor._abortHandle?.Set();
+			}
+
+			await executor.ExecuteAsync(NullLogger.Instance, CancellationToken.None);
+
+			if (abort)
+			{
+				Assert.IsTrue(executor._abortCancellationSource.IsCancellationRequested);
+			}
+			else
+			{
+				Assert.IsFalse(executor._abortCancellationSource.IsCancellationRequested);
+			}
+		}
+
 		[Ignore]
 		[TestMethod]
 		public async Task PollForStepAbortFailureTestAsync()
diff --git a/Engine/Source/Programs/Horde/Drivers/JobDriver/Execution/JobExecutor.cs b/Engine/Source/Programs/Horde/Drivers/JobDriver/Execution/JobExecutor.cs
--- a/Engine/Source/Programs/Horde/Drivers/JobDriver/Execution/JobExecutor.cs
+++ b/Engine/Source/Programs/Horde/Drivers/JobDriver/Execution/JobExecutor.cs
@@ -251,6 +251,31 @@ namespace JobDriver.Execution
 		/// </summary>
 		internal TimeSpan _stepAbortPollRetryDelay = TimeSpan.FromSeconds(30);
 
+		/// <summary>
+		/// Cancellation source for the job executor. Will be cancelled when <see cref="_abortHandle"/> is signaled
+		/// </summary>
+		/// <remarks>
+		/// Exposed as internal to ease testing.
+		/// </remarks>
+		internal CancellationTokenSource _abortCancellationSource = new();
+
+		/// <summary>
+		/// The event handle that will be signaled for the job executor to abort.
+		/// If the event is not created the process can still be canceled by the parent process but without the chance to complete 
+		/// </summary>
+		/// <remarks>
+		/// On Windows will be a shared named event created by the job handler parent process.
+		/// On Linux and macOS will be a local event that is signaled upon receiving a SIGTERM signal. 
+		/// </remarks>
+		/// <remarks>
+		/// Exposed as internal to ease testing.
+		/// </remarks>
+		internal EventWaitHandle? _abortHandle;
+
+		const string AbortEventKeyName = "UE_HORDE_JOB_EXECUTOR_ABORT_EVENT_NAME";
+
+		readonly PosixSignalRegistration? _sigTermHandler;
+
 		protected JobExecutor(JobExecutorOptions options, Tracer tracer, ILogger logger)
 		{
 			HordeClient = options.HordeClient;
@@ -267,6 +292,41 @@ namespace JobDriver.Execution
 
 			Tracer = tracer;
 			Logger = logger;
+
+			(_abortHandle, _sigTermHandler) = CreateAbortEvent(Environment.GetEnvironmentVariable(AbortEventKeyName), SigTermHandler);
+		}
+
+		private static (EventWaitHandle?, PosixSignalRegistration?) CreateAbortEvent(string? name, Action<PosixSignalContext> handler)
+		{
+			EventWaitHandle? ewh = null;
+			PosixSignalRegistration? psr = null;
+
+			if (OperatingSystem.IsWindows())
+			{
+				if (!String.IsNullOrEmpty(name))
+				{
+					try
+					{
+						EventWaitHandle.TryOpenExisting(name, out ewh);
+					}
+					catch
+					{
+						// ignored
+					}
+				}
+				else
+				{
+					// To ease testing on Windows; there won't be a shared event on Windows during tests
+					ewh = new EventWaitHandle(false, EventResetMode.ManualReset);
+				}
+			}
+			else if (OperatingSystem.IsLinux() || OperatingSystem.IsMacOS())
+			{
+				psr = PosixSignalRegistration.Create(PosixSignal.SIGTERM, handler);
+				ewh = new EventWaitHandle(false, EventResetMode.ManualReset);
+			}
+
+			return (ewh, psr);
 		}
 
 		/// <inheritdoc/>
@@ -278,6 +338,30 @@ namespace JobDriver.Execution
 
 		protected virtual void Dispose(bool disposing)
 		{
+			if (disposing)
+			{
+				_sigTermHandler?.Dispose();
+				_abortHandle?.Dispose();
+				_abortCancellationSource.Dispose();
+			}
+		}
+
+		void SigTermHandler(PosixSignalContext context)
+		{
+			if (_abortHandle == null)
+			{
+				return;
+			}
+
+			try
+			{
+				_abortHandle.Set();
+				context.Cancel = true;
+			}
+			catch
+			{
+				// ignored
+			}
 		}
 
 		public async Task ExecuteAsync(ILogger logger, CancellationToken cancellationToken)
@@ -289,6 +373,40 @@ namespace JobDriver.Execution
 			span.SetAttribute("horde.lease.id", LeaseId.ToString());
 			JobOptions.DecorateSpan(span);
 
+			using CancellationTokenSource stopAbortTask = new();
+			Task abortTask = CancelOnEventAsync(_abortHandle, _abortCancellationSource, logger, stopAbortTask.Token);
+			using CancellationTokenSource combined = CancellationTokenSource.CreateLinkedTokenSource(_abortCancellationSource.Token, cancellationToken);
+
+			try
+			{
+				await ExecuteInternalAsync(logger, combined.Token);
+			}
+			finally
+			{
+				await stopAbortTask.CancelAsync();
+				await abortTask;
+			}
+		}
+
+		static Task CancelOnEventAsync(EventWaitHandle? eventHandle, CancellationTokenSource cancellationSource, ILogger logger, CancellationToken cancellationToken = default)
+		{
+			return Task.Run(async () =>
+			{
+				if (eventHandle != null)
+				{
+					WaitHandle[] handles = [eventHandle, cancellationToken.WaitHandle];
+					int index = WaitHandle.WaitAny(handles);
+					if (index == Array.IndexOf(handles, eventHandle))
+					{
+						logger.LogInformation("Recieved cancellation request via an abort signal");
+						await cancellationSource.CancelAsync();
+					}
+				}
+			}, CancellationToken.None);
+		}
+
+		protected virtual async Task ExecuteInternalAsync(ILogger logger, CancellationToken cancellationToken)
+		{
 			string? cleanupVar = Environment.GetEnvironmentVariable("UE_HORDE_JOB_BATCH_CLEANUP");
 			if (cleanupVar != null && Int32.TryParse(cleanupVar, out int v) && v != 0)
 			{
@@ -328,6 +446,7 @@ namespace JobDriver.Execution
 				}
 			}
 
+			// TODO move this after FinishBatchAsync?
 			// If this lease was cancelled, don't bother updating the job state.
 			cancellationToken.ThrowIfCancellationRequested();
 
@@ -464,6 +583,8 @@ namespace JobDriver.Execution
 					// Print the finishing state
 					stepTimer.Stop();
 					logger.LogInformation("Completed in {Time}", stepTimer.Elapsed);
+
+					cancellationToken.ThrowIfCancellationRequested();
 				}
 			}
 			catch (Exception ex)
@@ -517,7 +638,7 @@ namespace JobDriver.Execution
 				if (cancellationToken.IsCancellationRequested && IsCancellationException(ex))
 				{
 					stepLogger.LogError("The step was cancelled by batch/lease");
-					throw;
+					return (JobStepOutcome.Failure, JobStepState.Aborted);
 				}
 
 				if (stepCancellationToken.IsCancellationRequested && IsCancellationException(ex))
diff --git a/Engine/Source/Programs/Horde/HordeAgent/Leases/Handlers/JobHandler.cs b/Engine/Source/Programs/Horde/HordeAgent/Leases/Handlers/JobHandler.cs
--- a/Engine/Source/Programs/Horde/HordeAgent/Leases/Handlers/JobHandler.cs
+++ b/Engine/Source/Programs/Horde/HordeAgent/Leases/Handlers/JobHandler.cs
@@ -1,6 +1,8 @@
 // Copyright Epic Games, Inc. All Rights Reserved.
 
+using System.Diagnostics;
 using System.Globalization;
+using System.Runtime.InteropServices;
 using EpicGames.Core;
 using EpicGames.Horde;
 using EpicGames.Horde.Agents.Leases;
@@ -19,11 +21,74 @@ namespace HordeAgent.Leases.Handlers
 	class JobHandler : LeaseHandler<ExecuteJobTask>
 	{
 		private readonly IOptionsMonitor<AgentSettings> _settings;
-		
+
+		private ManagedProcess? _driverProcess;
+		private ManagedProcessGroup? _driverProcessGroup;
+		private readonly string? _driverCancelEventName;
+		private readonly EventWaitHandle? _driverCancelEvent;
+		private const string DriverCancelEventKeyName = "UE_HORDE_JOB_EXECUTOR_ABORT_EVENT_NAME";
+
+		/// <summary>
+		/// The cancellation source to make sure the job driver is always terminated 
+		/// </summary>
+		private readonly CancellationTokenSource _cancellationSource = new();
+
+		/// <summary>
+		/// How long to give the job driver child process to handle the shutdown request before terminating the process ourselves
+		/// </summary>
+		private readonly TimeSpan _driverCancelGracePeriod = TimeSpan.FromSeconds(30);
+
 		public JobHandler(RpcLease lease, IOptionsMonitor<AgentSettings> settings)
 			: base(lease)
 		{
 			_settings = settings;
+
+			(_driverCancelEvent, _driverCancelEventName) = CreateDriverCancelEvent("Horde.JobDriver.Cancel");
+		}
+
+		private static (EventWaitHandle?, string?) CreateDriverCancelEvent(string name)
+		{
+			// Named events across processes are only supported on Windows
+			if (!OperatingSystem.IsWindows())
+			{
+				return (null, null);
+			}
+
+			// Must have a unique name per job to avoid a previous shared signaled event being used and the job driver exiting immediately
+			string uniqueName = $"{name}.{Guid.NewGuid()}";
+
+			bool createdNew = false;
+			EventWaitHandle? eventObj = null;
+			try
+			{
+				eventObj = new EventWaitHandle(false, EventResetMode.ManualReset, uniqueName, out createdNew);
+			}
+			catch
+			{
+				// ignored
+			}
+			finally
+			{
+				if (!createdNew)
+				{
+					eventObj?.Dispose();
+					eventObj = null;
+				}
+			}
+
+			return createdNew ? (eventObj, uniqueName) : (null, null);
+		}
+
+		protected override void Dispose(bool disposing)
+		{
+			if (disposing)
+			{
+				_driverProcess?.Dispose();
+				_driverProcessGroup?.Dispose();
+				_driverCancelEvent?.Dispose();
+				_cancellationSource.Dispose();
+			}
+			base.Dispose(disposing);
 		}
 
 		/// <inheritdoc/>
@@ -67,8 +132,20 @@ namespace HordeAgent.Leases.Handlers
 					environment["UE_HORDE_TEMP_STORAGE_MAX_BATCH_SIZE"] = Convert.ToString(_settings.CurrentValue.TempStorageMaxBatchSize.Value, CultureInfo.InvariantCulture);
 				}
 
-				exitCode = await RunDotNetProcessAsync(driverAssembly, arguments, environment, AgentApp.IsSelfContained, logger, cancellationToken);
-				logger.LogInformation("Driver finished with exit code {ExitCode}", exitCode);				
+				if (_driverCancelEventName != null)
+				{
+					environment[DriverCancelEventKeyName] = _driverCancelEventName;
+				}
+
+				// Need to know the driver process ID for cancellation, use an alternative to RunDotNetProcessAsync
+				StartDriverProcess(driverAssembly, arguments, environment, AgentApp.IsSelfContained, logger);
+				await using (cancellationToken.Register(DriverCancelHandler))
+				{
+					exitCode = await WaitForDriverProcessExitAsync(logger, _cancellationSource.Token);
+				}
+
+				cancellationToken.ThrowIfCancellationRequested();
+				logger.LogInformation("Driver finished with exit code {ExitCode}", exitCode);
 			}
 			catch (OperationCanceledException ex)
 			{
@@ -107,6 +184,117 @@ namespace HordeAgent.Leases.Handlers
 
 			return (exitCode == 0) ? LeaseResult.Success : LeaseResult.Failed;
 		}
+
+		private void DriverCancelHandler()
+		{
+			bool handled = false;
+			if (_driverCancelEvent != null)
+			{
+				_driverCancelEvent.Set();
+				handled = true;
+			}
+			else if (!OperatingSystem.IsWindows())
+			{
+				// Cannot use named events on non-Windows OSes, instead send SIGTERM to the job driver process
+				if (_driverProcess != null)
+				{
+					handled = SendSigTerm(_driverProcess.Id);
+				}
+			}
+
+			// Whether the job driver process was signaled to abort, the job driver must be terminated
+			if (handled)
+			{
+				_cancellationSource.CancelAfter(_driverCancelGracePeriod);
+			}
+			else
+			{
+				_cancellationSource.Cancel();
+			}
+		}
+
+		private void StartDriverProcess(FileReference entryAssembly, IEnumerable<string> arguments, IReadOnlyDictionary<string, string>? environment, bool useNativeHost, ILogger logger)
+		{
+			if (useNativeHost)
+			{
+				FileReference nativeHost = entryAssembly.ChangeExtension(OperatingSystem.IsWindows() ? ".exe" : null);
+				StartDriverProcess(nativeHost.FullName, arguments, environment, logger);
+			}
+			else
+			{
+				IEnumerable<string> allArguments = arguments.Prepend(entryAssembly.FullName);
+				StartDriverProcess("dotnet", allArguments, environment, logger);
+			}
+		}
+
+		private void StartDriverProcess(string executable, IEnumerable<string> arguments, IReadOnlyDictionary<string, string>? environment, ILogger logger)
+		{
+			string commandLine = CommandLineArguments.Join(arguments);
+			logger.LogInformation("Starting child process: {Executable} {CommandLine}", CommandLineArguments.Quote(executable), commandLine);
+			try
+			{
+				_driverProcessGroup = new ManagedProcessGroup();
+				_driverProcess = new ManagedProcess(_driverProcessGroup, executable, commandLine, null, environment, ProcessPriorityClass.Normal);
+			}
+			catch (Exception ex)
+			{
+				logger.LogError(ex, "Failed to start process: {Message}", ex.Message);
+				throw;
+			}
+		}
+
+		private async Task<int> WaitForDriverProcessExitAsync(ILogger logger, CancellationToken cancellationToken)
+		{
+			if (_driverProcess == null)
+			{
+				throw new InvalidOperationException("Job Driver process not created");
+			}
+
+			try
+			{
+				for (; ; )
+				{
+					string? line = await _driverProcess.ReadLineAsync(cancellationToken);
+					if (line == null)
+					{
+						break;
+					}
+
+					JsonLogEvent jsonLogEvent;
+					if (JsonLogEvent.TryParse(line, out jsonLogEvent))
+					{
+						logger.LogJsonLogEvent(jsonLogEvent);
+					}
+					else
+					{
+						logger.LogInformation("{Line}", line);
+					}
+				}
+
+				await _driverProcess.WaitForExitAsync(CancellationToken.None);
+				return _driverProcess.ExitCode;
+			}
+			catch (Exception ex)
+			{
+				logger.LogError(ex, "Failed to run process: {Message}", ex.Message);
+				throw;
+			}
+		}
+
+		[DllImport("libc", SetLastError = true)]
+		private static extern int kill(int pid, int sig);
+
+		private static bool SendSigTerm(int processId)
+		{
+			if (!(OperatingSystem.IsLinux() || OperatingSystem.IsMacOS()))
+			{
+				throw new PlatformNotSupportedException($"{nameof(SendSigTerm)} not supported");
+			}
+
+			const int SigTerm = 15;
+			int result = kill(processId, SigTerm);
+			return result == 0;
+		}
 	}
 
 	class JobHandlerFactory : LeaseHandlerFactory<ExecuteJobTask>
