Remove todo comment

#rnx

diff --git a/Engine/Source/Programs/Horde/Drivers/JobDriver/Execution/JobExecutor.cs b/Engine/Source/Programs/Horde/Drivers/JobDriver/Execution/JobExecutor.cs
--- a/Engine/Source/Programs/Horde/Drivers/JobDriver/Execution/JobExecutor.cs
+++ b/Engine/Source/Programs/Horde/Drivers/JobDriver/Execution/JobExecutor.cs
@@ -446,7 +446,6 @@ namespace JobDriver.Execution
 				}
 			}
 
-			// TODO move this after FinishBatchAsync?
 			// If this lease was cancelled, don't bother updating the job state.
 			cancellationToken.ThrowIfCancellationRequested();
 
