Horde: Restore throwing cancellation exception when step is abort and run step clean up when cancellation is receieved.

diff --git a/Engine/Source/Programs/Horde/Drivers/JobDriver.Tests/WorkerServiceTest.cs b/Engine/Source/Programs/Horde/Drivers/JobDriver.Tests/WorkerServiceTest.cs
--- a/Engine/Source/Programs/Horde/Drivers/JobDriver.Tests/WorkerServiceTest.cs
+++ b/Engine/Source/Programs/Horde/Drivers/JobDriver.Tests/WorkerServiceTest.cs
@@ -82,10 +82,8 @@ namespace JobDriver.Tests
 					return JobStepOutcome.Success;
 				});
 
-				(JobStepOutcome stepOutcome, JobStepState stepState) = await executor.ExecuteStepAsync(null!, NullLogger.Instance,
-					cancelSource.Token, stepCancelSource.Token);
-				Assert.AreEqual(JobStepOutcome.Failure, stepOutcome);
-				Assert.AreEqual(JobStepState.Aborted, stepState);
+				await Assert.ThrowsExceptionAsync<TaskCanceledException>(() => executor.ExecuteStepAsync(
+					null!, NullLogger.Instance, cancelSource.Token, stepCancelSource.Token));
 			}
 
 			{
diff --git a/Engine/Source/Programs/Horde/Drivers/JobDriver/Execution/JobExecutor.cs b/Engine/Source/Programs/Horde/Drivers/JobDriver/Execution/JobExecutor.cs
--- a/Engine/Source/Programs/Horde/Drivers/JobDriver/Execution/JobExecutor.cs
+++ b/Engine/Source/Programs/Horde/Drivers/JobDriver/Execution/JobExecutor.cs
@@ -398,7 +398,7 @@ namespace JobDriver.Execution
 					int index = WaitHandle.WaitAny(handles);
 					if (index == Array.IndexOf(handles, eventHandle))
 					{
-						logger.LogInformation("Recieved cancellation request via an abort signal");
+						logger.LogWarning("Received cancellation request via an abort signal");
 						await cancellationSource.CancelAsync();
 					}
 				}
@@ -558,13 +558,20 @@ namespace JobDriver.Execution
 								// Will get called even when cancellation token for the lease/batch fires
 								stepFinishedSource.SetResult(true); // Tell background poll task to stop
 								await stepPollTask;
-							}
-
-							// Kill any processes spawned by the step
-							await TerminateProcessHelper.TerminateProcessesAsync(TerminateCondition.AfterStep, WorkingDir, ProcessesToTerminate, logger, cancellationToken);
 
-							// Wait for the logger to finish
-							await stepLogger.StopAsync();
+								// Kill any processes spawned by the step
+								try
+								{
+									await TerminateProcessHelper.TerminateProcessesAsync(TerminateCondition.AfterStep, WorkingDir, ProcessesToTerminate, logger, CancellationToken.None);
+								}
+								catch (Exception ex)
+								{
+									logger.LogWarning(ex, "Exception while terminating processes: {Message}", ex.Message);
+								}
+
+								// Wait for the logger to finish
+								await stepLogger.StopAsync();
+							}
 
 							// Reflect the warnings/errors in the step outcome
 							if (stepOutcome > stepLogger.Outcome)
@@ -582,8 +589,6 @@ namespace JobDriver.Execution
 					// Print the finishing state
 					stepTimer.Stop();
 					logger.LogInformation("Completed in {Time}", stepTimer.Elapsed);
-
-					cancellationToken.ThrowIfCancellationRequested();
 				}
 			}
 			catch (Exception ex)
@@ -637,7 +642,7 @@ namespace JobDriver.Execution
 				if (cancellationToken.IsCancellationRequested && IsCancellationException(ex))
 				{
 					stepLogger.LogError("The step was cancelled by batch/lease");
-					return (JobStepOutcome.Failure, JobStepState.Aborted);
+					throw;
 				}
 
 				if (stepCancellationToken.IsCancellationRequested && IsCancellationException(ex))
