Horde: Set the Job Driver cancellation event to nonsignaled after process exit

diff --git a/Engine/Source/Programs/Horde/HordeAgent/Leases/Handlers/JobHandler.cs b/Engine/Source/Programs/Horde/HordeAgent/Leases/Handlers/JobHandler.cs
--- a/Engine/Source/Programs/Horde/HordeAgent/Leases/Handlers/JobHandler.cs
+++ b/Engine/Source/Programs/Horde/HordeAgent/Leases/Handlers/JobHandler.cs
@@ -157,6 +157,10 @@ namespace HordeAgent.Leases.Handlers
 				logger.LogError(ex, "Unhandled exception: {Message}", ex.Message);
 				exception = ex;
 			}
+			finally
+			{
+				_driverCancelEvent?.Reset();
+			}
 
 			if (exception != null || exitCode != 0)
 			{
