Horde: Update log message when the job executor receives cancellation signal

#rnx

diff --git a/Engine/Source/Programs/Horde/Drivers/JobDriver/Execution/JobExecutor.cs b/Engine/Source/Programs/Horde/Drivers/JobDriver/Execution/JobExecutor.cs
--- a/Engine/Source/Programs/Horde/Drivers/JobDriver/Execution/JobExecutor.cs
+++ b/Engine/Source/Programs/Horde/Drivers/JobDriver/Execution/JobExecutor.cs
@@ -398,7 +398,7 @@ namespace JobDriver.Execution
 					int index = WaitHandle.WaitAny(handles);
 					if (index == Array.IndexOf(handles, eventHandle))
 					{
-						logger.LogWarning("Received cancellation request via an abort signal");
+						logger.LogWarning("Received a cancellation request from the parent job handler process");
 						await cancellationSource.CancelAsync();
 					}
 				}
