[Horde][BuildHealth][Dashboard] Correctly consume cascaded command option & simplified label styling.

diff --git a/Engine/Source/Programs/Horde/HordeDashboard/src/components/buildhealth/BuildHealthOptionsController.ts b/Engine/Source/Programs/Horde/HordeDashboard/src/components/buildhealth/BuildHealthOptionsController.ts
--- a/Engine/Source/Programs/Horde/HordeDashboard/src/components/buildhealth/BuildHealthOptionsController.ts
+++ b/Engine/Source/Programs/Horde/HordeDashboard/src/components/buildhealth/BuildHealthOptionsController.ts
@@ -447,6 +447,7 @@ export class BuildHealthOptionsWriter implements IBuildHealthOptionWriter {
         for (let idx: number = 0; idx < stepNames.length; ++idx) {
             let stepName = stepNames[idx];
             let toggleAllStepId = encodeStepNameFromStringsPreserveCase(streamId, templateId, stepName);
+
             if (cascadeLinkedAction) {
                 this.toggleStep(toggleAllStepId, stepName, enabled);
             } else {
@@ -464,6 +465,10 @@ export class BuildHealthOptionsWriter implements IBuildHealthOptionWriter {
             let label = labels[idx];
             let decodedFullyQualifiedLabelName = decodeFullyQualifiedLabelName(label);
             let toggleAllLabelId = encodeLabelKeyFromStrings(streamId, templateId, decodedFullyQualifiedLabelName.labelCategory, decodedFullyQualifiedLabelName.label);
+
+            if (cascadeLinkedAction) {
+                this.toggleLabel(toggleAllLabelId, label, enabled);
+            }
             this.toggleLabelInternal(toggleAllLabelId, label, enabled);
         }
     }
diff --git a/Engine/Source/Programs/Horde/HordeDashboard/src/components/buildhealth/stepoutcome/StepOutcomeD3Component.tsx b/Engine/Source/Programs/Horde/HordeDashboard/src/components/buildhealth/stepoutcome/StepOutcomeD3Component.tsx
--- a/Engine/Source/Programs/Horde/HordeDashboard/src/components/buildhealth/stepoutcome/StepOutcomeD3Component.tsx
+++ b/Engine/Source/Programs/Horde/HordeDashboard/src/components/buildhealth/stepoutcome/StepOutcomeD3Component.tsx
@@ -468,7 +468,7 @@ function GenerateD3Table(stepOutcomeTable: StepOutcomeTable, handler: StepOutcom
                 .style("left", "32px")
                 .style("z-index", "1")
                 .style("background-color", modeColors.content)
-                .style("border", d => d.isLabel ? `1px dashed ${theme.semanticColors.bodyBackgroundHovered}` : `1px solid ${modeColors.background}`)
+                .style("border", `1px solid ${modeColors.background}`)
                 .style("white-space", "nowrap")
                 .style("overflow", "hidden")
                 .style("padding-left", "8px")
