[Horde][BuildHealth][Dashboard] Introduce parallel data source mode for new API endpoint

diff --git a/Engine/Source/Programs/Horde/HordeDashboard/src/backend/Api.ts b/Engine/Source/Programs/Horde/HordeDashboard/src/backend/Api.ts
--- a/Engine/Source/Programs/Horde/HordeDashboard/src/backend/Api.ts
+++ b/Engine/Source/Programs/Horde/HordeDashboard/src/backend/Api.ts
@@ -1581,6 +1581,9 @@ export interface GetJobStepRefResponse {
 	/**The job id */
 	jobId: string;
 
+	/**The step name */
+	stepName?: string;
+
 	/**The batch containing the step */
 	batchId: string;
 
@@ -1599,6 +1602,9 @@ export interface GetJobStepRefResponse {
 	/** The agent id */
 	agentId?: string;
 
+	/**State of the step. */
+	state?: JobStepState;
+
 	/**Outcome of the step, once complete. */
 	outcome?: JobStepOutcome;
 
diff --git a/Engine/Source/Programs/Horde/HordeDashboard/src/backend/Backend.ts b/Engine/Source/Programs/Horde/HordeDashboard/src/backend/Backend.ts
--- a/Engine/Source/Programs/Horde/HordeDashboard/src/backend/Backend.ts
+++ b/Engine/Source/Programs/Horde/HordeDashboard/src/backend/Backend.ts
@@ -490,6 +490,27 @@ export class Backend {
         });
     }
 
+    getJobStepHistoryBatch(streamId: string, stepNames: string[], count: number = 10, templateId: string) {
+        return new Promise<GetJobStepRefResponse[]>((resolve, reject) => {
+
+            const params: any = {
+                steps: stepNames,
+                count: count,
+                templateId: templateId
+            };
+
+            this._fetch.post(`/api/v1/streams/${streamId}/history`, {}, {
+                params: params
+            }).then((value) => {
+                let results = (value.data ?? []) as GetJobStepRefResponse[];
+                results = results.filter(r => r.outcome !== JobStepOutcome.Unspecified);
+                resolve(results);
+            }).catch(reason => {
+                reject(reason);
+            });
+        });
+    }
+
     getJobStepTrace(jobId: string, batchId: string, stepId: string) {
         return new Promise<GetJobStepTraceResponse>((resolve, reject) => {
             this._fetch.get(`/api/v1/jobs/${jobId}/batches/${batchId}/steps/${stepId}/trace`, {
diff --git a/Engine/Source/Programs/Horde/HordeDashboard/src/components/buildhealth/BuildHealthDataTypes.ts b/Engine/Source/Programs/Horde/HordeDashboard/src/components/buildhealth/BuildHealthDataTypes.ts
--- a/Engine/Source/Programs/Horde/HordeDashboard/src/components/buildhealth/BuildHealthDataTypes.ts
+++ b/Engine/Source/Programs/Horde/HordeDashboard/src/components/buildhealth/BuildHealthDataTypes.ts
@@ -12,6 +12,7 @@ export const STATES_URL_PARAM_NAME: string = "states";
 export const PARAMETER_KEY_PREFIX: string = "parameter_key";
 export const TIME_SPAN_URL_PARAM: string = "lastTimeRange";
 export const INCLUDE_PREFLIGHT_PARAM: string = "includePreflight";
+export const DEBUG_MODE_PARAM: string = "debugMode";
 export const KEY_SEPARATOR: string = "::";
 
 export const UI_SUMMARY_WARNING_AS_ERROR_PARAM: string = "warningAsError";
diff --git a/Engine/Source/Programs/Horde/HordeDashboard/src/components/buildhealth/BuildHealthOptions.ts b/Engine/Source/Programs/Horde/HordeDashboard/src/components/buildhealth/BuildHealthOptions.ts
--- a/Engine/Source/Programs/Horde/HordeDashboard/src/components/buildhealth/BuildHealthOptions.ts
+++ b/Engine/Source/Programs/Horde/HordeDashboard/src/components/buildhealth/BuildHealthOptions.ts
@@ -1,7 +1,7 @@
 // Copyright Epic Games, Inc. All Rights Reserved.
 
 import { makeObservable, observable, runInAction } from "mobx";
-import { INCLUDE_PREFLIGHT_PARAM, LABELS_URL_PARAM_ALL_LABELS, LABELS_URL_PARAM_NAME, PROJECTS_URL_PARAM_NAME, QUERY_SCHEMA_PARAM_NAME, STATES_URL_PARAM_ALL_STATES, STATES_URL_PARAM_NAME, STEPS_URL_PARAM_ALL_STEPS, STEPS_URL_PARAM_NAME, STREAMS_URL_PARAM_NAME, TEMPLATES_URL_PARAM_NAME, TIME_SPAN_URL_PARAM, UI_DATE_ANCHORS_PARAM, UI_SUMMARY_WARNING_AS_ERROR_PARAM } from "../buildhealth/BuildHealthDataTypes";
+import { DEBUG_MODE_PARAM, INCLUDE_PREFLIGHT_PARAM, LABELS_URL_PARAM_ALL_LABELS, LABELS_URL_PARAM_NAME, PROJECTS_URL_PARAM_NAME, QUERY_SCHEMA_PARAM_NAME, STATES_URL_PARAM_ALL_STATES, STATES_URL_PARAM_NAME, STEPS_URL_PARAM_ALL_STEPS, STEPS_URL_PARAM_NAME, STREAMS_URL_PARAM_NAME, TEMPLATES_URL_PARAM_NAME, TIME_SPAN_URL_PARAM, UI_DATE_ANCHORS_PARAM, UI_SUMMARY_WARNING_AS_ERROR_PARAM } from "../buildhealth/BuildHealthDataTypes";
 import { BuildHealthDataHandler } from "./BuildHealthDataHandler";
 import { decodeFullyQualifiedLabelName, decodeLabelKey, decodeStateKey, decodeStepKey, decodeTemplateKey, encodeFullyQualifiedLabelName, encodeLabelKey, encodeLabelKeyFromStrings, encodeStateKeyFromString, encodeStepKey, encodeStepNameFromStringsPreserveCase, encodeTemplateKey, encodeTemplateKeyFromStrings } from "./BuildHealthUtilities";
 import { Location } from 'react-router-dom';
@@ -38,6 +38,7 @@ type BuildHealthQueryParams = {
     stateNames: string[];
     timeSpanKey: string;
     includePreflight: boolean;
+    debugMode: boolean;
     querySchemaVersion: number;
 };
 
@@ -321,8 +322,11 @@ export function parseBuildHealthQueryParams(location: Location): { buildHealthQu
     const uiWarningsAsErrors = getBoolParam(UI_SUMMARY_WARNING_AS_ERROR_PARAM);
     const uiDateAnchors = getBoolParam(UI_DATE_ANCHORS_PARAM);
 
+    // Debug Modes
+    const debugMode = getBoolParam(DEBUG_MODE_PARAM);
+
     return {
-        buildHealthQueryParams: { projectIds, streamIds, templateIds: templateIds, labels, stepNames, stateNames, timeSpanKey, includePreflight, querySchemaVersion },
+        buildHealthQueryParams: { projectIds, streamIds, templateIds: templateIds, labels, stepNames, stateNames, timeSpanKey, includePreflight, debugMode, querySchemaVersion },
         buildHealthUIParams: { showDateAnchor: uiDateAnchors, summaryWarningAsErrors: uiWarningsAsErrors }
     };
 }
@@ -388,6 +392,10 @@ export async function loadBuildHealthOptionsFromParams(
     options.setWarningsAsSummaryFailure(uiParams.summaryWarningAsErrors);
     options.setHideDateAnchors(uiParams.showDateAnchor);
 
+    runInAction(() => {
+        optionsWriter.setDebugMode(params.debugMode);
+    })
+
     runInAction(() => {
         optionsWriter.setIncludePreflight(params.includePreflight);
     })
@@ -514,6 +522,7 @@ export class BuildHealthUIOptionsState {
     @observable warningsAsSummaryFailure: boolean = false;
     @observable linkedSelectionModeEnabled: boolean = false;
     @observable validStates: Set<string> = new Set<string>();
+    @observable debugMode: boolean = false;
 
     constructor() {
         makeObservable(this);
@@ -653,6 +662,7 @@ export class BuildHealthOptionsState {
     @observable enabledStates: Record<string, string> = {};
     @observable jobHistoryTimeSpan: TimeSpan = DEFAULT_JOB_HISTORY_TIME_SPAN;
     @observable includePreflight: boolean = false;
+    @observable debugMode: boolean = false;
 
     // #endregion -- User Selected Options --
 
diff --git a/Engine/Source/Programs/Horde/HordeDashboard/src/components/buildhealth/BuildHealthOptionsController.ts b/Engine/Source/Programs/Horde/HordeDashboard/src/components/buildhealth/BuildHealthOptionsController.ts
--- a/Engine/Source/Programs/Horde/HordeDashboard/src/components/buildhealth/BuildHealthOptionsController.ts
+++ b/Engine/Source/Programs/Horde/HordeDashboard/src/components/buildhealth/BuildHealthOptionsController.ts
@@ -80,6 +80,12 @@ export interface IBuildHealthOptionWriter {
      */
     setIncludePreflight(includePreflight: boolean): void;
 
+    /**
+     * Sets whether the system is in debug mode or not.
+     * @param debugMode True to indicate debug mode, false otherwise.
+     */
+    setDebugMode(debugMode: boolean): void;
+
     /**
      * Sets the current job history time span to utilize.
      * @param time The time span to utilize.
@@ -500,6 +506,14 @@ export class BuildHealthOptionsWriter implements IBuildHealthOptionWriter {
         this.state.includePreflight = includePreflight;
     }
 
+    /**
+     * @inheritdoc
+     */
+    @action
+    setDebugMode(debugMode: boolean) {
+        this.state.debugMode = debugMode;
+    }
+
     /**
      * @inheritdoc
      */
diff --git a/Engine/Source/Programs/Horde/HordeDashboard/src/components/buildhealth/BuildhealthView.tsx b/Engine/Source/Programs/Horde/HordeDashboard/src/components/buildhealth/BuildhealthView.tsx
--- a/Engine/Source/Programs/Horde/HordeDashboard/src/components/buildhealth/BuildhealthView.tsx
+++ b/Engine/Source/Programs/Horde/HordeDashboard/src/components/buildhealth/BuildhealthView.tsx
@@ -1260,6 +1260,7 @@ export const BuildHealthView: React.FC = observer(function ConstructBuildHealthV
         labels: buildHealthState.stepOutcomeEnabledLabelKeys,
         jobHistorySpan,
         includePreflights: buildHealthState.includePreflight,
+        debugMode: buildHealthState.debugMode
     }),
         // eslint-disable-next-line react-hooks/exhaustive-deps -- All necessary dependencies included for MobX observables
         [
@@ -1273,13 +1274,15 @@ export const BuildHealthView: React.FC = observer(function ConstructBuildHealthV
     const uiOptions = useMemo(() => ({
         includeDateAnchors: buildHealthUIState.includeDateAnchors,
         warningsAsSummaryFailure: buildHealthUIState.warningsAsSummaryFailure,
-        validStates: buildHealthUIState.validStates
+        validStates: buildHealthUIState.validStates,
+        debugMode: buildHealthUIState.debugMode
     }),
         // eslint-disable-next-line react-hooks/exhaustive-deps -- All necessary dependencies included for MobX observables
         [
             buildHealthUIState.includeDateAnchors,
             buildHealthUIState.warningsAsSummaryFailure,
-            buildHealthUIState.validStates
+            buildHealthUIState.validStates,
+            buildHealthUIState.debugMode
         ]);
 
     return (
diff --git a/Engine/Source/Programs/Horde/HordeDashboard/src/components/buildhealth/stepoutcome/StepOutcomeDataHandler.ts b/Engine/Source/Programs/Horde/HordeDashboard/src/components/buildhealth/stepoutcome/StepOutcomeDataHandler.ts
--- a/Engine/Source/Programs/Horde/HordeDashboard/src/components/buildhealth/stepoutcome/StepOutcomeDataHandler.ts
+++ b/Engine/Source/Programs/Horde/HordeDashboard/src/components/buildhealth/stepoutcome/StepOutcomeDataHandler.ts
@@ -1,10 +1,10 @@
 // Copyright Epic Games, Inc. All Rights Reserved.
 
-import { GetJobResponse, IssueQuery, JobStepState, JobStreamQuery, LabelState } from "horde/backend/Api";
+import { GetJobResponse, GetJobStepRefResponse, IssueQuery, JobStepState, JobStreamQuery, LabelState } from "horde/backend/Api";
 import { PollBase } from "horde/backend/PollBase";
-import { JobResponseData, LabelOutcomeTableEntry, StepOutcomeTable, StepOutcomeTableEntry, UpdatedChangeSummaryData } from "./StepOutcomeDataTypes";
+import { GetJobStepRefResponseWrapper, GetStepResponseWrapper, JobResponseData, LabelOutcomeTableEntry, StepOutcomeTable, StepOutcomeTableEntry, UpdatedChangeSummaryData } from "./StepOutcomeDataTypes";
 import backend from "horde/backend";
-import { decodeLabelKey, decodeStepKey, decodeTemplateKey, encodeLabelKeyFromStrings, encodeStepNameFromStrings } from "../BuildHealthUtilities";
+import { decodeLabelKey, decodeStepKey, decodeTemplateKey, encodeLabelKeyFromStrings, encodeStepNameFromStrings, encodeStepNameFromStringsPreserveCase } from "../BuildHealthUtilities";
 import { action, computed, makeObservable, observable, runInAction } from "mobx";
 import { summarizeJob } from "./StepOutcomeUtilities";
 
@@ -99,6 +99,7 @@
     labels?: string[];
     maxJobCount?: number;
     includePreflights: boolean;
+    debugMode: boolean;
 }
 
 // #endregion -- API Data Types & Utilities --
@@ -143,6 +144,119 @@
 
     //#endregion -- Constructor --
 
+    // #region -- Private API --
+
+    private constructStreamJobQueries(templateId: string, jobCount: number, modifiedAfter: string, queries: { streamId: string; query: JobStreamQuery; }[], streamId: string, debugMode: boolean) {
+        let completeFilter: string = !debugMode ? `${this.STREAM_JOBS_FILTER},${this.STREAM_JOBS_BATCH_STEP_FILTER}` : `${this.STREAM_JOBS_FILTER}`
+        const query: JobStreamQuery = {
+            template: [templateId],
+            count: jobCount,
+            filter: completeFilter,
+            modifiedAfter: modifiedAfter
+        };
+
+        queries.push({ streamId, query });
+    }
+
+    private constructJobStepHistoryQueries(stepMap: Map<string, Map<string, Set<string>>>, streamId: string, templateId: string, jobStepHistoryBatchQueries: { streamId: string; stepNames: string[]; count: number; templateId: string; }[]) {
+        let templateMap: Map<string, Set<string>> = stepMap.get(streamId)!;
+        let stepSet = templateMap?.get(templateId);
+        if (stepSet) {
+            jobStepHistoryBatchQueries.push({
+                streamId: streamId, stepNames: Array.from(stepSet).map(x => {
+                    let { streamId, templateId, stepName } = decodeStepKey(x);
+                    return stepName;
+                }), count: 100, templateId: templateId
+            });
+        }
+    }
+
+    private async issueAllStreamJobQueries(streamJobQueries: { streamId: string; query: JobStreamQuery; }[]) {
+        return await Promise.all(
+            streamJobQueries.map(({ streamId, query }) => backend.getStreamJobs(streamId, query))
+        );
+    }
+
+    private async issueAllJobStepHistoryQueries(jobStepHistoryBatchQueries: { streamId: string; stepNames: string[]; count: number; templateId: string; }[]) {
+        return await Promise.all(
+            jobStepHistoryBatchQueries.map(({ streamId, stepNames, count, templateId }) => backend.getJobStepHistoryBatch(streamId, stepNames, count, templateId))
+        );
+    }
+
+    private processJobStepResults(results: GetJobResponse[][], stepMap: Map<string, Map<string, Set<string>>>, labelMap: Map<string, Map<string, Set<string>>>, jobStepKeyToState: Map<string, JobStepState>) {
+        for (const streamResult of results) {
+            // Process all job responses
+            for (const jobResponse of streamResult) {
+                // Filter preflights as appropriate
+                if (!this.activeFilter.includePreflights && (jobResponse.preflightChange || jobResponse.preflightDescription)) {
+                    continue;
+                }
+
+                const batchResponses = jobResponse.batches ?? [];
+                let streamStepSet = stepMap.get(jobResponse.streamId);
+                let streamLabelSet = labelMap.get(jobResponse.streamId);
+                let templateStepSet = streamStepSet ? streamStepSet.get(jobResponse.templateId!) : undefined;
+                let templateLabelSet = streamLabelSet ? streamLabelSet.get(jobResponse.templateId!) : undefined;
+
+                if (templateStepSet && templateStepSet.size >= 0) {
+                    for (const batch of batchResponses) {
+                        for (const step of batch.steps) {
+                            let unifiedKey: string = jobResponse.id + step.id;
+                            if (!jobStepKeyToState.has(unifiedKey)) {
+                                jobStepKeyToState.set(unifiedKey, step.state);
+                            }
+
+                            let fullyQualifiedStepName = encodeStepNameFromStringsPreserveCase(jobResponse.streamId, jobResponse.templateId!, step.name);
+
+                            if (templateStepSet &&
+                                templateStepSet.size >= 0 &&
+                                templateStepSet.has(fullyQualifiedStepName)) {
+                                this.stepOutcomeTableData.addEntry(new StepOutcomeTableEntry(new GetStepResponseWrapper(step), jobResponse.id, jobResponse.name, jobResponse.streamId, jobResponse.createTime.toString(), jobResponse.change, this.changeData.get(jobResponse.change!)?.dateUtc));
+                            }
+                        }
+                    }
+                }
+
+                // Do not compute label steps if there are no steps for the template.
+                if (templateLabelSet && templateLabelSet.size >= 0) {
+                    const labelResponses = jobResponse.labels ?? [];
+                    for (let idx: number = 0; idx < labelResponses.length; ++idx) {
+                        let label = labelResponses[idx];
+                        let synthesizedStepOutcome: JobStepState | LabelState = LabelState.Unspecified;
+
+                        // Labels lack sufficient cancellation & skipped context, and inferring this from their .steps member is often incorrect (steps may not be present)
+                        // We therefore apply a simple heuristic: if there is cancellation info - we anotate the synthesized step as aborted; else skipped.
+                        if (label.state === LabelState.Unspecified) {
+                            synthesizedStepOutcome = jobResponse.abortedByUserInfo !== undefined ? JobStepState.Aborted : JobStepState.Skipped;
+                        }
+
+                        let fullyQualifiedLabelName = encodeLabelKeyFromStrings(jobResponse.streamId, jobResponse.templateId!, label.dashboardCategory, label.dashboardName).toLocaleLowerCase();
+
+                        if (templateLabelSet &&
+                            templateLabelSet.size >= 0 &&
+                            templateLabelSet.has(fullyQualifiedLabelName)) {
+                            this.stepOutcomeTableData.addEntry(new LabelOutcomeTableEntry(label, idx, jobResponse.id, jobResponse.name, jobResponse.streamId, jobResponse.createTime.toString(), jobResponse.change, this.changeData.get(jobResponse.change!)?.dateUtc, synthesizedStepOutcome));
+                        }
+                    }
+                }
+            }
+        }
+    }
+
+    private processJobStepHistoryResults(jobStepHistoryResults: GetJobStepRefResponse[][]) {
+        for (const item of jobStepHistoryResults) {
+            for (const subitem of item) {
+                if (this.jobData.has(subitem.jobId)) {
+                    let jobName: string = this.jobData.get(subitem.jobId)?.name!;
+                    let streamId: string = this.jobData.get(subitem.jobId)?.streamId!;
+                    this.stepOutcomeTableData.addEntry(new StepOutcomeTableEntry(new GetJobStepRefResponseWrapper(subitem), subitem.jobId, jobName, streamId, subitem.jobStartTime.toString(), subitem.change, this.changeData.get(subitem.change!)?.dateUtc));
+                }
+            }
+        }
+    }
+
+    // #endregion -- Private API --
+
     // #region -- Interface --
 
     @action
@@ -180,7 +294,7 @@
                 templateMap.set(templateId, new Set());
             }
 
-            templateMap.get(templateId)!.add(step.toLocaleLowerCase());
+            templateMap.get(templateId)!.add(step);
         });
 
         (this.activeFilter.labels ?? []).forEach(label => {
@@ -203,6 +317,7 @@
         const modifiedAfter = this.activeFilter.jobHistorySpan.start.toISOString();
         const jobCount = this.activeFilter.maxJobCount ?? StepOutcomeDataHandler.MAX_JOB_COUNT_DEFAULT;
         const jobStepKeyToState = new Map<string, JobStepState>();
+
         // Group jobs by streamId
         const jobsByStream: Record<string, string[]> = {};
 
@@ -215,33 +330,28 @@
             jobsByStream[streamId].push(templateId);
         });
 
-        const queries: { streamId: string; query: JobStreamQuery }[] = [];
         const changeQueries: { streamId: string; query: any }[] = [];
+        const streamJobQueries: { streamId: string; query: JobStreamQuery }[] = [];
+        const jobStepHistoryBatchQueries: { streamId: string; stepNames: string[], count: number, templateId: string }[] = [];
 
+        // Construct queries
         for (const streamId of streamList) {
             const jobsForStream = jobsByStream[streamId] ?? [];
             for (const templateId of jobsForStream) {
-                const query: JobStreamQuery = {
-                    template: [templateId],
-                    count: jobCount,
-                    filter: `${this.STREAM_JOBS_FILTER},${this.STREAM_JOBS_BATCH_STEP_FILTER}`,
-                    modifiedAfter: modifiedAfter
-                };
-                queries.push({ streamId, query });
                 changeQueries.push({ streamId: streamId, query: { filter: this.STREAM_CHANGES_FILTER } });
+                this.constructStreamJobQueries(templateId, jobCount, modifiedAfter, streamJobQueries, streamId, this.activeFilter.debugMode);
+                this.constructJobStepHistoryQueries(stepMap, streamId, templateId, jobStepHistoryBatchQueries);
             }
         }
 
-        // Parallelize backend calls for stream jobs
-        const results = await Promise.all(
-            queries.map(({ streamId, query }) => backend.getStreamJobs(streamId, query))
-        );
+        const jobStepHistoryResults = this.activeFilter.debugMode ? await this.issueAllJobStepHistoryQueries(jobStepHistoryBatchQueries) : undefined;
+        const streamJobResults = await this.issueAllStreamJobQueries(streamJobQueries);
 
         // Annotate the changes that we need to request in getChangeSummaries.
         let jobHistoryCount: Map<string, number[]> = new Map<string, number[]>();
 
-        results.forEach((jobs, i) => {
-            const streamId = queries[i].streamId;
+        streamJobResults.forEach((jobs, i) => {
+            const streamId = streamJobQueries[i].streamId;
 
             const newChanges = jobs
                 .map(job => job.change!)
@@ -257,15 +367,16 @@
         });
 
         // Flatten all jobData & memoize the summarized outcome
-        this.jobData = new Map<string, JobResponseData>(results.flat().map((job: GetJobResponse): [string, JobResponseData] => (
-            [
+        this.jobData = new Map<string, JobResponseData>(streamJobResults.flat().map((job: GetJobResponse): [string, JobResponseData] => {
+            return [
                 job.id,
                 {
                     ...job,
                     summarizedOutcome: summarizeJob(job)
                 }
             ]
-        )));
+        }
+        ));
 
         if (jobHistoryCount.size > 0) {
             // Parallelize backend calls for change summaries
@@ -283,66 +394,11 @@
             console.info(`Retrieved: ${flattenedResults.length} change summary results for StepOutcome query.`);
         }
 
-        for (const streamResult of results) {
-            // Process all job responses
-            for (const jobResponse of streamResult) {
-                // Filter preflights as appropriate
-                if (!this.activeFilter.includePreflights && (jobResponse.preflightChange || jobResponse.preflightDescription)) {
-                    continue;
-                }
-
-                const batchResponses = jobResponse.batches ?? [];
-                let streamStepSet = stepMap.get(jobResponse.streamId);
-                let streamLabelSet = labelMap.get(jobResponse.streamId);
-                let templateStepSet = streamStepSet ? streamStepSet.get(jobResponse.templateId!) : undefined;
-                let templateLabelSet = streamLabelSet ? streamLabelSet.get(jobResponse.templateId!) : undefined;
-
-                if (templateStepSet && templateStepSet.size >= 0) {
-                    for (const batch of batchResponses) {
-                        for (const step of batch.steps) {
-                            let unifiedKey: string = jobResponse.id + step.id;
-                            if (!jobStepKeyToState.has(unifiedKey)) {
-                                jobStepKeyToState.set(unifiedKey, step.state);
-                            }
-
-                            let fullyQualifiedStepName = encodeStepNameFromStrings(jobResponse.streamId, jobResponse.templateId!, step.name).toLocaleLowerCase();
-
-                            if (
-                                templateStepSet &&
-                                templateStepSet.size >= 0 &&
-                                templateStepSet.has(fullyQualifiedStepName)
-                            ) {
-                                this.stepOutcomeTableData.addEntry(new StepOutcomeTableEntry(step, jobResponse.id, jobResponse.name, jobResponse.streamId, jobResponse.createTime.toString(), jobResponse.change, this.changeData.get(jobResponse.change!)?.dateUtc));
-                            }
-                        }
-                    }
-                }
-
-                // Do not compute label steps if there are no steps for the template.
-                if (templateLabelSet && templateLabelSet.size >= 0) {
-                    const labelResponses = jobResponse.labels ?? [];
-                    for (let idx: number = 0; idx < labelResponses.length; ++idx) {
-                        let label = labelResponses[idx];
-                        let synthesizedStepOutcome: JobStepState | LabelState = LabelState.Unspecified;
-
-                        // Labels lack sufficient cancellation & skipped context, and inferring this from their .steps member is often incorrect (steps may not be present)
-                        // We therefore apply a simple heuristic: if there is cancellation info - we anotate the synthesized step as aborted; else skipped.
-                        if (label.state === LabelState.Unspecified) {
-                            synthesizedStepOutcome = jobResponse.abortedByUserInfo !== undefined ? JobStepState.Aborted : JobStepState.Skipped;
-                        }
-
-                        let fullyQualifiedLabelName = encodeLabelKeyFromStrings(jobResponse.streamId, jobResponse.templateId!, label.dashboardCategory, label.dashboardName).toLocaleLowerCase();
-
-                        if (
-                            templateLabelSet &&
-                            templateLabelSet.size >= 0 &&
-                            templateLabelSet.has(fullyQualifiedLabelName)
-                        ) {
-                            this.stepOutcomeTableData.addEntry(new LabelOutcomeTableEntry(label, idx, jobResponse.id, jobResponse.name, jobResponse.streamId, jobResponse.createTime.toString(), jobResponse.change, this.changeData.get(jobResponse.change!)?.dateUtc, synthesizedStepOutcome));
-                        }
-                    }
-                }
-            }
+        if (this.activeFilter.debugMode && jobStepHistoryResults) {
+            this.processJobStepHistoryResults(jobStepHistoryResults);
+        }
+        else if (streamJobResults) {
+            this.processJobStepResults(streamJobResults, stepMap, labelMap, jobStepKeyToState);
         }
 
         this.StepOutcomeTableData.summarize();
@@ -362,6 +418,10 @@
         return Promise.resolve();
     }
 
+    // #endregion -- Interface --
+
+    // #region -- Public API  --
+
     /**
      * Will populate the provided stepOutcomeTableEntry's issue data field,
      * @param stepOutcomeTableEntry The step outcome table entry to attempt to populate with issue data.
@@ -379,9 +439,6 @@
         return Promise.resolve();
     }
 
-    // #endregion -- Interface --
 
-    // #region -- Public API  --
-
     /**
      * Gets the last refresh date.
diff --git a/Engine/Source/Programs/Horde/HordeDashboard/src/components/buildhealth/stepoutcome/StepOutcomeDataTypes.ts b/Engine/Source/Programs/Horde/HordeDashboard/src/components/buildhealth/stepoutcome/StepOutcomeDataTypes.ts
--- a/Engine/Source/Programs/Horde/HordeDashboard/src/components/buildhealth/stepoutcome/StepOutcomeDataTypes.ts
+++ b/Engine/Source/Programs/Horde/HordeDashboard/src/components/buildhealth/stepoutcome/StepOutcomeDataTypes.ts
@@ -1,6 +1,6 @@
 // Copyright Epic Games, Inc. All Rights Reserved.
 
-import { ChangeSummaryData, GetJobResponse, GetLabelStateResponse, GetStepResponse, IssueData, JobStepOutcome, JobStepState, LabelOutcome, LabelState } from "horde/backend/Api";
+import { ChangeSummaryData, GetJobResponse, GetJobStepRefResponse, GetLabelStateResponse, GetStepResponse, IssueData, JobStepOutcome, JobStepState, LabelOutcome, LabelState } from "horde/backend/Api";
 import { computeDuration, DURATION_NOT_SET, formatDurationFromMs, toDate } from "./StepOutcomeUtilities";
 import { StepOutcomeDataHandler, StepOutcomeViewOptions } from "./StepOutcomeDataHandler";
 import { action, makeObservable, observable } from "mobx";
@@ -60,13 +60,6 @@ export type UpdatedChangeSummaryData = ChangeSummaryData & {
     id: ChangeId
 }
 
-/*
- * Convenience type to fold in extra metadata relevant for the StepResponse.
- */
-export type StepReponseData = GetStepResponse & {
-    stepDurationMs: number;
-}
-
 /*
 * Convenience type to fold in extra metadata relevant for the JobResponse.
 */
@@ -356,13 +349,96 @@ export class LabelOutcomeTableEntry extends OutcomeTableEntry {
     // #endregion -- Abstract API --
 }
 
+/**
+ * Ease of use type to support common property access for @see IJobStepAdapter .
+ */
+export type JobStepResponseUnion = GetJobStepRefResponse | GetStepResponse;
+
+/**
+ * Interface that describes the adapter class for the stepId and stepName properties.
+ * This is necessary since the @see GetJobStepRefResponse and @see GetStepResponse have disimilar interfaces into the underlying step structure.
+ */
+export interface IJobStepAdapter {
+    /**
+     * The step's id.
+     */
+    stepId: string;
+
+    /**
+     * The step's name.
+     */
+    stepName: string;
+
+    /**
+     * Gets the base union type for ease of access for common properties.
+     */
+    getBase(): JobStepResponseUnion;
+}
+
+/**
+ * Wrapper implementation for @see GetJobStepRefResponse .
+ */
+export class GetJobStepRefResponseWrapper implements IJobStepAdapter {
+    constructor(private base: GetJobStepRefResponse) { }
+
+    /**
+     * @inheritdoc
+     */
+    get stepId() {
+        return this.base.stepId;
+    }
+
+    /**
+     * @inheritdoc
+     */
+    get stepName() {
+        return this.base.stepName!;
+    }
+
+    /**
+     * @inheritdoc
+     */
+    getBase(): GetJobStepRefResponse {
+        return this.base;
+    }
+}
+
+/**
+ * Wrapper implementation for @see GetStepResponse .
+ */
+export class GetStepResponseWrapper implements IJobStepAdapter {
+    constructor(private base: GetStepResponse) { }
+
+    /**
+     * @inheritdoc
+     */
+    get stepId() {
+        return this.base.id;
+    }
+
+    /**
+     * @inheritdoc
+     */
+    get stepName() {
+        return this.base.name;
+    }
+
+    /**
+     * @inheritdoc
+     */
+    getBase(): GetStepResponse {
+        return this.base;
+    }
+}
+
 /**
  * Model for a table entry.
  */
 export class StepOutcomeTableEntry extends OutcomeTableEntry {
     // #region -- Private Members --
 
-    private stepResponse: StepReponseData;
+    private stepResponse: IJobStepAdapter;
+    private stepDurationMs : number;
 
     // #endregion -- Private Members --
 
@@ -378,10 +454,11 @@ export class StepOutcomeTableEntry extends OutcomeTableEntry {
      * @param change The change associated with the table entry.
      * @param changeDate The change submit date associated with the table entry.
      */
-    constructor(stepResponse: GetStepResponse, jobId: string, jobName: string, streamId: string, jobCreateTime: string, change?: number, changeDate?: Date | string) {
+    constructor(stepResponse: IJobStepAdapter, jobId: string, jobName: string, streamId: string, jobCreateTime: string, change?: number, changeDate?: Date | string) {
         super(jobId, jobName, streamId, jobCreateTime, change, changeDate);
 
-        this.stepResponse = { ...stepResponse, stepDurationMs: computeDuration(stepResponse.startTime, stepResponse.finishTime) };
+        this.stepResponse = stepResponse;
+        this.stepDurationMs = computeDuration(stepResponse.getBase().startTime, stepResponse.getBase().finishTime)
     }
 
     // #endregion -- Constructor --
@@ -392,42 +469,35 @@ export class StepOutcomeTableEntry extends OutcomeTableEntry {
      * @inheritdoc
      */
     get id(): string {
-        return this.stepResponse.id;
+        return this.stepResponse.stepId;
     }
 
     /**
      * @inheritdoc
      */
     get name(): string {
-        return this.stepResponse?.name ?? "";
-    }
-
-    /**
-     * @inheritdoc
-     */
-    get starTime() {
-        return this.stepResponse?.startTime;
+        return this.stepResponse?.stepName ?? "";
     }
 
     /**
      * @inheritdoc
      */
     get state(): JobStepState {
-        return this.stepResponse.state;
+        return this.stepResponse.getBase().state!;
     }
 
     /**
      * @inheritdoc
      */
     get outcome(): JobStepOutcome {
-        return this.stepResponse.outcome;
+        return this.stepResponse.getBase().outcome!;
     }
 
     /**
      * @inheritdoc
      */
     get startTime() {
-        return this.stepResponse.startTime;
+        return this.stepResponse.getBase().startTime;
     }
 
     /**
@@ -436,8 +506,8 @@ export class StepOutcomeTableEntry extends OutcomeTableEntry {
     toCSVRepresntation(relatedJobContext: JobResponseData | undefined): string {
         let baseUrl = window.location.origin;
         let jobUrl = `${baseUrl}/job/${this.jobId}`;
-        let stepUrl = `${baseUrl}/job/${this.jobId}?step=${this.stepResponse.id}`;
-        return `${this.change},${this.streamId},${this.jobName},${this.jobId},${this.stepResponse.name},${this.stepResponse.id},${this.stepResponse.state},${this.stepResponse.outcome},${this.stepResponse.stepDurationMs / 1000},${this.jobCreateTime},${this.stepResponse.startTime},${relatedJobContext?.summarizedOutcome},${relatedJobContext?.state},${jobUrl},${stepUrl}\n`;
+        let stepUrl = `${baseUrl}/job/${this.jobId}?step=${this.stepResponse.stepId}`;
+        return `${this.change},${this.streamId},${this.jobName},${this.jobId},${this.stepResponse.stepName!},${this.stepResponse.stepId},${this.stepResponse.getBase().state},${this.stepResponse.getBase().outcome},${this.stepDurationMs / 1000},${this.jobCreateTime},${this.stepResponse.getBase().startTime},${relatedJobContext?.summarizedOutcome},${relatedJobContext?.state},${jobUrl},${stepUrl}\n`;
     }
 
     /**
@@ -456,10 +526,10 @@ export class StepOutcomeTableEntry extends OutcomeTableEntry {
      * @returns The string representation of the step duration.
      */
     getDurationString(): string {
-        let durationStr: string = this.stepResponse.state === SkippedStr ? NOT_APPLICABLE_LABEL : (this.stepResponse.state === WaitingStr || this.stepResponse.state === ReadyStr ? "Not Started" : "Ongoing");
+        let durationStr: string = this.stepResponse.getBase().state === SkippedStr ? NOT_APPLICABLE_LABEL : (this.stepResponse.getBase().state === WaitingStr || this.stepResponse.getBase().state === ReadyStr ? "Not Started" : "Ongoing");
 
-        if (this.stepResponse.stepDurationMs != DURATION_NOT_SET) {
-            durationStr = formatDurationFromMs(this.stepResponse.stepDurationMs);
+        if (this.stepDurationMs != DURATION_NOT_SET) {
+            durationStr = formatDurationFromMs(this.stepDurationMs);
         }
 
         return durationStr;
