Horde: Update log message to include Agent ID when lease cannot be assigned due to conform required

#rnx

diff --git a/Engine/Source/Programs/Horde/Plugins/Build/HordeServer.Build/Jobs/JobTaskSource.cs b/Engine/Source/Programs/Horde/Plugins/Build/HordeServer.Build/Jobs/JobTaskSource.cs
--- a/Engine/Source/Programs/Horde/Plugins/Build/HordeServer.Build/Jobs/JobTaskSource.cs
+++ b/Engine/Source/Programs/Horde/Plugins/Build/HordeServer.Build/Jobs/JobTaskSource.cs
@@ -815,7 +815,7 @@ namespace HordeServer.Jobs
 				// Won't be assigning a lease to this waiter
 				if (waiter.LeaseSource.TrySetResult(null))
 				{
-					_logger.LogInformation("Failed to assign job/batch to agent. Conform requested");
+					_logger.LogInformation("Failed to assign job/batch to {AgentId}. Conform requested on agent", agent.Id);
 
 					// Don't try to assign another item to this waiter to avoid duplicating a conform request
 					waiter.Skip = true;
