Horde: Add additional logging to conform check

#rnx

diff --git a/Engine/Source/Programs/Horde/Plugins/Build/HordeServer.Build/Jobs/JobTaskSource.cs b/Engine/Source/Programs/Horde/Plugins/Build/HordeServer.Build/Jobs/JobTaskSource.cs
--- a/Engine/Source/Programs/Horde/Plugins/Build/HordeServer.Build/Jobs/JobTaskSource.cs
+++ b/Engine/Source/Programs/Horde/Plugins/Build/HordeServer.Build/Jobs/JobTaskSource.cs
@@ -778,10 +778,10 @@ namespace HordeServer.Jobs
 			}
 
 			// When a call to JobTaskSource.AssignLeaseAsync by the AgentService requests conform it will not be actioned
-			// by ConformTaskSource until the next iteration of AgentService.WaitForLeaseAsync. Make sure a job/batch is
-			// not assigned on that next iteration because it could take priority over the conform task.
+			// by ConformTaskSource until the next iteration of AgentService.WaitForLeaseAsync.
 			if (agent.ConformAttemptCount is null or 0 && agent.RequestFullConform)
 			{
+				_logger.LogInformation("Pending conform already requested on {AgentId}", agent.Id);
 				return true;
 			}
 
