[Horde][Analytics][Server]  Adding context to the issue telemetry.

diff --git a/Engine/Source/Programs/Horde/Plugins/Build/HordeServer.Build/Issues/IssueTelemetry.cs b/Engine/Source/Programs/Horde/Plugins/Build/HordeServer.Build/Issues/IssueTelemetry.cs
--- a/Engine/Source/Programs/Horde/Plugins/Build/HordeServer.Build/Issues/IssueTelemetry.cs
+++ b/Engine/Source/Programs/Horde/Plugins/Build/HordeServer.Build/Issues/IssueTelemetry.cs
@@ -14,17 +14,22 @@ namespace HordeServer.Issues
 	/// Record used for Issue telemetry. Models a <see cref="IIssue"/>.
 	/// </summary>
 	/// <remarks>Modifications to this record's structure should be versioned appropriately through <see cref="SchemaVersion"/> via <see cref="CurrentSchemaVersion"/>, and call sites should continue to publish in a backwards compatible manner.</remarks>
-	public record IssueTelemetry(string EventName, int SchemaVersion, int Id, DateTime? AcknowledgedAt, DateTime CreatedAt, UserId? OwnerId, string? ExternalIssueKey, CommitId? FixChange, bool FixedSystemic, CommitId? RootChange, DateTime LastSeenAt, DateTime? NominatedAt, UserId? NominatedById, DateTime? ResolvedAt, UserId? ResolvedById, IssueSeverity Severity, string Summary, DateTime? VerifiedAt, int? DuplicateIssueId, string? RootCauseCategory, string? RootCauseSummary)
+	public record IssueTelemetry(string EventName, int SchemaVersion, string Context, int Id, DateTime? AcknowledgedAt, DateTime CreatedAt, UserId? OwnerId, string? ExternalIssueKey, CommitId? FixChange, bool FixedSystemic, CommitId? RootChange, DateTime LastSeenAt, DateTime? NominatedAt, UserId? NominatedById, DateTime? ResolvedAt, UserId? ResolvedById, IssueSeverity Severity, string Summary, DateTime? VerifiedAt, int? DuplicateIssueId, string? RootCauseCategory, string? RootCauseSummary)
 	{
 		/// <summary>
 		/// Default event name for the IssueTelemetry.
 		/// </summary>
-		public static readonly string DefaultEventName = "State.Issue";
+		public static readonly string DefaultEventName = "State.Issue.Dev";
+
+		/// <summary>
+		/// Default context for the Horde Issue Telemetry.
+		/// </summary>
+		public static readonly string DefaultContext = "Horde";
 
 		/// <summary>
 		/// Current schema version.
 		/// </summary>
-		public static readonly int CurrentSchemaVersion = 2;
+		public static readonly int CurrentSchemaVersion = 3;
 
 		/// <summary>
 		/// Telemetry payload for <see cref="IIssue"/>.
@@ -33,6 +38,7 @@ namespace HordeServer.Issues
 			: this(
 					DefaultEventName,
 					CurrentSchemaVersion,
+					DefaultContext,
 					issue.Id,
 					issue.AcknowledgedAt,
 					issue.CreatedAt,
