[Horde][BuildHealth][Dashboard] Optionally re-open issue when assigning to other user

diff --git a/Engine/Source/Programs/Horde/HordeDashboard/src/components/IssueViewV2.tsx b/Engine/Source/Programs/Horde/HordeDashboard/src/components/IssueViewV2.tsx
--- a/Engine/Source/Programs/Horde/HordeDashboard/src/components/IssueViewV2.tsx
+++ b/Engine/Source/Programs/Horde/HordeDashboard/src/components/IssueViewV2.tsx
@@ -2056,7 +2056,7 @@ export const IssueModalV2: React.FC<{ popHistoryOnClose: boolean, issueId?: stri
 
 const AssignToOtherModal: React.FC<{ defaultUser?: GetUserResponse, onClose: () => void }> = ({ defaultUser, onClose }) => {
 
-   const [state, setState] = useState<{ error?: string, submitting?: boolean, userId?: string }>({});
+   const [state, setState] = useState<{ error?: string, submitting?: boolean, userId?: string, reopen?: boolean }>({});
 
    const { hordeClasses } = getHordeStyling();
 
@@ -2064,9 +2064,8 @@ const AssignToOtherModal: React.FC<{ defaultUser?: GetUserResponse, onClose: ()
    const issueId = issue.id;
 
    let headerText = `Assign Issue ${issueId}`;
-   if (issue.resolvedAt) {
-      headerText = `Reopen Issue ${issueId} and Assign?`;
-   }
+
+   let showReopenOption = issue.resolvedAt ? true: false;
 
    const close = () => {
       onClose();
@@ -2086,12 +2085,15 @@ const AssignToOtherModal: React.FC<{ defaultUser?: GetUserResponse, onClose: ()
          return;
       }
 
+      const resolved = state.reopen ? false : undefined;
+      const fixChange = state.reopen ? 0 : undefined;
+
       const request: UpdateIssueRequest = {
          nominatedById: (!userId || dashboard.userId === userId) ? null : dashboard.userId,
          ownerId: userId ? userId : "",
          acknowledged: dashboard.userId === userId ? true : false,
-         resolved: false,
-         fixChange: 0
+         resolved: resolved,
+         fixChange: fixChange
       };
 
       try {
@@ -2112,7 +2114,10 @@ const AssignToOtherModal: React.FC<{ defaultUser?: GetUserResponse, onClose: ()
 
    };
 
-   const height = state.error ? 265 : 225;
+   let height = state.error ? 265 : 225;
+   if (showReopenOption) {
+      height += 35;
+   }
 
    return <Modal isOpen={true} className={hordeClasses.modal} styles={{ main: { padding: 8, width: 540, height: height, minHeight: height } }} onDismiss={() => { close() }}>
 
@@ -2142,6 +2147,14 @@ const AssignToOtherModal: React.FC<{ defaultUser?: GetUserResponse, onClose: ()
                <Text variant="small" style={{ paddingTop: 4 }}>Leave blank to assign to default owner</Text>
             </Stack>
          </Stack>
+         
+         {showReopenOption && <Stack styles={{ root: { paddingTop: 20} }}>
+            <Checkbox label='Reopen issue?'
+               defaultChecked={false}
+               onChange={(ev, newValue) => {
+                  setState({ ...state, reopen: !!newValue});
+               }} />
+         </Stack>}
 
          <Stack horizontal tokens={{ childrenGap: 16 }} styles={{ root: { paddingTop: 24, paddingLeft: 8, paddingBottom: 8 } }}>
             <Stack grow />
