[Horde][BuildHealth][Server] Only assign issue to user if previously unassigned when marking fixed via Slack workflow

diff --git a/Engine/Source/Programs/Horde/Plugins/Build/HordeServer.Build/Notifications/Sinks/SlackNotificationSink.cs b/Engine/Source/Programs/Horde/Plugins/Build/HordeServer.Build/Notifications/Sinks/SlackNotificationSink.cs
--- a/Engine/Source/Programs/Horde/Plugins/Build/HordeServer.Build/Notifications/Sinks/SlackNotificationSink.cs
+++ b/Engine/Source/Programs/Horde/Plugins/Build/HordeServer.Build/Notifications/Sinks/SlackNotificationSink.cs
@@ -3226,9 +3226,13 @@ namespace HordeServer.Notifications.Sinks
 						if (TryMatch(payload.View.CallbackId, @"^issue_(\d+)_markfixed_([a-fA-F0-9]{24})$", out match))
 						{
 							int issueId = Int32.Parse(match.Groups[1].Value, NumberStyles.Integer, CultureInfo.InvariantCulture);
+							IIssue? issue = await _issueService.Collection.GetIssueAsync(issueId, cancellationToken);
+
 							UserId userId = UserId.Parse(match.Groups[2].Value);
 
-							UserId? resolvedById = userId;
+							// Set the issue as resolved by the current user if it was previously unassigned, otherwise leave resolvedById, and consequently OwnerId, unchanged.
+							// This is to facilitate workflow where triage personnel resolve an issue that they aren't the owner of. If it was unassigned, they'd like to then become the owner, so they can assign it out.
+							UserId? resolvedById = issue?.OwnerId != null ? null: userId;
 							if (payload.View.State.TryGetValue("fixed_by", "fixed_by_action", out string? fixedByStr))
 							{
 								UserId fixedById;
@@ -3243,7 +3247,7 @@ namespace HordeServer.Notifications.Sinks
 							if (payload.View.State.TryGetValue("fix_cl", "fix_cl_action", out fixChangeStr))
 							{
 								fixCommitId = new CommitId(fixChangeStr);
-								_logger.LogInformation("Marking issue {IssueId} fixed by user {UserId} in {Change}", issueId, resolvedById, fixCommitId);
+								_logger.LogInformation("Marking issue {IssueId} fixed by user {UserId} in {Change}", issueId, resolvedById ?? userId, fixCommitId);
 							}
 
 							string? rootChange;
@@ -3320,8 +3324,6 @@ namespace HordeServer.Notifications.Sinks
 								dupeIssueId = -1;
 							}
 
-							IIssue? issue = await _issueService.Collection.GetIssueAsync(issueId, cancellationToken);
-
 							string? summary = null;
 
 							if (issue != null)
@@ -3340,7 +3342,7 @@ namespace HordeServer.Notifications.Sinks
 								summary = curSummary;
 							}
 
-							await _issueService.UpdateIssueAsync(issueId, summary: summary, fixCommitId: fixCommitId, resolvedById: fixCommitId != null ? resolvedById : null, initiatedById: userId, duplicateIssueId: dupeIssueId, rootCommitId: rootCommit, rootCauseCategory: rootCategory, rootCauseSummary: rootSummary, cancellationToken: cancellationToken);
+							await _issueService.UpdateIssueAsync(issueId, summary: summary, fixCommitId: fixCommitId, resolvedById: resolvedById, initiatedById: userId, duplicateIssueId: dupeIssueId, rootCommitId: rootCommit, rootCauseCategory: rootCategory, rootCauseSummary: rootSummary, cancellationToken: cancellationToken);
 
 						}
 						else if (TryMatch(payload.View.CallbackId, @"^issue_(\d+)_ack_([a-fA-F0-9]{24})$", out match))
