Horde: Update message when lease cannot be assigned to agent due to a pending conform

#rnx

diff --git a/Engine/Source/Programs/Horde/Plugins/Build/HordeServer.Build/Jobs/JobTaskSource.cs b/Engine/Source/Programs/Horde/Plugins/Build/HordeServer.Build/Jobs/JobTaskSource.cs
--- a/Engine/Source/Programs/Horde/Plugins/Build/HordeServer.Build/Jobs/JobTaskSource.cs
+++ b/Engine/Source/Programs/Horde/Plugins/Build/HordeServer.Build/Jobs/JobTaskSource.cs
@@ -777,11 +777,9 @@ namespace HordeServer.Jobs
 				return false;
 			}
 
-			// When a call to JobTaskSource.AssignLeaseAsync by the AgentService requests conform it will not be actioned
-			// by ConformTaskSource until the next iteration of AgentService.WaitForLeaseAsync.
 			if (agent.ConformAttemptCount is null or 0 && agent.RequestFullConform)
 			{
-				_logger.LogInformation("Pending conform already requested on {AgentId}", agent.Id);
+				_logger.LogInformation("Skipping assigning lease to {AgentId}, as a conform is already pending", agent.Id);
 				return true;
 			}
 
