[Horde][ReleaseNotes] Release notes

diff --git a/Engine/Source/Programs/Horde/Docs/ReleaseNotes.md b/Engine/Source/Programs/Horde/Docs/ReleaseNotes.md
--- a/Engine/Source/Programs/Horde/Docs/ReleaseNotes.md
+++ b/Engine/Source/Programs/Horde/Docs/ReleaseNotes.md
@@ -1,5 +1,28 @@
 # Release Notes
 
+## 2025-11-06
+
+* Fixing an issue where we were disposing of the HttpResponse content before reading it. (47903704)
+* Send cancellation request to Job Driver process When then Job Handler receives a (47638012)
+* Restore throwing cancellation exception when step is abort and run step clean up (47657684)
+* Set the Job Driver cancellation event to nonsignaled after process exit (47659409)
+* Update log message when the job executor receives cancellation signal (47728221)
+* Add support for conforming with PerforceMaterializer (47733607)
+* Correctly consume cascaded command option & simplified label styling. (47740983)
+* Introduce parallel data source mode for new API endpoint (47741988)
+* Update log message to include Agent ID when lease cannot be assigned due to conf (47776364)
+* Add additional logging to conform check (47776588)
+* Update message when lease cannot be assigned to agent due to conform being requi (47783624)
+* Fixed an error where we retry a read request but re-use the same request message (47847344)
+* Fixed a bug where we re-use the same HttpRequestMessage when retrying a WriteReq (47851781)
+* Moving MultipartFormDataContent creation outside of the Polly retry loop. (47855479)
+* Fixing issue where the HttpContent of the request would be re-used between retry (47858228)
+* Restricted - Applied standard formatting to canned sql queries & added new issueSpan join issue (47858668)
+* Adding context to the issue telemetry. (47862369)
+* Optionally re-open issue when assigning to other user (47866841)
+* Only assign issue to user if previously unassigned when marking fixed via Slack (47871287)
+* Update message when lease cannot be assigned to agent due to a pending conform (47898336)
+
 ## 2025-10-30
 
 * Add prevention suggestion wording to root cause summary placeholder (47573311)
