[Horde][ReleaseNotes] Cleanup errant items

diff --git a/Engine/Source/Programs/Horde/Docs/ReleaseNotes.md b/Engine/Source/Programs/Horde/Docs/ReleaseNotes.md
--- a/Engine/Source/Programs/Horde/Docs/ReleaseNotes.md
+++ b/Engine/Source/Programs/Horde/Docs/ReleaseNotes.md
@@ -17,7 +17,6 @@
 * Fixed a bug where we re-use the same HttpRequestMessage when retrying a WriteReq (47851781)
 * Moving MultipartFormDataContent creation outside of the Polly retry loop. (47855479)
 * Fixing issue where the HttpContent of the request would be re-used between retry (47858228)
-* Restricted - Applied standard formatting to canned sql queries & added new issueSpan join issue (47858668)
 * Adding context to the issue telemetry. (47862369)
 * Optionally re-open issue when assigning to other user (47866841)
 * Only assign issue to user if previously unassigned when marking fixed via Slack (47871287)
