Horde: Add reason to lease outcomes

diff --git a/Engine/Source/Programs/Horde/HordeAgent/Leases/Handlers/ComputeHandler.cs b/Engine/Source/Programs/Horde/HordeAgent/Leases/Handlers/ComputeHandler.cs
--- a/Engine/Source/Programs/Horde/HordeAgent/Leases/Handlers/ComputeHandler.cs
+++ b/Engine/Source/Programs/Horde/HordeAgent/Leases/Handlers/ComputeHandler.cs
@@ -101,7 +101,7 @@ namespace HordeAgent.Leases.Handlers
 				catch (OperationCanceledException ex) when (cts.IsCancellationRequested && idleTimeoutTransport.TimeSinceActivity > idleTimeoutTransport.NoDataTimeout)
 				{
 					logger.LogError(ex, "Lease was terminated due to no data being received for {Time} seconds", (int)idleTimeoutTransport.NoDataTimeout.TotalSeconds);
-					return LeaseResult.Failed;
+					return new LeaseResult(LeaseOutcome.Failed, LeaseOutcomeReason.TimedOut);
 				}
 			}
 			catch (OperationCanceledException)
@@ -111,7 +111,7 @@ namespace HordeAgent.Leases.Handlers
 			catch (Exception ex)
 			{
 				logger.LogError(ex, "Exception while executing compute task: {Message}", ex.Message);
-				return LeaseResult.Failed;
+				return new LeaseResult(LeaseOutcome.Failed, LeaseOutcomeReason.UnexpectedError);
 			}
 			finally
 			{
diff --git a/Engine/Source/Programs/Horde/HordeAgent/Leases/Handlers/ConformHandler.cs b/Engine/Source/Programs/Horde/HordeAgent/Leases/Handlers/ConformHandler.cs
--- a/Engine/Source/Programs/Horde/HordeAgent/Leases/Handlers/ConformHandler.cs
+++ b/Engine/Source/Programs/Horde/HordeAgent/Leases/Handlers/ConformHandler.cs
@@ -46,7 +46,7 @@ namespace HordeAgent.Leases.Handlers
 			int exitCode = await RunDotNetProcessAsync(driverAssembly, arguments, environment, AgentApp.IsSelfContained, serverLogger, cancellationToken);
 			serverLogger.LogInformation("Driver finished with exit code {ExitCode}", exitCode);
 
-			return (exitCode == 0) ? LeaseResult.Success : LeaseResult.Failed;
+			return (exitCode == 0) ? LeaseResult.Success : new LeaseResult(LeaseOutcome.Failed, LeaseOutcomeReason.NonZeroExitCode);
 		}
 	}
 
diff --git a/Engine/Source/Programs/Horde/HordeAgent/Leases/Handlers/JobHandler.cs b/Engine/Source/Programs/Horde/HordeAgent/Leases/Handlers/JobHandler.cs
--- a/Engine/Source/Programs/Horde/HordeAgent/Leases/Handlers/JobHandler.cs
+++ b/Engine/Source/Programs/Horde/HordeAgent/Leases/Handlers/JobHandler.cs
@@ -186,7 +186,7 @@ namespace HordeAgent.Leases.Handlers
 				}
 			}
 
-			return (exitCode == 0) ? LeaseResult.Success : LeaseResult.Failed;
+			return (exitCode == 0) ? LeaseResult.Success : new LeaseResult(LeaseOutcome.Failed, LeaseOutcomeReason.NonZeroExitCode);
 		}
 
 		private void DriverCancelHandler()
diff --git a/Engine/Source/Programs/Horde/HordeAgent/Leases/Handlers/UpgradeHandler.cs b/Engine/Source/Programs/Horde/HordeAgent/Leases/Handlers/UpgradeHandler.cs
--- a/Engine/Source/Programs/Horde/HordeAgent/Leases/Handlers/UpgradeHandler.cs
+++ b/Engine/Source/Programs/Horde/HordeAgent/Leases/Handlers/UpgradeHandler.cs
@@ -83,7 +83,7 @@ namespace HordeAgent.Leases.Handlers
 					if (!isUpdateSelfContained)
 					{
 						logger.LogError("Current running agent is packaged as a self-contained app, but incoming update is not");
-						return LeaseResult.Failed;
+						return new LeaseResult(LeaseOutcome.Failed, LeaseOutcomeReason.UnexpectedError);
 					}
 
 					executable = "HordeAgent";
@@ -96,7 +96,7 @@ namespace HordeAgent.Leases.Handlers
 					if (!File.Exists(executable))
 					{
 						logger.LogError("Unable to find self-contained {AgentExe} in extracted archive", executable);
-						return LeaseResult.Failed;
+						return new LeaseResult(LeaseOutcome.Failed, LeaseOutcomeReason.UnexpectedError);
 					}
 				}
 				else
@@ -111,7 +111,7 @@ namespace HordeAgent.Leases.Handlers
 					if (!FileReference.Exists(newAssemblyFileName))
 					{
 						logger.LogError("Unable to find {AgentExe} in extracted archive", newAssemblyFileName);
-						return LeaseResult.Failed;
+						return new LeaseResult(LeaseOutcome.Failed, LeaseOutcomeReason.UnexpectedError);
 					}
 
 					arguments.AppendArgument(newAssemblyFileName.FullName);
diff --git a/Engine/Source/Programs/Horde/HordeAgent/Leases/LeaseHandler.cs b/Engine/Source/Programs/Horde/HordeAgent/Leases/LeaseHandler.cs
--- a/Engine/Source/Programs/Horde/HordeAgent/Leases/LeaseHandler.cs
+++ b/Engine/Source/Programs/Horde/HordeAgent/Leases/LeaseHandler.cs
@@ -47,10 +47,10 @@ namespace HordeAgent.Leases
 		/// <summary>
 		/// Reason for cancellation
 		/// </summary>
-		public string CancellationReason => _cancellationReason ?? "unknown";
+		public string CancellationReason => _outcomeReason.ToString();
 
-		readonly CancellationTokenSource _cancellationSource = new CancellationTokenSource();
-		string? _cancellationReason;
+		readonly CancellationTokenSource _cancellationSource = new ();
+		LeaseOutcomeReason _outcomeReason = LeaseOutcomeReason.None;
 
 		/// <summary>
 		/// Constructor
@@ -92,9 +92,9 @@ namespace HordeAgent.Leases
 		/// <summary>
 		/// Cancels the lease
 		/// </summary>
-		public void Cancel(string reason)
+		public void Cancel(LeaseOutcomeReason reason)
 		{
-			Interlocked.CompareExchange(ref _cancellationReason, reason, null);
+			_outcomeReason = reason;
 			_cancellationSource.Cancel();
 		}
 
@@ -110,7 +110,7 @@ namespace HordeAgent.Leases
 			logger.LogInformation("Handling lease {LeaseId}: {LeaseTypeUrl}", Id, RpcLease.Payload.TypeUrl);
 
 			// Get the lease outcome
-			LeaseResult result = LeaseResult.Failed;
+			LeaseResult result = new (LeaseOutcome.Failed, LeaseOutcomeReason.UnexpectedError);
 			try
 			{
 				span.SetAttribute("horde.lease.task", RpcPayload.TypeUrl);
@@ -133,12 +133,15 @@ namespace HordeAgent.Leases
 			{
 				newRpcLease.State = RpcLeaseState.Cancelled;
 				newRpcLease.Outcome = RpcLeaseOutcome.Failed;
+				newRpcLease.OutcomeReasons.AddRange(RpcLease.ToProtobufOutcomeReasons(result.OutcomeReason | _outcomeReason));
 				newRpcLease.Output = Google.Protobuf.ByteString.Empty;
 			}
 			else
 			{
+				LeaseOutcomeReason reason = result.Outcome is LeaseOutcome.Success ? result.OutcomeReason : result.OutcomeReason | _outcomeReason;
 				newRpcLease.State = (result.Outcome == LeaseOutcome.Cancelled) ? RpcLeaseState.Cancelled : RpcLeaseState.Completed;
 				newRpcLease.Outcome = (RpcLeaseOutcome)result.Outcome;
+				newRpcLease.OutcomeReasons.AddRange(RpcLease.ToProtobufOutcomeReasons(reason));
 				newRpcLease.Output = (result.Output != null) ? ByteString.CopyFrom(result.Output) : ByteString.Empty;
 			}
 			RpcLease = newRpcLease;
diff --git a/Engine/Source/Programs/Horde/HordeAgent/Leases/LeaseManager.cs b/Engine/Source/Programs/Horde/HordeAgent/Leases/LeaseManager.cs
--- a/Engine/Source/Programs/Horde/HordeAgent/Leases/LeaseManager.cs
+++ b/Engine/Source/Programs/Horde/HordeAgent/Leases/LeaseManager.cs
@@ -3,6 +3,7 @@
 using System.Diagnostics;
 using System.Text;
 using EpicGames.Core;
+using EpicGames.Horde.Agents.Leases;
 using Google.Protobuf.WellKnownTypes;
 using Grpc.Core;
 using HordeAgent.Services;
@@ -184,7 +185,7 @@ namespace HordeAgent.Leases
 				try
 				{
 					_logger.LogInformation("Draining leases... ({NumLeases} remaining)", _activeLeases.Count);
-					await DrainLeasesAsync("session terminating");
+					await DrainLeasesAsync(LeaseOutcomeReason.AgentSessionTerminated);
 				}
 				catch (Exception ex)
 				{
@@ -198,14 +199,14 @@ namespace HordeAgent.Leases
 		/// <summary>
 		/// Drain and terminate any currently running leases
 		/// </summary>
-		/// <param name="reason">Human-readable reason</param>
+		/// <param name="reason">Reason for draining the leases</param>
 		/// <param name="graceful">Whether to warn workload about upcoming lease termination via signal file</param>
-		async Task DrainLeasesAsync(string reason, bool graceful = false)
+		async Task DrainLeasesAsync(LeaseOutcomeReason reason, bool graceful = false)
 		{
 			if (graceful)
 			{
 				TimeSpan ttl = TimeSpan.FromSeconds(1);
-				await WriteTerminationSignalFileAsync(_settings.CurrentValue.GetTerminationSignalFile().FullName, reason, DateTime.UtcNow + ttl, ttl, CancellationToken.None);
+				await WriteTerminationSignalFileAsync(_settings.CurrentValue.GetTerminationSignalFile().FullName, reason.ToString(), DateTime.UtcNow + ttl, ttl, CancellationToken.None);
 				await Task.Delay(TimeSpan.FromSeconds(6)); // Allow workload some time to act on signal file
 			}
 			
@@ -372,7 +373,7 @@ namespace HordeAgent.Leases
 									if (handler != null)
 									{
 										_logger.LogInformation("Cancelling lease {LeaseId}", serverLease.Id);
-										handler.Cancel("cancelled by server");
+										handler.Cancel(LeaseOutcomeReason.ServerInitiated);
 									}
 								}
 								if (serverLease.State == RpcLeaseState.Pending && !_activeLeases.Any(x => x.RpcLease.Id == serverLease.Id))
@@ -422,7 +423,7 @@ namespace HordeAgent.Leases
 						if (_activeLeases.Count > 0)
 						{
 							_logger.LogInformation("Agent marked itself as busy. Draining any active leases to prevent them from using up local resources...");
-							await DrainLeasesAsync("user is active", true);
+							await DrainLeasesAsync(LeaseOutcomeReason.AgentBusy, true);
 						}
 					}
 					else if (_activeLeases.Count == 0)
@@ -458,7 +459,7 @@ namespace HordeAgent.Leases
 			if (!_typeUrlToLeaseHandler.TryGetValue(payload.TypeUrl, out LeaseHandlerFactory? leaseHandlerFactory))
 			{
 				_logger.LogError("Invalid lease payload type ({PayloadType})", payload.TypeUrl);
-				return new DefaultLeaseHandler(lease, LeaseResult.Failed);
+				return new DefaultLeaseHandler(lease, new LeaseResult(LeaseOutcome.Failed, LeaseOutcomeReason.UnexpectedError));
 			}
 			return leaseHandlerFactory.CreateHandler(lease);
 		}
diff --git a/Engine/Source/Programs/Horde/HordeAgent/Leases/LeaseResult.cs b/Engine/Source/Programs/Horde/HordeAgent/Leases/LeaseResult.cs
--- a/Engine/Source/Programs/Horde/HordeAgent/Leases/LeaseResult.cs
+++ b/Engine/Source/Programs/Horde/HordeAgent/Leases/LeaseResult.cs
@@ -15,6 +15,11 @@ namespace HordeAgent.Leases
 		/// </summary>
 		public LeaseOutcome Outcome { get; }
 
+		/// <summary>
+		/// Reason for the outcome of the lease
+		/// </summary>
+		public LeaseOutcomeReason OutcomeReason { get; }
+
 		/// <summary>
 		/// Output from executing the task
 		/// </summary>
@@ -25,28 +30,18 @@ namespace HordeAgent.Leases
 		/// </summary>
 		public SessionResult? SessionResult { get; }
 
-		/// <summary>
-		/// Static instance of a cancelled result
-		/// </summary>
-		public static LeaseResult Cancelled { get; } = new LeaseResult(LeaseOutcome.Cancelled);
-
-		/// <summary>
-		/// Static instance of a failed result
-		/// </summary>
-		public static LeaseResult Failed { get; } = new LeaseResult(LeaseOutcome.Failed);
-
 		/// <summary>
 		/// Static instance of a succesful result without a payload
 		/// </summary>
-		public static LeaseResult Success { get; } = new LeaseResult(LeaseOutcome.Success);
+		public static LeaseResult Success { get; } = new (LeaseOutcome.Success, LeaseOutcomeReason.None);
 
 		/// <summary>
-		/// Constructor for 
+		/// Constructor
 		/// </summary>
-		/// <param name="outcome"></param>
-		private LeaseResult(LeaseOutcome outcome)
+		public LeaseResult(LeaseOutcome outcome, LeaseOutcomeReason outcomeReason)
 		{
 			Outcome = outcome;
+			OutcomeReason = outcomeReason;
 		}
 
 		/// <summary>
diff --git a/Engine/Source/Programs/Horde/Plugins/Build/HordeServer.Build.Tests/Agents/Leases/LeasesControllerDbTest.cs b/Engine/Source/Programs/Horde/Plugins/Build/HordeServer.Build.Tests/Agents/Leases/LeasesControllerDbTest.cs
--- a/Engine/Source/Programs/Horde/Plugins/Build/HordeServer.Build.Tests/Agents/Leases/LeasesControllerDbTest.cs
+++ b/Engine/Source/Programs/Horde/Plugins/Build/HordeServer.Build.Tests/Agents/Leases/LeasesControllerDbTest.cs
@@ -38,7 +38,7 @@ namespace HordeServer.Tests.Agents.Leases
 			LeaseId id = new LeaseId(BinaryIdUtils.CreateNew());
 			SessionId sessionId = SessionIdUtils.GenerateNewId();
 			ILease lease = await LeaseCollection.AddAsync(id, null, "myLease", new AgentId("agent-1"), sessionId, null, null, null, startTime, Array.Empty<byte>());
-			await LeaseCollection.TrySetOutcomeAsync(id, startTime + duration, LeaseOutcome.Success, null);
+			await LeaseCollection.TrySetOutcomeAsync(id, startTime + duration, LeaseOutcome.Success, LeaseOutcomeReason.None, null);
 			return lease;
 		}
 	}
diff --git a/Engine/Source/Programs/Horde/Plugins/Build/HordeServer.Build.Tests/Fleet/LeaseUtilizationAwsMetricStrategyTests.cs b/Engine/Source/Programs/Horde/Plugins/Build/HordeServer.Build.Tests/Fleet/LeaseUtilizationAwsMetricStrategyTests.cs
--- a/Engine/Source/Programs/Horde/Plugins/Build/HordeServer.Build.Tests/Fleet/LeaseUtilizationAwsMetricStrategyTests.cs
+++ b/Engine/Source/Programs/Horde/Plugins/Build/HordeServer.Build.Tests/Fleet/LeaseUtilizationAwsMetricStrategyTests.cs
@@ -224,7 +224,7 @@ namespace HordeServer.Tests.Fleet
 			ILease lease = await LeaseCollection.AddAsync(new LeaseId(BinaryIdUtils.CreateNew()), null, "placeholderLease", agent.Id, agent.SessionId!.Value, new StreamId("placeholderStream"), poolId, null, startTime, payload);
 			if (duration != null)
 			{
-				bool wasModified = await LeaseCollection.TrySetOutcomeAsync(lease.Id, startTime + duration.Value, LeaseOutcome.Success, null);
+				bool wasModified = await LeaseCollection.TrySetOutcomeAsync(lease.Id, startTime + duration.Value, LeaseOutcome.Success, LeaseOutcomeReason.None, null);
 				Assert.IsTrue(wasModified);
 			}
 
diff --git a/Engine/Source/Programs/Horde/Plugins/Build/HordeServer.Build.Tests/Fleet/LeaseUtilizationStrategyTest.cs b/Engine/Source/Programs/Horde/Plugins/Build/HordeServer.Build.Tests/Fleet/LeaseUtilizationStrategyTest.cs
--- a/Engine/Source/Programs/Horde/Plugins/Build/HordeServer.Build.Tests/Fleet/LeaseUtilizationStrategyTest.cs
+++ b/Engine/Source/Programs/Horde/Plugins/Build/HordeServer.Build.Tests/Fleet/LeaseUtilizationStrategyTest.cs
@@ -120,7 +120,7 @@ namespace HordeServer.Tests.Fleet
 			}
 
 			ILease lease = await LeaseCollection.AddAsync(new LeaseId(BinaryIdUtils.CreateNew()), null, "placeholderLease", agent.Id, agent.SessionId!.Value, new StreamId("placeholderStream"), pool.Id, null, startTime, payload);
-			bool wasModified = await LeaseCollection.TrySetOutcomeAsync(lease.Id, startTime + duration, LeaseOutcome.Success, null);
+			bool wasModified = await LeaseCollection.TrySetOutcomeAsync(lease.Id, startTime + duration, LeaseOutcome.Success, LeaseOutcomeReason.None,null);
 			Assert.IsTrue(wasModified);
 			return lease;
 		}
diff --git a/Engine/Source/Programs/Horde/Plugins/Compute/HordeServer.Compute/Agents/AgentCollection.cs b/Engine/Source/Programs/Horde/Plugins/Compute/HordeServer.Compute/Agents/AgentCollection.cs
--- a/Engine/Source/Programs/Horde/Plugins/Compute/HordeServer.Compute/Agents/AgentCollection.cs
+++ b/Engine/Source/Programs/Horde/Plugins/Compute/HordeServer.Compute/Agents/AgentCollection.cs
@@ -961,13 +961,14 @@ namespace HordeServer.Agents
 				foreach (RpcSessionLease oldLease in agent.Session.Leases.Where(x => !newLeaseIds.Contains(x.Id)))
 				{
 					RpcLease? rpcLease = options.Leases?.FirstOrDefault(x => x.Id == oldLease.Id);
-					if (rpcLease != null && rpcLease.State == RpcLeaseState.Completed)
+					if (rpcLease is { State: RpcLeaseState.Completed })
 					{
-						await _leaseCollection.TrySetOutcomeAsync(oldLease.Id, _clock.UtcNow, (LeaseOutcome)rpcLease.Outcome, rpcLease.Output.ToArray(), CancellationToken.None);
+						await _leaseCollection.TrySetOutcomeAsync(oldLease.Id, _clock.UtcNow, (LeaseOutcome)rpcLease.Outcome, LeaseOutcomeReason.None, rpcLease.Output.ToArray(), CancellationToken.None);
 					}
 					else
 					{
-						await _leaseCollection.TrySetOutcomeAsync(oldLease.Id, _clock.UtcNow, LeaseOutcome.Cancelled, null, CancellationToken.None);
+						LeaseOutcomeReason reason = rpcLease?.ToNativeOutcomeReason() ?? LeaseOutcomeReason.None;
+						await _leaseCollection.TrySetOutcomeAsync(oldLease.Id, _clock.UtcNow, LeaseOutcome.Cancelled, reason, null, CancellationToken.None);
 					}
 				}
 			}
diff --git a/Engine/Source/Programs/Horde/Plugins/Compute/HordeServer.Compute/Agents/AgentService.cs b/Engine/Source/Programs/Horde/Plugins/Compute/HordeServer.Compute/Agents/AgentService.cs
--- a/Engine/Source/Programs/Horde/Plugins/Compute/HordeServer.Compute/Agents/AgentService.cs
+++ b/Engine/Source/Programs/Horde/Plugins/Compute/HordeServer.Compute/Agents/AgentService.cs
@@ -581,9 +581,9 @@ namespace HordeServer.Agents
 		/// Cancels the specified agent lease
 		/// </summary>
 		/// <param name="lease">The lease to cancel</param>
+		/// <param name="outcomeReason">Reason for the cancellation</param>
 		/// <param name="cancellationToken">Cancellation token for the operation</param>
-		/// <returns></returns>
-		public async Task<bool> CancelLeaseAsync(ILease lease, CancellationToken cancellationToken)
+		public async Task<bool> CancelLeaseAsync(ILease lease, LeaseOutcomeReason outcomeReason, CancellationToken cancellationToken)
 		{
 			// If it's already complete, don't do anything
 			if (lease.FinishTime != null)
@@ -596,7 +596,7 @@ namespace HordeServer.Agents
 			if (agent == null)
 			{
 				_logger.LogWarning("Forcing lease {LeaseId} to cancelled state; agent {AgentId} not found", lease.Id, lease.AgentId);
-				await _leases.TrySetOutcomeAsync(lease.Id, _clock.UtcNow, LeaseOutcome.Cancelled, null, cancellationToken);
+				await _leases.TrySetOutcomeAsync(lease.Id, _clock.UtcNow, LeaseOutcome.Cancelled, outcomeReason, null, cancellationToken);
 				return false;
 			}
 
@@ -605,7 +605,7 @@ namespace HordeServer.Agents
 			if (index == -1)
 			{
 				_logger.LogWarning("Forcing lease {LeaseId} to cancelled state; no longer running on agent {AgentId}", lease.Id, lease.AgentId);
-				await _leases.TrySetOutcomeAsync(lease.Id, _clock.UtcNow, LeaseOutcome.Cancelled, null, cancellationToken);
+				await _leases.TrySetOutcomeAsync(lease.Id, _clock.UtcNow, LeaseOutcome.Cancelled, outcomeReason, null, cancellationToken);
 				return false;
 			}
 
diff --git a/Engine/Source/Programs/Horde/Plugins/Compute/HordeServer.Compute/Agents/Leases/ILease.cs b/Engine/Source/Programs/Horde/Plugins/Compute/HordeServer.Compute/Agents/Leases/ILease.cs
--- a/Engine/Source/Programs/Horde/Plugins/Compute/HordeServer.Compute/Agents/Leases/ILease.cs
+++ b/Engine/Source/Programs/Horde/Plugins/Compute/HordeServer.Compute/Agents/Leases/ILease.cs
@@ -76,6 +76,11 @@ namespace HordeServer.Agents.Leases
 		/// </summary>
 		public LeaseOutcome Outcome { get; }
 
+		/// <summary>
+		/// Reason for the outcome. Can be a combination of reasons.
+		/// </summary>
+		public LeaseOutcomeReason OutcomeReason { get; set; }
+
 		/// <summary>
 		/// Output from executing the lease
 		/// </summary>
diff --git a/Engine/Source/Programs/Horde/Plugins/Compute/HordeServer.Compute/Agents/Leases/ILeaseCollection.cs b/Engine/Source/Programs/Horde/Plugins/Compute/HordeServer.Compute/Agents/Leases/ILeaseCollection.cs
--- a/Engine/Source/Programs/Horde/Plugins/Compute/HordeServer.Compute/Agents/Leases/ILeaseCollection.cs
+++ b/Engine/Source/Programs/Horde/Plugins/Compute/HordeServer.Compute/Agents/Leases/ILeaseCollection.cs
@@ -101,9 +101,10 @@ namespace HordeServer.Agents.Leases
 		/// <param name="leaseId">The lease to update</param>
 		/// <param name="finishTime">Time at which the lease finished</param>
 		/// <param name="outcome">Outcome of the lease</param>
+		/// <param name="outcomeReason">Reason for the outcome</param>
 		/// <param name="output">Output data from the task</param>
 		/// <param name="cancellationToken">Cancellation token for the operation</param>
 		/// <returns>True if the lease was updated, false otherwise</returns>
-		Task<bool> TrySetOutcomeAsync(LeaseId leaseId, DateTime finishTime, LeaseOutcome outcome, byte[]? output, CancellationToken cancellationToken = default);
+		Task<bool> TrySetOutcomeAsync(LeaseId leaseId, DateTime finishTime, LeaseOutcome outcome, LeaseOutcomeReason outcomeReason, byte[]? output, CancellationToken cancellationToken = default);
 	}
 }
diff --git a/Engine/Source/Programs/Horde/Plugins/Compute/HordeServer.Compute/Agents/Leases/LeaseCollection.cs b/Engine/Source/Programs/Horde/Plugins/Compute/HordeServer.Compute/Agents/Leases/LeaseCollection.cs
--- a/Engine/Source/Programs/Horde/Plugins/Compute/HordeServer.Compute/Agents/Leases/LeaseCollection.cs
+++ b/Engine/Source/Programs/Horde/Plugins/Compute/HordeServer.Compute/Agents/Leases/LeaseCollection.cs
@@ -8,6 +8,7 @@ using EpicGames.Horde.Logs;
 using EpicGames.Horde.Streams;
 using HordeServer.Server;
 using HordeServer.Utilities;
+using MongoDB.Bson;
 using MongoDB.Bson.Serialization.Attributes;
 using MongoDB.Driver;
 using OpenTelemetry.Trace;
@@ -44,6 +45,9 @@ namespace HordeServer.Agents.Leases
 			public DateTime? FinishTime { get; set; }
 			public byte[] Payload { get; set; }
 			public LeaseOutcome Outcome { get; set; }
+			
+			[BsonRepresentation(BsonType.Int64)]
+			public LeaseOutcomeReason OutcomeReason { get; set; } = LeaseOutcomeReason.None;
 
 			[BsonIgnoreIfNull]
 			public byte[]? Output { get; set; }
@@ -193,13 +197,16 @@ namespace HordeServer.Agents.Leases
 		}
 
 		/// <inheritdoc/>
-		public async Task<bool> TrySetOutcomeAsync(LeaseId leaseId, DateTime finishTime, LeaseOutcome outcome, byte[]? output, CancellationToken cancellationToken)
+		public async Task<bool> TrySetOutcomeAsync(LeaseId leaseId, DateTime finishTime, LeaseOutcome outcome, LeaseOutcomeReason outcomeReason, byte[]? output, CancellationToken cancellationToken)
 		{
 			FilterDefinitionBuilder<LeaseDocument> filterBuilder = Builders<LeaseDocument>.Filter;
 			FilterDefinition<LeaseDocument> filter = filterBuilder.Eq(x => x.Id, leaseId) & filterBuilder.Eq(x => x.FinishTime, null);
 
 			UpdateDefinitionBuilder<LeaseDocument> updateBuilder = Builders<LeaseDocument>.Update;
-			UpdateDefinition<LeaseDocument> update = updateBuilder.Set(x => x.FinishTime, finishTime).Set(x => x.Outcome, outcome);
+			UpdateDefinition<LeaseDocument> update = updateBuilder
+				.Set(x => x.FinishTime, finishTime)
+				.Set(x => x.Outcome, outcome)
+				.Set(x => x.OutcomeReason, outcomeReason);
 
 			if (output != null && output.Length > 0)
 			{
diff --git a/Engine/Source/Programs/Horde/Plugins/Compute/HordeServer.Compute/Agents/Leases/LeasesController.cs b/Engine/Source/Programs/Horde/Plugins/Compute/HordeServer.Compute/Agents/Leases/LeasesController.cs
--- a/Engine/Source/Programs/Horde/Plugins/Compute/HordeServer.Compute/Agents/Leases/LeasesController.cs
+++ b/Engine/Source/Programs/Horde/Plugins/Compute/HordeServer.Compute/Agents/Leases/LeasesController.cs
@@ -232,7 +232,7 @@ namespace HordeServer.Agents.Leases
 				return NotFound(leaseId);
 			}
 
-			await _agentService.CancelLeaseAsync(lease, cancellationToken);
+			await _agentService.CancelLeaseAsync(lease, LeaseOutcomeReason.UserInitiated, cancellationToken);
 			return Ok();
 		}
 	}
diff --git a/Engine/Source/Programs/Horde/Plugins/Compute/HordeServer.Compute/Protos/horde/messages/lease.proto.cs b/Engine/Source/Programs/Horde/Plugins/Compute/HordeServer.Compute/Protos/horde/messages/lease.proto.cs
--- a/Engine/Source/Programs/Horde/Plugins/Compute/HordeServer.Compute/Protos/horde/messages/lease.proto.cs
+++ b/Engine/Source/Programs/Horde/Plugins/Compute/HordeServer.Compute/Protos/horde/messages/lease.proto.cs
@@ -10,12 +10,39 @@ namespace HordeCommon.Rpc.Messages
 	partial class RpcLease
 	{
 		/// <summary>
-		/// Typed acesssor for the <see cref="RpcLease.IdString"/> field.
+		/// Typed accessor for the <see cref="RpcLease.IdString"/> field.
 		/// </summary>
 		public LeaseId Id
 		{
 			get => LeaseId.Parse(IdString);
 			set => IdString = value.ToString();
 		}
+
+		/// <summary>
+		/// Convert outcome reason(s) from Protobuf-format to C# enum
+		/// </summary>
+		public LeaseOutcomeReason ToNativeOutcomeReason()
+		{
+			return OutcomeReasons
+				.Select(r => (LeaseOutcomeReason)(long)r)
+				.Aggregate(LeaseOutcomeReason.None, (current, reason) => current | reason);
+		}
+		
+		/// <summary>
+		/// Convert C# enum with outcome reason(s) to Protobuf-format
+		/// </summary>
+		public static List<RpcLeaseOutcomeReason> ToProtobufOutcomeReasons(LeaseOutcomeReason reasons)
+		{
+			List<RpcLeaseOutcomeReason> result = [];
+			foreach (RpcLeaseOutcomeReason rpcReason in Enum.GetValues(typeof(RpcLeaseOutcomeReason)))
+			{
+				if (rpcReason != RpcLeaseOutcomeReason.None && ((long)reasons & (long)rpcReason) != 0)
+				{
+					result.Add(rpcReason);
+				}
+			}
+    
+			return result;
+		}
 	}
 }
diff --git a/Engine/Source/Programs/Shared/EpicGames.Horde/Agents/Leases/LeaseMessages.cs b/Engine/Source/Programs/Shared/EpicGames.Horde/Agents/Leases/LeaseMessages.cs
--- a/Engine/Source/Programs/Shared/EpicGames.Horde/Agents/Leases/LeaseMessages.cs
+++ b/Engine/Source/Programs/Shared/EpicGames.Horde/Agents/Leases/LeaseMessages.cs
@@ -1,5 +1,7 @@
 // Copyright Epic Games, Inc. All Rights Reserved.
 
+using System;
+
 #pragma warning disable CA1027 // Mark enums with FlagsAttribute
 
 namespace EpicGames.Horde.Agents.Leases
@@ -30,6 +32,49 @@ namespace EpicGames.Horde.Agents.Leases
 		Cancelled = 4
 	}
 
+	/// <summary>
+	/// Reason for an outcome of a lease. Values must match lease_outcome.proto.
+	/// </summary>
+	[Flags]
+	public enum LeaseOutcomeReason : long
+	{
+		/// <summary>No specific reason</summary>
+		None = 0,
+		
+		/// <summary>Lease finished due to an unhandled or unknown error</summary>
+		UnexpectedError = 1,
+		
+		/// <summary>Lease finished with a non-zero exit code</summary>
+		NonZeroExitCode = 2,
+
+		/// <summary>User initiated</summary>
+		UserInitiated = 4,
+
+		/// <summary>Server initiated</summary>
+		ServerInitiated = 8,
+		
+		/// <summary>Agent initiated</summary>
+		AgentInitiated = 16,
+		
+		/// <summary>Agent's session was terminated</summary>
+		AgentSessionTerminated = 32,
+		
+		/// <summary>Agent became busy during lease (e.g non-Horde related work came up running on a workstation)</summary>
+		AgentBusy = 64,
+
+		/// <summary>
+		/// Lease deemed to be a duplicate
+		/// For example two identical jobs executing or too many leases from a single user
+		/// </summary>
+		Duplicate = 128,
+		
+		/// <summary>Lease execution timed out</summary>
+		TimedOut = 256,
+		
+		/// <summary>Lease no longer needed despite scheduled/started</summary>
+		NoLongerNeeded = 512,
+	}
+
 	/// <summary>
 	/// State of a lease. Values must match lease_state.proto.
 	/// </summary>
