Horde: Fix config object append bug leading to infinite loop

diff --git a/Engine/Source/Programs/Horde/HordeServer.Shared/Configuration/ConfigObject.cs b/Engine/Source/Programs/Horde/HordeServer.Shared/Configuration/ConfigObject.cs
--- a/Engine/Source/Programs/Horde/HordeServer.Shared/Configuration/ConfigObject.cs
+++ b/Engine/Source/Programs/Horde/HordeServer.Shared/Configuration/ConfigObject.cs
@@ -153,7 +153,7 @@ namespace HordeServer.Configuration
 						{
 							propertyInfo.SetValue(target, sourceValue);
 						}
-						else if (sourceValue != null)
+						else if (sourceValue != null && !ReferenceEquals(sourceValue, targetValue))
 						{
 							IList sourceList = (IList)sourceValue;
 							IList targetList = (IList)targetValue;
diff --git a/Engine/Source/Programs/Horde/HordeServer.Tests/Configuration/ConfigTests.cs b/Engine/Source/Programs/Horde/HordeServer.Tests/Configuration/ConfigTests.cs
--- a/Engine/Source/Programs/Horde/HordeServer.Tests/Configuration/ConfigTests.cs
+++ b/Engine/Source/Programs/Horde/HordeServer.Tests/Configuration/ConfigTests.cs
@@ -8,6 +8,7 @@ using System.Text.Json;
 using System.Threading;
 using System.Threading.Tasks;
 using EpicGames.Core;
+using EpicGames.Horde.Streams;
 using HordeServer.Acls;
 using HordeServer.Configuration;
 using HordeServer.Plugins;
@@ -240,6 +241,32 @@ namespace HordeServer.Tests.Configuration
 			Assert.AreEqual("myStream", workspaces["stream1"].Stream);
 		}
 
+		[TestMethod]
+		public void ConfigMergeAppend_NoInfiniteLoop()
+		{
+			// Referencing same workspace base from multiple streams should not cause an infinite loop
+			// ConfigObject.MergeDefaults did not handle this
+			
+			BuildConfig buildConfig = new ();
+			buildConfig.Projects.Add(new ProjectConfig
+			{
+				WorkspaceTypes =
+				{
+					{ "base", new WorkspaceConfig { View = ["myView"] } },
+					{ "foo", new WorkspaceConfig { Identifier = "myFoo", Base = "base" } },
+				},
+				Streams = [
+					new StreamConfig { Id = new StreamId("stream1"), WorkspaceTypes = { {"stream1", new WorkspaceConfig { Base = "foo", Stream = "myStream1" } } } },
+					new StreamConfig { Id = new StreamId("stream2"), WorkspaceTypes = { {"stream2", new WorkspaceConfig { Base = "foo", Stream = "myStream2" } } } }
+				]
+			});
+			
+			GlobalConfig gc = new();
+			gc.Plugins.AddComputeConfig(new ComputeConfig());
+			gc.Plugins.AddBuildConfig(buildConfig);
+			gc.PostLoad(new ServerSettings(), Array.Empty<ILoadedPlugin>(), Array.Empty<IDefaultAclModifier>());
+		}
+
 		class ObjectValue
 		{
 			public string Value { get; set; } = "";
