Horde: Make awaiting on background notification tasks more robust

- Catches any exception directly in the task
- Simplify by using a list and only waiting in StopAsync

diff --git a/Engine/Source/Programs/Horde/Plugins/Build/HordeServer.Build/Notifications/NotificationService.cs b/Engine/Source/Programs/Horde/Plugins/Build/HordeServer.Build/Notifications/NotificationService.cs
--- a/Engine/Source/Programs/Horde/Plugins/Build/HordeServer.Build/Notifications/NotificationService.cs
+++ b/Engine/Source/Programs/Horde/Plugins/Build/HordeServer.Build/Notifications/NotificationService.cs
@@ -1,6 +1,5 @@
 // Copyright Epic Games, Inc. All Rights Reserved.
 
-using System.Collections.Concurrent;
 using System.Diagnostics.Metrics;
 using System.Security.Claims;
 using EpicGames.Core;
@@ -97,14 +96,10 @@ namespace HordeServer.Notifications
 		private readonly IOptionsMonitor<BuildServerConfig> _settings;
 
 		/// <summary>
-		/// List of asychronous tasks currently executing
+		/// Async notification tasks currently executing. Tracked so they can be awaited before shutting down
 		/// </summary>
-		private readonly ConcurrentQueue<Task> _newTasks = new ConcurrentQueue<Task>();
-
-		/// <summary>
-		/// Set when there are new tasks to wait for
-		/// </summary>
-		private AsyncEvent _newTaskEvent = new AsyncEvent();
+		private readonly List<Task> _tasks = [];
+		private readonly object _tasksLock = new();
 
 		/// <summary>
 		/// Settings for the application.
@@ -124,7 +119,6 @@ namespace HordeServer.Notifications
 		readonly IServerInfo _serverInfo;
 		readonly Counter<int> _jobCounter;
 		readonly Histogram<double> _jobDurationHistogram;
-		readonly BackgroundTask _backgroundTask;
 
 		static string RedisQueueListKey(string notificationType) => "NotificationService.queued." + notificationType;
 
@@ -162,7 +156,6 @@ namespace HordeServer.Notifications
 			_cache = cache;
 			_redisConnectionPool = redisService.ConnectionPool;
 			_serverInfo = serverInfo;
-			_backgroundTask = new BackgroundTask(ExecuteAsync);
 
 			issueService.OnIssueUpdated += NotifyIssueUpdated;
 			jobService.OnJobAborted += NotifyJobAborted;
@@ -179,7 +172,6 @@ namespace HordeServer.Notifications
 		/// <inheritdoc/>
 		public async Task StartAsync(CancellationToken cancellationToken)
 		{
-			_backgroundTask.Start();
 			await _ticker.StartAsync();
 		}
 
@@ -187,7 +179,7 @@ namespace HordeServer.Notifications
 		public async Task StopAsync(CancellationToken cancellationToken)
 		{
 			await _ticker.StopAsync();
-			await _backgroundTask.StopAsync(cancellationToken);
+			await WaitForAsyncTasksToCompleteAsync(cancellationToken);
 		}
 
 		/// <inheritdoc/>
@@ -200,7 +192,31 @@ namespace HordeServer.Notifications
 			_jobService.OnLabelUpdate -= NotifyLabelUpdate;
 
 			await _ticker.DisposeAsync();
-			await _backgroundTask.DisposeAsync();
+		}
+
+		private async Task WaitForAsyncTasksToCompleteAsync(CancellationToken ct)
+		{
+			Task[] allTasks;
+			lock (_tasksLock)
+			{
+				allTasks = _tasks.ToArray();
+			}
+
+			if (allTasks.Length > 0)
+			{
+				_logger.LogInformation("Waiting for {Count} notification tasks to complete...", allTasks.Length);
+				using CancellationTokenSource cts = CancellationTokenSource.CreateLinkedTokenSource(ct);
+				cts.CancelAfter(TimeSpan.FromSeconds(20));
+				try
+				{
+					await Task.WhenAll(allTasks).WaitAsync(cts.Token);
+				}
+				catch (OperationCanceledException)
+				{
+					int incomplete = allTasks.Count(t => !t.IsCompleted);
+					_logger.LogWarning("Stop timeout reached with {Count} notification tasks still in flight", incomplete);
+				}
+			}
 		}
 
 		/// <inheritdoc/>
@@ -267,8 +283,15 @@ namespace HordeServer.Notifications
 			}
 		}
 
-		/// <inheritdoc/>
-		public void NotifyJobScheduled(IPoolConfig pool, bool poolHasAgentsOnline, IJob job, IGraph graph, JobStepBatchId batchId)
+		/// <summary>
+		/// Notify all subscribers that a job step was scheduled
+		/// </summary>
+		/// <param name="pool">The pool job was scheduled in</param>
+		/// <param name="poolHasAgentsOnline">Whether the pool has agents online</param>
+		/// <param name="job">The job containing the step that has finished</param>
+		/// <param name="graph">Graph of the job</param>
+		/// <param name="batchId">The batch id</param>
+        public void NotifyJobScheduled(IPoolConfig pool, bool poolHasAgentsOnline, IJob job, IGraph graph, JobStepBatchId batchId)
 		{
 			if (pool.EnableAutoscaling && !poolHasAgentsOnline)
 			{
@@ -338,7 +361,31 @@ namespace HordeServer.Notifications
 		/// <param name="taskFunc">Function to generate an async task</param>
 		void EnqueueTask(Func<CancellationToken, Task> taskFunc)
 		{
-			_newTasks.Enqueue(Task.Run(() => taskFunc(CancellationToken.None), CancellationToken.None));
+			Task task = Task.Run(async () =>
+			{
+				try
+				{
+					await taskFunc(CancellationToken.None);
+				}
+				catch (Exception e)
+				{
+					_logger.LogError(e, "Exception while executing notification task");
+				}
+			});
+			
+			// Self-cleanup when done
+			_ = task.ContinueWith(_ =>
+			{
+				lock (_tasksLock)
+				{
+					_tasks.Remove(task);
+				}
+			}, TaskScheduler.Default);
+			
+			lock (_tasksLock)
+			{
+				_tasks.Add(task);
+			}
 		}
 
 		/// <summary>
@@ -434,65 +481,6 @@ namespace HordeServer.Notifications
 			return notifications;
 		}
 
-		/// <inheritdoc/>
-		async Task ExecuteAsync(CancellationToken stoppingToken)
-		{
-			// This background service just waits for tasks to finish and prints any exception info. The only reason to do this is to
-			// ensure we finish processing everything before shutdown.
-			using (CancellationTask stoppingTask = new CancellationTask(stoppingToken))
-			{
-				List<Task> tasks = new List<Task>();
-				tasks.Add(_newTaskEvent.Task);
-
-				for (; ; )
-				{
-					// Add any new tasks to be monitored
-					Task? newTask;
-					while (_newTasks.TryDequeue(out newTask))
-					{
-						tasks.Add(newTask);
-					}
-
-					// If we don't have any
-					if (tasks.Count == 1)
-					{
-						await Task.WhenAny(_newTaskEvent.Task, stoppingTask.Task);
-						if (stoppingToken.IsCancellationRequested)
-						{
-							break;
-						}
-					}
-					else
-					{
-						// Wait for something to finish
-						Task task = await Task.WhenAny(tasks);
-						if (task == _newTaskEvent.Task)
-						{
-							_newTaskEvent = new AsyncEvent();
-							tasks[0] = _newTaskEvent.Task;
-						}
-						else
-						{
-							try
-							{
-								await task;
-							}
-							catch (Exception ex)
-							{
-								_logger.LogError(ex, "Exception while executing notification");
-							}
-							tasks.Remove(task);
-						}
-					}
-				}
-			}
-		}
-
-		internal Task ExecuteBackgroundForTestAsync(CancellationToken stoppingToken)
-		{
-			return ExecuteAsync(stoppingToken);
-		}
-
 		/// <summary>
 		/// Gets the <see cref="INotificationTrigger"/> for a given trigger ID, if any.
 		/// </summary>
