Horde: Fix issue with logging of config exception stacktraces

diff --git a/Engine/Source/Programs/Horde/Plugins/Build/HordeServer.Build/Notifications/Sinks/SlackNotificationSink.cs b/Engine/Source/Programs/Horde/Plugins/Build/HordeServer.Build/Notifications/Sinks/SlackNotificationSink.cs
--- a/Engine/Source/Programs/Horde/Plugins/Build/HordeServer.Build/Notifications/Sinks/SlackNotificationSink.cs
+++ b/Engine/Source/Programs/Horde/Plugins/Build/HordeServer.Build/Notifications/Sinks/SlackNotificationSink.cs
@@ -2310,7 +2310,7 @@ namespace HordeServer.Notifications.Sinks
 					"Sending config update failure notification: {Message}, Exception: {ExceptionType} Stacktrace: {Stacktrace}", 
 					ex.Message,
 					ex.GetType().FullName,
-					ex.StackTrace?[..500]);
+					ex.StackTrace?.Length > 500 ? ex.StackTrace[..500] : ex.StackTrace);
 
 				List<string> details = new List<string>();
 				if (ex is ConfigException configEx)
