Implement test recipes, test health and test audit classes
Add support for test audit backend history read and post
Add support for test recipe read
Add support for test errorfingerprints

#rnx

diff --git a/Engine/Source/Programs/Horde/HordeDashboard/plugins/test-automation/api.ts b/Engine/Source/Programs/Horde/HordeDashboard/plugins/test-automation/api.ts
--- a/Engine/Source/Programs/Horde/HordeDashboard/plugins/test-automation/api.ts
+++ b/Engine/Source/Programs/Horde/HordeDashboard/plugins/test-automation/api.ts
@@ -102,8 +102,24 @@ export type GetPhaseSessionsRequest = {
     phaseIds?: string[];
 }
 
+/// Test description
+export type TestSessionDescriptionResponse = {
+
+    /// The team associated with this test
+    team?: string;
+
+    /// The harness associated with this test
+    harness?: string;
+
+    /// The result profile associated with this test
+    profile?: string;
+
+    /// The intent associated with this test
+    intent?: string;
+}
+
 /// Test name reference
-export type GetTestNameResponse = {
+export type GetTestNameResponse = TestSessionDescriptionResponse & {
 
     /// The id of the test
     id: string;
@@ -217,6 +233,16 @@ export enum TestPhaseOutcome {
     Unknown = "Unknown"
 }
 
+/// Error Fingerprint Entry
+export type GetErrorFingerprintEntry = {
+
+    /// The error fingerprint key
+    key: string;
+
+    /// The list of impacted phases
+    phases: string[];
+}
+
 /// Describes test session outcome
 export type GetTestSessionResponse = {
 
@@ -250,6 +276,9 @@ export type GetTestSessionResponse = {
     /// The test name reference associated with the test session
     nameRef: string;
 
+    /// The recipe associated with the test session
+    recipeId?: string;
+
     /// The outcome of the test session
     outcome: TestOutcome;
 
@@ -270,6 +299,9 @@ export type GetTestSessionResponse = {
 
     /// The step job id associated with the test session
     stepId: string;
+
+    /// The error finerprints associated with this test session
+    errorFingerprints?: GetErrorFingerprintEntry[];
 }
 
 /// Describes test phase session outcome
@@ -346,6 +378,9 @@ export type TestSessionSummaryResponse = {
     
     /// The count of the phases that failed
     phasesFailedCount: number;
+
+    /// The command line used to run the session
+    commandline?: string;
 }
 
 export type PhaseSessionDetailsResponse = {
@@ -393,6 +428,21 @@ export type DeviceDetailsResponse = {
     metadata: Record<string, string>;
 }
 
+export type TestSessionRecipeResponse = {
+
+    /// The test node class
+    nodeClass: string,
+
+    /// The version of the test
+    version: string,
+
+    /// The timeout in minutes used by this recipe
+    timeoutMinutes: number,
+
+    /// The command line parameters associated with this recipe
+    params?: Record<string, any>
+}
+
 export type TestSessionDetailsResponse = {
 
     /// The test session summary
@@ -409,6 +459,12 @@ export type TestSessionDetailsResponse = {
 
     /// The tags associated with the test session
     tags?: string[];
+
+    /// The test description associated with the test session
+    description?: TestSessionDescriptionResponse;
+
+    /// The test recipe associated with the test session
+    recipe?: TestSessionRecipeResponse;
 }
 
 export type JobStepTestDataItem = {
@@ -426,6 +482,93 @@ export type JobStepTestDataItem = {
     testNameRef: string
 }
 
+export type TestRecipeRequest = {
+
+    /// The recipe ids requested 
+    recipeIds?: string[],
+
+    /// The test ids associated with the recipe requested
+    testIds?: string[],
+
+    /// The stream ids associated with the recipe requested
+    streamIds?: string[]
+}
+
+export type TestRecipeResponse = TestSessionRecipeResponse & {
+
+    /// The test recipe id
+    id: string,
+
+    /// The test name reference
+    nameRef: string,
+
+    /// The stream id associated with this recipe
+    streamId: string,
+
+    /// The job template id associated with this recipe
+    jobTemplateId: string,
+
+    /// The job step node name associated with this recipe
+    jobStep: string,
+
+    /// The workflow id associated with this recipe
+    workflowId?: string,
+}
+
+export type TestAuditPost = {
+
+    /// Whether or not the test is scheduled to run
+    enableSchedule: boolean,
+
+    /// The date of the last audit (ISO formatted)
+    lastAuditDate?: string,
+
+    /// Whether or not the test is under audit
+    isUnderAudit?: boolean,
+
+    /// The owner of the test
+    owner?: string,
+
+    /// The customers of the test
+    customers?: string[],
+
+    /// The notification settings as stream/workflow pair
+    notifications?: Record<string, string>,
+
+    // The team associated with this test
+    team?: string,
+
+    /// The hardness associated with this test
+    harness?: string,
+
+    /// The intent associated with this test
+    intent?: string,
+
+    /// The profile associated with this test
+    profile?: string,
+
+    /// The notes associated with this audit post
+    notes?: string
+}
+
+export type TestAuditPostRequest = TestAuditPost & {
+
+    /// The test id associated with the audit post
+    testId: string
+}
+
+export type TestAuditResponse = TestAuditPost & {
+
+    /// The test id associated with this audit
+    testId: string,
+
+    /// The used id associated with this audit
+    userId: string,
+
+    /// The date of the audit post
+    userInputDate: string,
+}
+
 export async function getTestDataV2(id: string): Promise<TestData> {
     return new Promise<TestData>((resolve, reject) => {
         backend.fetch.get(`/api/v3/testdata/${id}`, {
@@ -482,6 +625,17 @@ export async function getTestTagsV2(tagIds?: string[], tagNames?: string[]): Pro
     });
 }
 
+export async function getTestRecipesV2(recipeIds?: string[], testIds?: string[], streamIds?: string[]): Promise<TestRecipeResponse[]> {
+    const request: TestRecipeRequest = { recipeIds: recipeIds, testIds: testIds, streamIds: streamIds };
+    return new Promise<TestRecipeResponse[]>((resolve, reject) => {
+        backend.fetch.post(`/api/v3/testdata/recipes`, request).then((value) => {
+            resolve(value.data as TestRecipeResponse[]);
+        }).catch((reason) => {
+            reject(reason);
+        });
+    });
+}
+
 export async function getTestPhasesV2(testIds?: string[], testKeys?: string[]): Promise<GetTestPhaseResponse[]> {
     const request: GetTestPhasesRequest = { testIds: testIds, testKeys: testKeys };
     return new Promise<GetTestPhaseResponse[]>((resolve, reject) => {
@@ -561,3 +715,29 @@ export async function getJobStepTestDataV2(jobId: string, stepId: string): Promi
         });
     });
 }
+
+export async function postTestAuditV2(testId: string, auditPost: TestAuditPost): Promise<TestAuditResponse> {
+    const request: TestAuditPostRequest = { testId: testId, ...auditPost };
+    return new Promise<TestAuditResponse>((resolve, reject) => {
+        backend.fetch.post(`/api/v3/testdata/audit/post`, request).then((value) => {
+            resolve(value.data as TestAuditResponse);
+        }).catch((reason) => {
+            reject(reason);
+        });
+    });
+}
+
+export async function getTestAuditsV2(testIds: string[], minCreateTime?: string): Promise<TestAuditResponse[]> {
+    const query: any = {};
+    if (minCreateTime) {
+        query.minCreateTime = minCreateTime;
+    }
+    const request = { testIds: testIds };
+    return new Promise<TestAuditResponse[]>((resolve, reject) => {
+        backend.fetch.post(`/api/v3/testdata/audit/history`, request, query).then((value) => {
+            resolve(value.data as TestAuditResponse[]);
+        }).catch((reason) => {
+            reject(reason);
+        });
+    });
+}
diff --git a/Engine/Source/Programs/Horde/HordeDashboard/plugins/test-automation/testAutomationCommon.tsx b/Engine/Source/Programs/Horde/HordeDashboard/plugins/test-automation/testAutomationCommon.tsx
--- a/Engine/Source/Programs/Horde/HordeDashboard/plugins/test-automation/testAutomationCommon.tsx
+++ b/Engine/Source/Programs/Horde/HordeDashboard/plugins/test-automation/testAutomationCommon.tsx
@@ -1,9 +1,9 @@
 // Copyright Epic Games, Inc. All Rights Reserved.
 
-import { mergeStyles, mergeStyleSets, Text, ComboBox, IComboBoxOption, SelectableOptionMenuItemType, IContextualMenuProps, IContextualMenuItem, IconButton, IContextualMenuListProps, Stack, Checkbox, FontIcon, Modal, DefaultButton, ScrollablePane, SelectionZone, DetailsList, Selection, DetailsListLayoutMode, SelectionMode, IColumn, PrimaryButton } from "@fluentui/react";
+import { mergeStyles, mergeStyleSets, Text, ComboBox, IComboBoxOption, SelectableOptionMenuItemType, IContextualMenuProps, IContextualMenuItem, IconButton, IContextualMenuListProps, Stack, Checkbox, FontIcon, Modal, DefaultButton, ScrollablePane, SelectionZone, DetailsList, Selection, DetailsListLayoutMode, SelectionMode, IColumn, PrimaryButton, Label, TextField } from "@fluentui/react";
 import { MetadataRef, PhaseSessionResult, TestDataHandler, TestNameRef, TestPhaseStatus, TestSessionDetails, TestSessionResult } from "./testData";
 import dashboard, { StatusColor } from "horde/backend/Dashboard";
-import { useCallback, useEffect, useMemo, useState } from "react";
+import { memo, useCallback, useEffect, useMemo, useState } from "react";
 import { TestOutcome, TestPhaseOutcome } from "./api";
 import { PhaseHistoryGraph } from "./phaseHistoryGraph";
 import { projectStore } from "horde/backend/ProjectStore";
@@ -61,31 +61,47 @@ let _styles: any;
 
 const getStyles = () => {
 
-   const border = `1px solid ${dashboard.darktheme ? "#2D2B29" : "#EDEBE9"}`
-
-   const styles = _styles ?? mergeStyleSets({
-      list: {
-         selectors: {
-            'a': {
-               height: "unset !important",
-            },
-            '.ms-List-cell': {
-
-               borderTop: border,
-               borderRight: border,
-               borderLeft: border
-            },
-            '.ms-List-cell:nth-last-child(-n + 1)': {
-               borderBottom: border
-            },
-            ".ms-DetailsRow #artifactview": {
-               opacity: 0
-            },
-            ".ms-DetailsRow:hover #artifactview": {
-               opacity: 1
-            },
-         }
-      }
+    const border = `1px solid ${dashboard.darktheme ? "#2D2B29" : "#EDEBE9"}`;
+    const highlight = dashboard.darktheme ? "brightness(1.2)" : "brightness(0.9)";
+
+    const styles = _styles ?? mergeStyleSets({
+        list: {
+            selectors: {
+                'a': {
+                height: "unset !important",
+                },
+                '.ms-List-cell': {
+
+                borderTop: border,
+                borderRight: border,
+                borderLeft: border
+                },
+                '.ms-List-cell:nth-last-child(-n + 1)': {
+                borderBottom: border
+                },
+                ".ms-DetailsRow #artifactview": {
+                opacity: 0
+                },
+                ".ms-DetailsRow:hover #artifactview": {
+                opacity: 1
+                },
+            }
+        },
+        editable: {
+            selectors: {
+                '> .Edit': {
+                    backgroundColor: "#67676748",
+                    visibility: 'hidden',
+                    cursor: 'pointer'
+                },
+                '> .Edit:hover': {
+                    filter: highlight,
+                },
+                ':hover > .Edit': {
+                    visibility: 'visible'
+                }
+            }
+        },
    });
 
    _styles = styles;
@@ -818,6 +834,27 @@ export const PhaseArtifactsModal: React.FC<{ phase: TestPhaseStatus, artifactPat
         </Stack>
 }
 
+export class DomainNameColors {    
+    setDomain(Names: string[]) {
+        // Only reset the colors if the requested name array is different than what we've already got
+        if ((Names.length > 0) && (JSON.stringify(this._colorWheel.domain()) !== JSON.stringify(Names))) {
+            this._colorWheel.domain(Array.from(Names));
+        }
+    }
+    
+    getDomain() {
+        return this._colorWheel.domain();
+    }
+
+    get(name?: string) {
+        if (!name) return;
+        return this._colorWheel(name) as string;
+    }
+
+    private _colorRange = dashboard.darktheme ? [...d3.schemeDark2, ...d3.schemeCategory10] : [...d3.schemeSet2,...d3.schemeTableau10];
+    private _colorWheel = d3.scaleOrdinal().range(this._colorRange);
+}
+
 class TagColors {
     
     setTagColors(tagNames: string[]) {
@@ -843,3 +880,53 @@ class TagColors {
 }
 export const tagColors = new TagColors();
 
+export const getTestUniqueMetaIdentifiers = (testMeta: Map<string, [TestNameRef, MetadataRef[]][]>): Map<string, string[]> => {
+    // Convoluted way to get unique meta identifier for each test with the same name
+    return new Map(testMeta.entries().filter(([, testRefs]) => testRefs.length > 1).map(([testName, testRefs]) => {
+        const testCommonKeys: string[][] = testRefs.map(([, metaRefs]) => MetadataRef.identifyCommonKeys(metaRefs));
+        const commonKeys: string[] = MetadataRef.identifyCommonKeys(testRefs.map(([, metaRefs]) => metaRefs).flat());
+        const uniqueKeys = testCommonKeys.reduce(
+            (current, testCommonKey) => current.filter((key) => testCommonKey.filter(k => !commonKeys.includes(k)).includes(key)),
+            testCommonKeys.shift() ?? []
+        );
+        return [testName, uniqueKeys];
+    }));
+}
+
+export const StreamSelector: React.FC<{ streams: string[], onClick: (stream: string) => void, selected?: string, disabled?: boolean }> = memo(({ streams, onClick, selected, disabled }) => {
+    const menuProps: IContextualMenuProps = getStreamOptions(streams, onClick, selected);
+    return <DefaultButton
+                style={{ height: 22, fontSize: 11, whiteSpace: 'nowrap', borderWidth: 0 }}
+                menuProps={menuProps}
+                text={selected? projectStore.streamById(selected)?.fullname ?? selected : "Compare stream"}
+                disabled={disabled}
+            />
+});
+
+export const EditableTextField: React.FC<{ label: string, onChange: (value: string | undefined) => void, text?: string, placeholder?: string, multiline?: boolean, style?: any, loaded?: boolean }> = ({ label, text, onChange, placeholder, multiline, style, loaded }) => {
+
+    const [enableEditing, setEnableEditing] = useState<boolean>(false); 
+
+    const styles = getStyles();
+
+    return <Stack className={styles.editable} horizontal verticalAlign="center" tokens={{childrenGap: 6}} style={style} grow>
+                <Stack horizontal tokens={{childrenGap: 6}} grow>
+                    <Label>{label}: </Label>
+                    <TextField
+                        key={`text-${loaded}`}
+                        defaultValue={text}
+                        onNotifyValidationResult={(_, value) => { if (value !== text) { onChange(value); } setEnableEditing(false); }}
+                        validateOnFocusOut
+                        validateOnLoad={false}
+                        multiline={multiline}
+                        styles={{root: {width: '100%'}}}
+                        placeholder={placeholder}
+                        disabled={!enableEditing}
+                    />
+                </Stack>
+                <Stack className="Edit" style={{height: multiline? 50 : 28, padding: 2}} verticalFill verticalAlign="center" onClick={() => {setEnableEditing(!enableEditing);}}>
+                    <FontIcon style={{ fontSize: 12 }} iconName="Edit" />
+                </Stack>
+            </Stack>    
+}
+
diff --git a/Engine/Source/Programs/Horde/HordeDashboard/plugins/test-automation/testAutomationHealth.tsx b/Engine/Source/Programs/Horde/HordeDashboard/plugins/test-automation/testAutomationHealth.tsx
new file mode 100644
--- /dev/null
+++ b/Engine/Source/Programs/Horde/HordeDashboard/plugins/test-automation/testAutomationHealth.tsx
@@ -0,0 +1,814 @@
+// Copyright Epic Games, Inc. All Rights Reserved.
+
+import { DetailsList, FontIcon, IColumn, IconButton, Label, Modal, SelectionMode, Stack, Text, HoverCard, HoverCardType, IPlainCardProps, DirectionalHint, DetailsListLayoutMode, Spinner, SpinnerSize, Toggle, DatePicker, ComboBox, IComboBoxOption, FocusZone, FocusZoneTabbableElements, FocusTrapCallout, PrimaryButton, DefaultButton } from "@fluentui/react";
+import { MetadataRef, TestDataHandler, TestNameRef, TestSessionResult } from "./testData";
+import { useCallback, useEffect, useMemo, useState } from "react";
+import { observer } from "mobx-react-lite";
+import { TestOutcome } from "./api";
+import { DomainNameColors, EditableTextField, getTestUniqueMetaIdentifiers, StreamSelector } from "./testAutomationCommon";
+import { TestHistoryGraphWidget } from "./testAutomationSummary";
+import { StoreKey, userStorage } from "./testAutomationUserStorage";
+import { projectStore } from "horde/backend/ProjectStore";
+import dashboard from "horde/backend/Dashboard";
+import { getShortNiceTime, msecToElapsed } from "horde/base/utilities/timeUtils";
+import { booleanColors, cadenceColors, cadenceOrders, healthColors, healthStars, TestCadence, TestHealth, TestHealthExplain, TestHealthItem, TestHealthTiming } from "./testHealthData";
+
+/// Colors
+const profileColors = new DomainNameColors();
+
+const teamColors = new DomainNameColors();
+
+const harnessColors = new DomainNameColors();
+
+const healthBadgeBackground = () => dashboard.darktheme? "#67676748" : "#8f8f8f9e";
+const textEmphasisColor = () => dashboard.darktheme? "#ffffff75" : "#00000075";
+
+const onRenderExplainCard = (explainData: TestHealthExplain): JSX.Element => {
+    return <Stack style={{pointerEvents: 'none', padding: 10}} tokens={{childrenGap: 5}}>
+                {!!explainData.successRate && <Text>{explainData.successRate}% Success</Text>}
+                {!!explainData.failureRate &&
+                    <Stack>
+                        <Text>{explainData.failureRate}% Failure</Text>
+                        {!!explainData.catastrophicFailureRate && <Text style={{paddingLeft: 25}}>&bull; {explainData.catastrophicFailureRate}% Catastrophic</Text>}
+                        {!!explainData.redundantErrorRate && <Text style={{paddingLeft: 25}}>&bull; {explainData.redundantErrorRate}% Redundant</Text>}
+                    </Stack>
+                }
+            </Stack>
+}
+
+const renderHealthBadge = (item: TestHealthItem, width?: number): JSX.Element => {
+    const StatusCardProps: IPlainCardProps | undefined = item.healthExplain? {
+        onRenderPlainCard: onRenderExplainCard,
+        renderData: item.healthExplain,
+        directionalHint: DirectionalHint.rightCenter,
+        calloutProps: { isBeakVisible: true }
+    } : { onRenderPlainCard: () => null /** empty card */ };
+    return <HoverCard plainCardProps={StatusCardProps} type={HoverCardType.plain} cardOpenDelay={20} sticky={false}>
+                <Stack style={{backgroundColor: healthBadgeBackground(), padding: 3, borderRadius: 4, width: width, height: 24, cursor: 'inherited'}} horizontal verticalAlign="center" title={item.health} tokens={{childrenGap: 3}} horizontalAlign="center">
+                    {new Array(healthStars.get(item.health)).fill(1).map((_, i) => <FontIcon key={`star-${i}`} style={{ fontSize: 11, color: healthColors.get(item.health) }} iconName="Star" />)}
+                </Stack>
+            </HoverCard>
+}
+
+type TestHealthColumn = IColumn & {
+    compare?: (a: TestHealthItem, b: TestHealthItem) => number;
+}
+
+type TestHealthListState = {
+    items: TestHealthItem[];
+    columns: TestHealthColumn[];
+}
+
+const getHealthColumns = (items: TestHealthItem[], columns: TestHealthColumn[], onClickItem: (item: TestHealthItem) => void, setListState: (state: TestHealthListState) => void): TestHealthColumn[] => {
+    const innerState: TestHealthListState = {
+        items: items,
+        columns: []
+    }
+    
+    const onClickColumn = (column: TestHealthColumn) => {
+        const newColumns: TestHealthColumn[] = innerState.columns.slice();
+        const cColumn: TestHealthColumn = newColumns.find(c => column.key === c.key)!;
+        newColumns.forEach(newColumn => {
+            if (newColumn === cColumn) {
+                newColumn.isSortedDescending = !newColumn.isSortedDescending;
+                newColumn.isSorted = true;
+            } else {
+                newColumn.isSorted = false;
+                newColumn.isSortedDescending = !newColumn.isSortedDescending;
+            }
+        });
+
+        const newItems = items.slice().sort((a, b) => (cColumn.isSortedDescending? 1 : -1) * (cColumn.compare?.(a, b) ?? 0));
+
+        setListState({
+            items: newItems,
+            columns: newColumns
+        })
+    }
+    
+    const newColumns: TestHealthColumn[] = [
+        {
+            key: 'name',
+            name: 'Test',
+            minWidth: 32,
+            maxWidth: 250,
+            isResizable: true,
+            isSorted: false,
+            isSortedDescending: true,
+            onRender: (item: TestHealthItem) => {
+                const maskMetaString = item.displayMeta;
+                let title = item.name;
+                if (maskMetaString) title += ` - ${maskMetaString}`;
+                return <Stack onClick={() => onClickItem(item)} style={{cursor: 'pointer'}} horizontal>
+                            <Text title={title} style={{textOverflow: 'ellipsis', overflow: 'hidden', whiteSpace: 'nowrap'}}>
+                                {item.name}{maskMetaString && <span style={{fontSize: 11}}> - {maskMetaString}</span>}
+                            </Text>
+                        </Stack>
+            },
+            onColumnClick: (_, column) => onClickColumn(column),
+            compare: (a: TestHealthItem, b: TestHealthItem) => a.name.localeCompare(b.name)
+        },
+        {
+            key: 'health',
+            name: 'Health',
+            minWidth: 80,
+            maxWidth: 80,
+            isResizable: true,
+            isCollapsible: true,
+            isSorted: true,
+            isSortedDescending: true,
+            onRender: (item: TestHealthItem) => <Stack onClick={() => onClickItem(item)} style={{cursor: 'pointer'}}>{renderHealthBadge(item)}</Stack>,
+            onColumnClick: (_, column) => onClickColumn(column),
+            compare: (a: TestHealthItem, b: TestHealthItem) => healthStars.get(a.health)! - healthStars.get(b.health)!
+        },
+        {
+            key: 'profile',
+            name: 'Profile',
+            minWidth: 90,
+            maxWidth: 90,
+            isResizable: true,
+            isCollapsible: true,
+            isSorted: false,
+            isSortedDescending: true,
+            onRender: (item: TestHealthItem) => {
+                const profile = item.profile ?? "N/A";
+                return <Stack style={{backgroundColor: profileColors.get(item.profile), padding: 3, borderRadius: 4, cursor: 'default'}} horizontalAlign="center" horizontal>
+                            <Text title={profile} style={{textOverflow: 'ellipsis', overflow: 'hidden', whiteSpace: 'nowrap'}}>{profile}</Text>
+                        </Stack>
+            },
+            compare: (a: TestHealthItem, b: TestHealthItem) => a.profile?.localeCompare(b.profile ?? "") ?? -1
+        },
+        {
+            key: 'team',
+            name: 'Team',
+            minWidth: 70,
+            maxWidth: 70,
+            isResizable: true,
+            isCollapsible: true,
+            isSorted: false,
+            isSortedDescending: true,
+            onRender: (item: TestHealthItem) => {
+                const team = item.team  ?? "N/A";
+                return <Stack style={{backgroundColor: teamColors.get(item.team), padding: 3, borderRadius: 4, cursor: 'default'}} horizontalAlign="center" horizontal>
+                            <Text title={team} style={{textOverflow: 'ellipsis', overflow: 'hidden', whiteSpace: 'nowrap'}}>{team}</Text>
+                        </Stack>
+            },
+            compare: (a: TestHealthItem, b: TestHealthItem) => a.team?.localeCompare(b.team ?? "") ?? -1
+        },
+        {
+            key: 'hasNotification',
+            name: 'Notification',
+            minWidth: 32,
+            maxWidth: 70,
+            isResizable: true,
+            isCollapsible: true,
+            isSorted: false,
+            isSortedDescending: true,
+            onRender: (item: TestHealthItem) => {
+                return <Stack style={{backgroundColor: booleanColors.get(!item.notification), padding: 3, borderRadius: 4, cursor: 'default'}} horizontalAlign="center">
+                            <Text>{!!item.notification? "yes" : "no"}</Text>
+                        </Stack>
+            },
+            onColumnClick: (_, column) => onClickColumn(column),
+            compare: (a: TestHealthItem, b: TestHealthItem) => a.notification?.localeCompare(b.notification ?? "") ?? -1
+        },
+        {
+            key: 'cadence',
+            name: 'Cadence',
+            minWidth: 100,
+            maxWidth: 100,
+            isResizable: true,
+            isCollapsible: true,
+            isSorted: false,
+            isSortedDescending: true,
+            onRender: (item: TestHealthItem) => {
+                return <Stack style={{backgroundColor: cadenceColors.get(item.cadence), padding: 3, borderRadius: 4, cursor: 'default'}} horizontalAlign="center">
+                            <Text>{item.cadence}</Text>
+                        </Stack>
+            },
+            onColumnClick: (_, column) => onClickColumn(column),
+            compare: (a: TestHealthItem, b: TestHealthItem) => cadenceOrders.get(a.cadence)! - cadenceOrders.get(b.cadence)!
+        },
+        {
+            key: 'lastCompletedDate',
+            name: 'Last Uninterrupted Run',
+            minWidth: 32,
+            maxWidth: 100,
+            isResizable: true,
+            isCollapsible: true,
+            isSorted: false,
+            isSortedDescending: true,
+            onRender: (item: TestHealthItem) => {
+                return <Text>{getShortNiceTime(item.timings?.latestCompletedRun, true, true)}</Text>
+            },
+            onColumnClick: (_, column) => onClickColumn(column),
+            compare: (a: TestHealthItem, b: TestHealthItem) => (a.timings?.latestCompletedRun?.getTime() ?? 0) - (b.timings?.latestCompletedRun?.getTime() ?? 0)
+        },
+    ];
+    
+    newColumns.forEach(column => {
+        innerState.columns.push(column);
+        if (columns.length > 0) {
+            // copy column states
+            const oldColumn = columns.find(c => column.key === c.key);
+            if (oldColumn) {
+                column.isSorted = oldColumn.isSorted;
+                column.isSortedDescending = oldColumn.isSortedDescending;
+            }
+        }
+    });
+
+    const sortedColumn = newColumns.find(item => item.isSorted);
+    if (sortedColumn) {
+        items.sort((a, b) => (sortedColumn.isSortedDescending? 1 : -1) * (sortedColumn.compare?.(a, b) ?? 0))
+    }
+
+    return innerState.columns;
+}
+
+const percRound = (value: number) => Math.round(value * 100);
+
+const getTestHealthItems = (handler: TestDataHandler, stream?: string, test?: TestNameRef) => {
+    if (!stream) {
+        return [];
+    }
+
+    const testStatus = handler.getStatusStream(stream);
+    if (!testStatus?.tests.size) {
+        return [];
+    }
+
+    let filteredTests = handler.filteredTests;
+    if (!filteredTests.length) {
+        return [];
+    }
+    if (!!test) {
+        if (!filteredTests.find(t => t === test)) {
+            return [];
+        }
+        filteredTests = [test];
+    }
+
+    const filteredMetadata = handler.filteredMetadata;
+    
+    // collect data and reorganize it for compution
+    const testHealths: TestHealthItem[] = [];
+    const testMeta: Map<string, [TestNameRef, MetadataRef[]][]> = new Map();
+
+    filteredTests.forEach((test) => {
+        const status = testStatus!.tests.get(test);
+        if (!status) return;
+
+        let totalPhaseCount = 0;
+        let totalUndefinedCount = 0;
+        let totalFailureCount = 0;
+        let totalSuccessCount = 0;
+        let latestCompletedRun: Date | undefined;
+        const testCompleteRunDurations: number[] = [];
+        const testCatastrophicDurations: number[] = [];
+        const errorFingerprints: Set<string> = new Set<string>();
+        const recipeIds: Set<string> = new Set<string>();
+
+        const metaRefs: MetadataRef[] = [];
+        const testRefMeta: [TestNameRef, MetadataRef[]][] = testMeta.get(test.name) ?? testMeta.set(test.name, []).get(test.name)!;
+        testRefMeta.push([test, metaRefs]);
+
+        filteredMetadata.forEach(meta => {
+            const metaStatus = status.sessions.get(meta);
+            if (!metaStatus) return;
+
+            const session = metaStatus.getLastSession();
+            if (!session) return;
+
+            if (session.recipeId) recipeIds.add(session.recipeId);
+
+            metaStatus.history.forEach(session => {
+                const outcome = session.outcome;
+                if ( outcome === TestOutcome.Skipped ) return;
+
+                switch (outcome) {
+                    case TestOutcome.Unspecified:
+                        testCatastrophicDurations.push(session.duration);
+                        break;
+                    case TestOutcome.Failure:
+                    case TestOutcome.Success:
+                        testCompleteRunDurations.push(session.duration);
+                        if (!latestCompletedRun || session.start > latestCompletedRun) {
+                            latestCompletedRun = session.start;
+                        }
+                        break;
+                    default:
+                }
+
+                totalPhaseCount += session.phasesSucceededCount + session.phasesUnspecifiedCount + session.phasesFailedCount;
+                totalFailureCount += session.phasesUnspecifiedCount + session.phasesFailedCount;
+                totalUndefinedCount += session.phasesUnspecifiedCount;
+                totalSuccessCount += session.phasesSucceededCount;
+                session.errorFingerprints?.keys().forEach(error => errorFingerprints.add(error));
+            })
+
+            metaRefs.push(meta);
+        });
+
+        // Compute Cadence
+        let cadenceValue = TestCadence.NA;
+        if (!!handler.filterState.weeks) {
+            const highestRunCount = status.sessions.get(metaRefs.reduce((pmeta, meta) => {
+                const pCount = !!pmeta? status.sessions.get(pmeta)?.history.length ?? 0 : 0;
+                const cCount = status.sessions.get(meta)?.history.length ?? 0;
+                return cCount >= pCount ? meta : pmeta;
+            }, undefined)!)!.history.length;
+
+            if (highestRunCount > 0) {
+                const rangeDayCount = handler.filterState.weeks * 7;
+                const cadenceRate = highestRunCount / rangeDayCount;
+                if (cadenceRate > 2) {
+                    /** more than 2 runs / 1 day */
+                    cadenceValue = TestCadence.Incremental;
+                } else if (cadenceRate > 0.4) {
+                    /** more than 3 runs / 7 days */
+                    cadenceValue = TestCadence.Daily;
+                } else if (cadenceRate > 0.14) {
+                    /** more than 1 run / 7 days */
+                    cadenceValue = TestCadence.Weekly;
+                } else {
+                    cadenceValue = TestCadence.Limited;
+                }
+            } else {
+                cadenceValue = TestCadence.NotRun;
+            }
+        }
+
+        // Compute Health Heuristic
+        let healthValue = TestHealth.NA;
+        let healthExplain: TestHealthExplain | undefined;
+        if (totalPhaseCount > 0) {
+            const catastrophicFailureRate = totalFailureCount == 0 ? 0: percRound(totalUndefinedCount/totalFailureCount);
+            const failureRate = percRound(totalFailureCount/totalPhaseCount);
+            const successRate = percRound(totalSuccessCount/totalPhaseCount);
+            const redundantErrorRate = totalFailureCount == 0 ? 0: percRound(1 - errorFingerprints.size/totalFailureCount);
+            let nStars = healthStars.get(TestHealth.Fair)!;
+            if (failureRate > 15) {
+                if (catastrophicFailureRate > 40 || successRate == 0) {
+                    nStars -= 2;
+                } else {
+                    nStars -= 1;
+                }
+                if (redundantErrorRate < 80) {
+                    nStars += redundantErrorRate < 60 ? 1 : 0;
+                } else if (failureRate > 40) {
+                    nStars -= 1;
+                }
+            }
+            if (successRate > 60) {
+                nStars += 1;
+            }
+            healthValue = healthStars.entries().find(([_, value]) => value === nStars)?.[0] ?? healthValue;
+            healthExplain = {
+                catastrophicFailureRate: catastrophicFailureRate,
+                failureRate: failureRate,
+                successRate: successRate,
+                redundantErrorRate: redundantErrorRate
+            }
+        }
+
+        // Compute timings
+        const timings: TestHealthTiming = {
+            averageToInterruptSecs: testCatastrophicDurations.length === 0? 0 : testCatastrophicDurations.reduce((acc, value) => acc + value, 0) / testCatastrophicDurations.length,
+            averageToCompleteRunSecs: testCompleteRunDurations.length === 0? 0 : testCompleteRunDurations.reduce((acc, value) => acc + value, 0) / testCompleteRunDurations.length,
+            latestCompletedRun: latestCompletedRun,
+        }
+
+        const testHealthItem = new TestHealthItem(test, handler, stream);
+        testHealthItem.metaRefs = metaRefs;
+        testHealthItem.health = healthValue;
+        testHealthItem.healthExplain = healthExplain;
+        testHealthItem.cadence = cadenceValue;
+        testHealthItem.setLazyDisplayMeta(() => status.sessions.keys().find((meta) => !!meta)?.getCommonValues(testMetaKeys.get(test.name))?.join(" / "));
+        testHealthItem.timings = timings;
+        testHealthItem.linkRecipeIds(recipeIds);
+
+        testHealths.push(testHealthItem);
+    });
+
+    const testMetaKeys = getTestUniqueMetaIdentifiers(testMeta);
+
+    return testHealths;
+}
+
+const TestHealthTimes: React.FC<{ test: TestHealthItem }> = ({test}) => {
+    const [init, setInit] = useState<boolean>(false);
+
+    useEffect(() => {
+        test.fetchRecipes().then(() => setInit(true));
+    }, []);
+
+    if (!test.timings) return null;
+
+    const averageToCompleteRunString = useMemo(() => test.timings? msecToElapsed(test.timings.averageToCompleteRunSecs * 1000, true, true) : "", [test]);
+    const averageToInterruptString = useMemo(() => test.timings? msecToElapsed(test.timings.averageToInterruptSecs * 1000, true, true) : "", [test]);
+    const timeoutString = useMemo(() => test.recipes? msecToElapsed(test.recipes.map(r => r.timeoutMinutes).reduce((acc, value) => value > acc? value : acc, 0 ) * 60 * 1000, true, true) : "", [test, init]);
+
+    return <Text>
+                {!!test.timings.averageToCompleteRunSecs &&
+                    <span>
+                        <span>Completed run Avg. time:</span>
+                        <span style={{color: textEmphasisColor()}}> {averageToCompleteRunString}</span>
+                    </span>
+                }                                    
+                {!test.timings.averageToCompleteRunSecs && <span>No completed run recorded</span> }
+                {!!test.timings.averageToInterruptSecs &&
+                    <span>
+                        <span> | Interrupted run Avg. time:</span>
+                        <span style={{color: textEmphasisColor()}}> {averageToInterruptString}</span>
+                    </span>
+                }
+                {!!test.recipes &&
+                    <span>
+                        <span> | Max Timeout:</span>
+                        <span  style={{color: textEmphasisColor()}}> {timeoutString}</span>
+                    </span>
+                }
+            </Text> 
+}
+
+export const TestHealthView: React.FC<{ test: TestHealthItem, handler: TestDataHandler, onUpdateItem: (item: TestHealthItem) => void }> = observer(({ test, handler, onUpdateItem }) => {
+    const [streams, setStreams] = useState<{current: string, toCompare?: string}>(
+        {
+            current: handler.stream!,
+            toCompare: userStorage.getItem(StoreKey.HistoryToCompare, () => undefined)
+        });
+    const [testHealthToCompare, setTestHealthToCompare] = useState<TestHealthItem | undefined>(getTestHealthItems(handler, userStorage.getItem(StoreKey.HistoryToCompare), test.nameRef).find(() => true));
+
+    handler.subscribeToAuditsLoading();
+    handler.subscribeToRecipesLoading();
+    test.subscribeUpdate();
+
+    // collecting metadata
+    let metadataSet = new Set(handler.getStatusStream(streams.current)!.tests.get(test.nameRef)!.getMetadata());
+    if (!!streams.toCompare) {
+        const toCompareMetadataList = handler.getStatusStream(streams.toCompare!)?.tests.get(test.nameRef)?.getMetadata();
+        if (!!toCompareMetadataList) {
+            metadataSet = metadataSet.union(new Set(toCompareMetadataList!));
+        }
+    }
+
+    const commonMetaKeys = metadataSet.size == 1? handler.filterState.metadata?.keys().toArray() : MetadataRef.identifyCommonKeys(Array.from(metadataSet));
+    const availableStreams = handler.availableStreams.filter((s) => s !== streams.current);
+
+    const onClickStream = useCallback((targetStream: string) => {
+        setStreams({...streams, toCompare: undefined});
+        handler.query(targetStream, true)
+            .catch((reason) => console.error(reason))
+            .finally(() => {
+                setTestHealthToCompare(getTestHealthItems(handler, targetStream, test.nameRef).find(_ => true));
+                setStreams({...streams, toCompare: targetStream});
+                userStorage.setItem(StoreKey.HistoryToCompare, targetStream);
+            });
+    }, [streams, test]);
+
+    const onDetailsClick = useCallback((session: TestSessionResult) => {
+        if (session.testDataId) {
+            handler.setSearchParam('session', session.testDataId);
+        }
+        else if (session.jobId && session.stepId) {
+            window.open(`/job/${session.jobId}?step=${session.stepId}`, '_blank');
+        }
+    }, []);
+
+    const profileOptions: IComboBoxOption[] = profileColors.getDomain().map(p => ({key: p, text: p, styles: {optionText: {backgroundColor: profileColors.get(p), padding: 3, borderRadius: 4, minWidth: 72}}}));
+    const teamOptions: IComboBoxOption[] = teamColors.getDomain().map(p => ({key: p, text: p, styles: {optionText: {backgroundColor: teamColors.get(p), padding: 3, borderRadius: 4, minWidth: 72}}}));
+    const harnessOptions: IComboBoxOption[] = harnessColors.getDomain().map(p => ({key: p, text: p}));
+
+    return <Stack tokens={{childrenGap: 16}} style={{position: 'relative'}}>
+                <Stack tokens={{childrenGap: 8}} style={{padding: 14, borderRadius: 4}} >
+                    <Stack horizontal tokens={{childrenGap: 12}} verticalAlign="center">
+                        <ComboBox
+                            key={`profile-${test.updated}`}
+                            styles={{input: {backgroundColor: profileColors.get(test.profile)}, root: {backgroundColor: profileColors.get(test.profile), width: 150}}}
+                            placeholder="Profile"
+                            allowFreeform
+                            autoComplete="on"
+                            options={profileOptions}
+                            text={test.profile}
+                            onChange={(_1, _2, _3, value) => { test.profile = value; onUpdateItem(test); }}
+                        />
+                        <ComboBox
+                            key={`team-${test.updated}`}
+                            styles={{input: {backgroundColor: teamColors.get(test.team)}, root: {backgroundColor: teamColors.get(test.team), width: 150}}}
+                            placeholder="Team"
+                            allowFreeform
+                            autoComplete="on"
+                            options={teamOptions}
+                            text={test.team}
+                            onChange={(_1, _2, _3, value) => { test.team = value; onUpdateItem(test); }}
+                        />
+                        <ComboBox
+                            key={`harness-${test.updated}`}
+                            styles={{root: {width: 170}}}
+                            placeholder="Harness"
+                            allowFreeform
+                            autoComplete="on"
+                            options={harnessOptions}
+                            text={test.harness}
+                            onChange={(_1, _2, _3, value) => { test.harness = value; }}
+                        />
+                    </Stack>
+                    <Stack horizontal tokens={{childrenGap: 12}} style={{paddingLeft: 12, paddingTop: 6}} grow>
+                        <EditableTextField
+                            key={`intent-${test.updated}`}
+                            label="Intent"
+                            text={test.intent}
+                            onChange={(value) => test.intent = value}
+                            placeholder="a short description"
+                            loaded={handler.queryAuditsLoading}
+                        />
+                    </Stack>
+
+                    <Stack horizontal style={{paddingLeft: 12}}>
+                        <Stack horizontal verticalAlign="center" tokens={{childrenGap: 6}} style={{cursor: 'default'}} title="edit disabled for now">
+                            <Label>Worflow Notification: </Label>
+                            <Text style={{color:textEmphasisColor()}}>{!!test.notification? test.notification : "none"}</Text>
+                        </Stack>
+                    </Stack>
+
+                    <Stack horizontal tokens={{childrenGap: 12}} style={{paddingLeft: 12}} title="edit disabled for now">
+                        <Stack style={{borderRadius: 4, cursor: 'default'}} verticalAlign="center">
+                            <Stack  horizontal tokens={{childrenGap: 6}} verticalAlign="center">
+                                <Label>Owner: </Label>
+                            </Stack>
+                            <Text style={{color:textEmphasisColor()}}>{test.owner ?? "N/A"}</Text>
+                        </Stack>
+                        <Stack style={{borderRadius: 4, cursor: 'default'}} verticalAlign="center">
+                            <Stack  horizontal tokens={{childrenGap: 6}} verticalAlign="center">
+                                <Label>Customers: </Label>
+                            </Stack>
+                            <Text style={{color:textEmphasisColor()}}>{test.customers?.join(', ')  ?? "N/A"}</Text>
+                        </Stack>
+                    </Stack>
+
+                    <Stack style={{padding: '4px 12px', backgroundColor: booleanColors.get(test.underAudit)}}>
+                        <Stack horizontal tokens={{childrenGap: 6}} verticalAlign="center">
+                            <Toggle
+                                key={`under-audit-${test.updated}`}
+                                label="Under Audit"
+                                checked={test.underAudit}
+                                onChange={(_, checked) => {test.underAudit = !!checked}}
+                                inlineLabel
+                                onText="Yes" offText="No"
+                            />
+                            <Stack horizontal horizontalAlign="start" tokens={{childrenGap: 6}} grow>
+                                <Label>Last date: </Label>
+                                <DatePicker
+                                    key={`last-audit-${test.updated}`}
+                                    styles={{root: { width: 150 }}}
+                                    placeholder="Select a date... "
+                                    value={test.lastAuditDate}
+                                    onSelectDate={(date) => test.lastAuditDate = (date ?? undefined)}
+                                />
+                            </Stack>
+                        </Stack>
+                        <EditableTextField
+                            key={`notes-${test.updated}`}
+                            label="Notes"
+                            text={test.notes}
+                            onChange={(value) => test.notes = value}
+                            multiline
+                            placeholder="add notes"
+                            loaded={handler.queryAuditsLoading}
+                        />
+                    </Stack>
+
+                </Stack>
+                <Stack style={{paddingLeft: 16}} tokens={{childrenGap: 16}}>
+                    <Stack style={{paddingLeft: 12}} tokens={{childrenGap: 16}}>
+                        <Stack horizontal tokens={{childrenGap: 12}} verticalAlign="center">
+                            <Stack style={{width: 200, backgroundColor: "#8e8e8e48", padding:  '3px 12px', borderRadius: 3}}>
+                                <Text>{projectStore.streamById(streams.current)?.fullname ?? streams.current}</Text>
+                            </Stack>
+                            <Stack style={{cursor: 'default'}}>{renderHealthBadge(test, 75)}</Stack>
+                            <Stack style={{backgroundColor: cadenceColors.get(test.cadence), padding:  '3px 12px', borderRadius: 4, cursor: 'default'}} verticalAlign="center">
+                                <Text title="Cadence">{test.cadence}</Text>
+                            </Stack>
+                            {test.timings && <TestHealthTimes test={test} />}
+                        </Stack>
+                        <Stack horizontal tokens={{childrenGap: 16}}>
+                            <Stack style={{backgroundColor: booleanColors.get(false), width: 6, borderRadius: 3}}></Stack>
+                            <TestHistoryGraphWidget test={test.nameRef} stream={streams.current} commonMetaKeys={commonMetaKeys} handler={handler} onClick={onDetailsClick} />
+                        </Stack>
+                    </Stack>
+                    <Stack style={{paddingLeft: 12}} tokens={{childrenGap: 16}}>
+                        <Stack horizontal tokens={{childrenGap: 12}} verticalAlign="center">
+                            <Stack style={{width: 200}}>
+                                <StreamSelector streams={availableStreams} selected={streams.toCompare} disabled={availableStreams.length === 0} onClick={onClickStream} />
+                            </Stack>
+                            {!!testHealthToCompare && <Stack style={{cursor: 'default'}}>{renderHealthBadge(testHealthToCompare, 75)}</Stack>}
+                            {!!testHealthToCompare && 
+                                <Stack style={{backgroundColor: cadenceColors.get(testHealthToCompare.cadence), padding:  '3px 12px', borderRadius: 4, cursor: 'default'}} verticalAlign="center">
+                                    <Text title="Cadence">{testHealthToCompare.cadence}</Text>
+                                </Stack>
+                            }
+                            {!!testHealthToCompare && testHealthToCompare.timings && <TestHealthTimes test={testHealthToCompare} />}
+                        </Stack>
+                        <Stack horizontal tokens={{childrenGap: 16}}>
+                            <Stack style={{backgroundColor: booleanColors.get(false), width: 6, borderRadius: 3}}></Stack>
+                            {!streams.toCompare && !handler.subQueryLoading &&
+                                <Stack style={{minHeight: 40}} verticalAlign="center"><Text>select a stream</Text></Stack>
+                            }
+                            {!!streams.toCompare && !!testHealthToCompare && 
+                                <TestHistoryGraphWidget test={test.nameRef} stream={streams.toCompare} commonMetaKeys={commonMetaKeys} handler={handler} onClick={onDetailsClick} />
+                            }
+                            {!!streams.toCompare && !testHealthToCompare && 
+                                <Stack horizontalAlign='center' style={{ padding: 4 }}>
+                                    <Text style={{ fontSize: 12 }}>No Results</Text>
+                                </Stack>
+                            }
+                            {handler.subQueryLoading &&
+                                <Stack horizontalAlign='center' style={{ padding: 12}} tokens={{ childrenGap: 10 }} horizontal>
+                                    <Text style={{ fontSize: 12 }}>Loading Data</Text>
+                                    <Spinner size={SpinnerSize.small} />
+                                </Stack>
+                            }
+                        </Stack>
+                    </Stack>
+                </Stack>
+
+                { (handler.queryAuditsLoading || handler.queryRecipesLoading) &&
+                    <Stack horizontalAlign='center' style={{ width: "100%", position: "absolute", top: 0 }} tokens={{childrenGap: 8}}>
+                        <Spinner size={SpinnerSize.large} />
+                    </Stack>
+                }
+            </Stack>
+});
+
+const TestHealthDetailsModal: React.FC<{ test: TestHealthItem, handler: TestDataHandler, onDismiss: () => void, onUpdateItem: (item: TestHealthItem) => void }> = observer(({ test, handler, onDismiss, onUpdateItem }) => {
+    const [isCalloutVisible, setCalloutVisible] = useState(false);
+    const [isCommitting, setCommitting] = useState(false);
+
+    useEffect(() => {
+        // discard changes on clean up
+        return () => { if(test.isTransient) { test.discardTransaction(); onUpdateItem(test); } }
+    }, [test]);
+
+    test.subscribeUpdate();
+
+    return <Modal isOpen={true} isBlocking={true} topOffsetFixed={true} styles={{ main: { padding: 8, width: 1134, hasBeenOpened: false, top: "80px", position: "absolute" } }} onDismiss={() => onDismiss()} >
+            <Stack grow>
+                <Stack horizontal verticalAlign="center">
+                    <Stack><Text style={{ paddingLeft: 8, fontSize: 16, fontWeight: 600 }}>Health Details - {`${test.name}`}</Text></Stack>
+                    <Stack grow />
+                    {test.isTransient &&
+                        <Stack style={{ paddingBottom: 4, paddingRight: 6 }} title="Revert changes">
+                            <IconButton
+                                iconProps={{ iconName: 'History', style: { fontSize: 12 } }}
+                                ariaLabel="Revert changes"
+                                onClick={() => { test.discardTransaction(); onUpdateItem(test); }}
+                            />
+                        </Stack>
+                    }
+                    {test.isTransient &&
+                        <Stack style={{ paddingBottom: 4, paddingRight: 6 }} title="Commit changes and close">
+                            <IconButton
+                                iconProps={{ iconName: 'Save', style: { fontSize: 16 } }}
+                                ariaLabel="Commit and close"
+                                onClick={() => { test.commitTransaction(); onDismiss(); }}
+                            />
+                        </Stack>
+                    }
+                    <Stack style={{ paddingBottom: 4 }} className="Close">
+                        <IconButton
+                            iconProps={{ iconName: 'Cancel' }}
+                            ariaLabel="Close popup modal"
+                            onClick={() => { if(test.isTransient) { setCalloutVisible(true) } else { onDismiss(); } }}
+                        />
+                        {isCalloutVisible &&
+                            <FocusTrapCallout
+                                target=".Close"
+                                style={{padding: '20px 24px', width: 300}}
+                                isBeakVisible
+                                onDismiss={() => setCalloutVisible(false)}
+                            >
+                                {!isCommitting &&
+                                    <Stack tokens={{childrenGap: 6}}>
+                                        <Text block>You have uncommitted changes to this audit.</Text>
+                                        <Text block>Do want to commit them?</Text>
+                                        <FocusZone handleTabKey={FocusZoneTabbableElements.all} isCircularNavigation>
+                                            <Stack horizontal horizontalAlign="center">
+                                            <PrimaryButton 
+                                                iconProps={{ iconName: 'Save', style: { fontSize: 16 } }}
+                                                onClick={() => { setCommitting(true); test.commitTransaction().then(() => { setCommitting(false); setCalloutVisible(false); onDismiss(); }) }}
+                                            >Commit</PrimaryButton>
+                                            <DefaultButton
+                                                iconProps={{ iconName: 'History', style: { fontSize: 12 } }}
+                                                onClick={() => { setCalloutVisible(false); onDismiss(); }}
+                                            >Discard</DefaultButton>
+                                            <DefaultButton
+                                                onClick={() => { setCalloutVisible(false); }}
+                                            >Cancel</DefaultButton>
+                                            </Stack>
+                                        </FocusZone>
+                                    </Stack>
+                                }
+                                {isCommitting &&
+                                    <Stack>
+                                        <Text block>Committing changes...</Text>
+                                    </Stack>
+                                }
+                            </FocusTrapCallout>
+                        }
+                    </Stack>
+                </Stack>
+                <Stack tokens={{ childrenGap: 40 }} style={{ padding: 8, overflow: "auto", maxHeight: 'calc(100vh - 200px)'} }>
+                    <TestHealthView test={test} handler={handler} onUpdateItem={onUpdateItem}/>
+                </Stack>
+            </Stack>
+        </Modal>
+});
+
+export const TestAutomationHealth: React.FC<{ handler: TestDataHandler}> = observer(({ handler }) => {
+    const [healthDetailsState, setHealthDetailsState] = useState<{test?: TestHealthItem}>({});
+    const [healthListState, setHealthListState] = useState<TestHealthListState>({items: [], columns: []});
+
+    const updateColorDomains = useCallback((items: TestHealthItem[]) => {
+        const profiles = new Set(items.map(i => i.profile).filter(i => i) as string[]);
+        const teams = new Set(items.map(i => i.team).filter(i => i) as string[]);
+        const harnesses = new Set(items.map(i => i.harness).filter(i => i) as string[]);
+        profileColors.setDomain(Array.from(profiles).sort());
+        teamColors.setDomain(Array.from(teams).sort());
+        harnessColors.setDomain(Array.from(harnesses).sort())
+    }, []);
+
+    const initHealthListData = useCallback(() => {
+        const initItems = getTestHealthItems(handler, handler.stream);
+        const initColumns = getHealthColumns(initItems, healthListState.columns, (item) => {
+            setHealthDetailsState({test: item});
+            handler.setSearchParam('health', item.nameRef.id);
+        }, setHealthListState);
+
+        updateColorDomains(initItems);
+
+        setHealthListState({
+            items: initItems,
+            columns: initColumns
+        });
+
+        return initItems;
+    }, []);
+
+    const forceListRefresh = () => {
+        updateColorDomains(healthListState.items);
+        setHealthListState({
+            ...healthListState,
+            items: healthListState.items.slice()
+        });
+    }
+
+    // initiate data on filtering change
+    useEffect(() => {
+
+        const items = initHealthListData();
+
+        // load audit async
+        handler.queryAudits(items.map(i => i.nameRef));
+
+    }, [handler.filterUpdated]);
+
+    // find health item if not loaded already
+    const testId = handler.getSearchParam('health') as string | undefined;
+    useEffect(() => {
+        if ((!!testId && !healthDetailsState.test)
+            || (!!healthDetailsState.test && healthDetailsState.test.nameRef.id !== testId)) {
+
+            // if no data were generated for the stream, do it now
+            if (healthListState.items.length === 0) initHealthListData();
+
+            const testHealth = !!testId ? healthListState.items.find(i => i.nameRef.id === testId) : undefined;
+
+            setHealthDetailsState({test: testHealth});
+        }
+    }, [handler.searchUpdated, healthDetailsState]);
+
+    return  <Stack grow style={{position: "relative"}}>
+                {!!healthDetailsState.test &&
+                    <TestHealthDetailsModal
+                        test={healthDetailsState.test}
+                        handler={handler}
+                        onDismiss={() => {setHealthDetailsState({}); handler.removeSearchParam('health')}}
+                        onUpdateItem={(item: TestHealthItem) => { forceListRefresh(); }}
+                    />
+                }
+                <DetailsList
+                    items={healthListState.items}
+                    compact={true}
+                    columns={healthListState.columns}
+                    isHeaderVisible={true}
+                    selectionMode={SelectionMode.none}
+                    layoutMode={DetailsListLayoutMode.justified}
+                />
+                { handler.queryAuditsLoading &&
+                    <Stack horizontalAlign='center' style={{ width: "100%", position: "absolute", top: 0 }} tokens={{childrenGap: 8}}>
+                        <Spinner size={SpinnerSize.large} />
+                    </Stack>
+                }
+            </Stack>
+});
diff --git a/Engine/Source/Programs/Horde/HordeDashboard/plugins/test-automation/testAutomationSummary.tsx b/Engine/Source/Programs/Horde/HordeDashboard/plugins/test-automation/testAutomationSummary.tsx
--- a/Engine/Source/Programs/Horde/HordeDashboard/plugins/test-automation/testAutomationSummary.tsx
+++ b/Engine/Source/Programs/Horde/HordeDashboard/plugins/test-automation/testAutomationSummary.tsx
@@ -1,6 +1,6 @@
 // Copyright Epic Games, Inc. All Rights Reserved.
 
-import { DefaultButton, FontIcon, IconButton, mergeStyleSets, Modal, Stack, Text, HoverCard, HoverCardType, IPlainCardProps, DirectionalHint, Spinner, SpinnerSize, IContextualMenuItem, IContextualMenuProps, IComboBoxOption } from "@fluentui/react";
+import { DefaultButton, FontIcon, IconButton, mergeStyleSets, Modal, Stack, Text, HoverCard, HoverCardType, IPlainCardProps, DirectionalHint, Spinner, SpinnerSize, IComboBoxOption } from "@fluentui/react";
 import { useState, useEffect, useCallback, memo } from "react";
 import { observer } from "mobx-react-lite";
 import { TestDataHandler, TestNameRef, MetadataRef, TestSessionStatus, TestSessionResult, TestStatus, TestTagRef } from "./testData";
@@ -9,21 +9,21 @@ import dashboard, { StatusColor } from "horde/backend/Dashboard";
 import { getShortNiceTime } from "horde/base/utilities/timeUtils";
 import { getHordeStyling } from "horde/styles/Styles";
 import { projectStore } from 'horde/backend/ProjectStore';
-import { SessionStatusBar, SessionValues, StatusBar, StatusBarStack, MultiOptionChooser, getTestSessionStatusColor, getStreamOptions, getStatusColors, statusTexts } from "./testAutomationCommon";
+import { SessionStatusBar, SessionValues, StatusBar, StatusBarStack, MultiOptionChooser, getTestSessionStatusColor, getStatusColors, statusTexts, getTestUniqueMetaIdentifiers, StreamSelector } from "./testAutomationCommon";
 import { TestHistoryGraph } from "./testHistoryGraph";
-import { userStorage } from "./testAutomationUserStorage";
+import { StoreKey, userStorage } from "./testAutomationUserStorage";
 
 type SessionCallback = (session: TestSessionResult) => void;
 type HistorySessionCallback = (test: TestNameRef, onClick: SessionCallback) => void;
 
-const TestHistoryGraphWidget: React.FC<{ test: TestNameRef, stream: string, handler: TestDataHandler, onClick: SessionCallback, commonMetaKeys?: string[], selectedMetaKeys?: string[]}> = ({ test, stream, handler, onClick, commonMetaKeys, selectedMetaKeys }) => {
+export const TestHistoryGraphWidget: React.FC<{ test: TestNameRef, stream: string, handler: TestDataHandler, onClick: SessionCallback, commonMetaKeys?: string[], selectedMetaKeys?: string[]}> = ({ test, stream, handler, onClick, commonMetaKeys, selectedMetaKeys }) => {
     const [container, setContainer] = useState<HTMLDivElement | null>(null);
     const [state, setState] = useState<{ graph?: TestHistoryGraph }>({});
 
     useEffect(() => {
         const graph = new TestHistoryGraph(test, stream, handler, commonMetaKeys, onClick);
         setState({ graph: graph });
-    }, []);
+    }, [commonMetaKeys]);
 
     useEffect(() => {
         // refresh on meta selection update
@@ -52,20 +52,10 @@ const TestHistoryGraphWidget: React.FC<{ test: TestNameRef, stream: string, hand
             </Stack>
 }
 
-const StreamSelector: React.FC<{ streams: string[], onClick: (stream: string) => void, selected?: string, disabled?: boolean }> = memo(({ streams, onClick, selected, disabled }) => {
-    const menuProps: IContextualMenuProps = getStreamOptions(streams, onClick, selected);
-    return <DefaultButton
-                style={{ height: 22, fontSize: 11, whiteSpace: 'nowrap', borderWidth: 0 }}
-                menuProps={menuProps}
-                text={selected? projectStore.streamById(selected)?.fullname ?? selected : "Compare stream"}
-                disabled={disabled}
-            />
-});
-
 export const TestHistoryView: React.FC<{ test: TestNameRef, handler: TestDataHandler, onClick: SessionCallback  }> = observer(({ test, handler, onClick }) => {
 
     const { hordeClasses } = getHordeStyling();
-    const [streams, setStreams] = useState<{current: string, toCompare?: string}>({current: handler.stream!, toCompare: userStorage.getItem('historyToCompare', () => undefined)});
+    const [streams, setStreams] = useState<{current: string, toCompare?: string}>({current: handler.stream!, toCompare: userStorage.getItem(StoreKey.HistoryToCompare, () => undefined)});
     const [selectedMeta, setSelectedMeta] = useState<Set<string>>(new Set());
 
     handler.subscribeToSubQueryLoading();
@@ -76,7 +66,7 @@ export const TestHistoryView: React.FC<{ test: TestNameRef, handler: TestDataHan
             .catch((reason) => console.error(reason))
             .finally(() => {
                 setStreams({...streams, toCompare: targetStream});
-                userStorage.setItem('historyToCompare', targetStream);
+                userStorage.setItem(StoreKey.HistoryToCompare, targetStream);
             });
     }, [streams]);
 
@@ -203,10 +193,10 @@ const TestHistoryModal: React.FC<{ test: TestNameRef, handler: TestDataHandler,
 }
 
 const TestSummaryCard: React.FC<{ test: TestNameRef, sessions: [MetadataRef, TestSessionResult][], metaKeyMask: string[] | undefined, handler: TestDataHandler, onDetailsClick: SessionCallback, onHistoryClick: HistorySessionCallback }> = memo(({ test, sessions, metaKeyMask, handler, onDetailsClick, onHistoryClick }) => {
-    const [expanded, setExpanded] = useState(userStorage.getItem('summaryExpandedCards', () => new Map<string, boolean>()).get(test) ?? false);
+    const [expanded, setExpanded] = useState(userStorage.getItem(StoreKey.SummaryExpandedCards, () => new Map<string, boolean>()).get(test) ?? false);
 
     const expandCard = useCallback((bExpand) => {
-        userStorage.getItem('summaryExpandedCards').set(test, bExpand);
+        userStorage.getItem(StoreKey.SummaryExpandedCards).set(test, bExpand);
         setExpanded(bExpand);
     }, [test]);
  
@@ -405,16 +395,7 @@ const filterSessionsWithMetaMask = (testStatus: TestStatus, filteredTests: TestN
         }
     });
 
-    // Convoluted way to get unique meta identifier for each test with the same name
-    const testMetaKeys = new Map(testMeta.entries().filter(([, testRefs]) => testRefs.length > 1).map(([testName, testRefs]) => {
-        const testCommonKeys: string[][] = testRefs.map(([, metaRefs]) => MetadataRef.identifyCommonKeys(metaRefs));
-        const commonKeys: string[] = MetadataRef.identifyCommonKeys(testRefs.map(([, metaRefs]) => metaRefs).flat());
-        const uniqueKeys = testCommonKeys.reduce(
-            (current, testCommonKey) => current.filter((key) => testCommonKey.filter(k => !commonKeys.includes(k)).includes(key)),
-            testCommonKeys.shift() ?? []
-        );
-        return [testName, uniqueKeys];
-    }));
+    const testMetaKeys = getTestUniqueMetaIdentifiers(testMeta);
 
     failing.sort(([testA, ], [testB, ]) => testFailureCount.get(testB)! - testFailureCount.get(testA)!);
 
diff --git a/Engine/Source/Programs/Horde/HordeDashboard/plugins/test-automation/testAutomationUserStorage.ts b/Engine/Source/Programs/Horde/HordeDashboard/plugins/test-automation/testAutomationUserStorage.ts
--- a/Engine/Source/Programs/Horde/HordeDashboard/plugins/test-automation/testAutomationUserStorage.ts
+++ b/Engine/Source/Programs/Horde/HordeDashboard/plugins/test-automation/testAutomationUserStorage.ts
@@ -21,3 +21,9 @@ class TransientStorage {
 }
 
 export const userStorage = new TransientStorage();
+
+export enum StoreKey {
+    HistoryToCompare = "historyToCompare",
+    SummaryExpandedCards = "summaryExpandedCards"
+}
+
diff --git a/Engine/Source/Programs/Horde/HordeDashboard/plugins/test-automation/testAutomationView.tsx b/Engine/Source/Programs/Horde/HordeDashboard/plugins/test-automation/testAutomationView.tsx
--- a/Engine/Source/Programs/Horde/HordeDashboard/plugins/test-automation/testAutomationView.tsx
+++ b/Engine/Source/Programs/Horde/HordeDashboard/plugins/test-automation/testAutomationView.tsx
@@ -1,6 +1,6 @@
 // Copyright Epic Games, Inc. All Rights Reserved.
 
-import { Stack, FontIcon, IComboBoxOption, IContextualMenuItem, IContextualMenuProps, Text, Spinner, SpinnerSize, Label, DefaultButton, ITag, TagPicker, IPickerItemProps, PrimaryButton, Modal, IconButton, TooltipHost, DirectionalHint } from "@fluentui/react";
+import { Stack, FontIcon, IComboBoxOption, IContextualMenuItem, IContextualMenuProps, Text, Spinner, SpinnerSize, Label, DefaultButton, ITag, TagPicker, IPickerItemProps, PrimaryButton, Modal, IconButton, TooltipHost, DirectionalHint, PivotItem, Pivot } from "@fluentui/react";
 import { Breadcrumbs } from "horde/components/Breadcrumbs";
 import { TopNav } from "horde/components/TopNav";
 import { getHordeStyling } from "horde/styles/Styles";
@@ -16,6 +16,7 @@ import { JobStepTestDataItem } from "./api";
 
 import { TestDataV2 } from "./testDataV2Fetcher";
 import dashboard, { StatusColor } from "horde/backend/Dashboard";
+import { TestAutomationHealth } from "./testAutomationHealth";
 
 const handler = new TestDataHandler();
 
@@ -253,6 +254,8 @@ const TestAutomationSidebarLeft: React.FC = observer(() => {
     handler.subscribeToFilter();
     const filterMetaKeys = handler.filterState.metadata?.keys().toArray().sort((a, b) => a.localeCompare(b)) ?? [];
 
+    const view = (handler.getSearchParam("view") as string | undefined) ?? "0";
+
     const [metaKeys, setMetaKeys] = useState<string[]>(filterMetaKeys);
     const { hordeClasses } = getHordeStyling();
 
@@ -273,33 +276,42 @@ const TestAutomationSidebarLeft: React.FC = observer(() => {
         }
     }, [handler.searchUpdated]);
 
-   return <Stack style={{ width: 300, paddingRight: 18 }}>
-            <Stack className={hordeClasses.modal}>
-                <Stack key="chooser_stream">
-                    <StreamChooser />
-                </Stack>
-                <Stack key="chooser_test">
-                    <TestChooser />
-                </Stack>
-                <Stack key="chooser_tag">
-                    <TagChooser />
-                </Stack>
-                {!!metaKeys.length && metaKeys.map((key) => <Stack key={`chooser_meta_${key}`}><MetaValueChooser metakey={key} onRemove={updateMetaKey} /></Stack>)}
-                <Stack>
-                    <MetaKeyChooser key="chooser_meta" onUpdateKey={updateMetaKey} />
-                </Stack>
-            </Stack>
+   return <Stack style={{ width: 300, paddingRight: 18 }} className={hordeClasses.modal}>
+            <Pivot selectedKey={view} onLinkClick={(item) => {handler.setSearchParam("view", item?.props?.itemKey ?? "0")}}>
+                <PivotItem headerText="Summary" itemKey="0">
+                    <Stack key="chooser_stream">
+                        <StreamChooser />
+                    </Stack>
+                    <Stack key="chooser_test">
+                        <TestChooser />
+                    </Stack>
+                    <Stack key="chooser_tag">
+                        <TagChooser />
+                    </Stack>
+                    {!!metaKeys.length && metaKeys.map((key) => <Stack key={`chooser_meta_${key}`}><MetaValueChooser metakey={key} onRemove={updateMetaKey} /></Stack>)}
+                    <Stack>
+                        <MetaKeyChooser key="chooser_meta" onUpdateKey={updateMetaKey} />
+                    </Stack>
+                </PivotItem>
+                <PivotItem headerText="Health" itemKey="1">
+                    <Stack key="chooser_stream">
+                        <StreamChooser />
+                    </Stack>
+                    <Stack key="chooser_test">
+                        <TestChooser />
+                    </Stack>
+                    <Stack key="chooser_tag">
+                        <TagChooser />
+                    </Stack>
+                    {!!metaKeys.length && metaKeys.map((key) => <Stack key={`chooser_meta_${key}`}><MetaValueChooser metakey={key} onRemove={updateMetaKey} /></Stack>)}
+                    <Stack>
+                        <MetaKeyChooser key="chooser_meta" onUpdateKey={updateMetaKey} />
+                    </Stack>
+                </PivotItem>
+            </Pivot>
         </Stack>
 });
 
-const TestAutomationSummary: React.FC = () => {
-
-    return <Stack tokens={{ childrenGap: 24 }}>
-                <TestSummary handler={handler} />
-                {/** Additional or alternate view could be implemented here */}
-            </Stack>
-}
-
 const TestAutomationPanel: React.FC = observer(() => {
 
     if (handler.queryLoading) {
@@ -316,6 +328,7 @@ const TestAutomationPanel: React.FC = observer(() => {
     }
 
     const status = handler.filteredTests;
+    const view = handler.getSearchParam("view") as string | undefined ?? "0";
 
     return <Stack grow>
                 { !status.length &&
@@ -323,8 +336,11 @@ const TestAutomationPanel: React.FC = observer(() => {
                         <Text style={{ fontSize: 24 }}>No Results</Text>
                     </Stack>
                 }
-                { !!status.length &&
-                    <TestAutomationSummary />
+                { !!status.length && view === "0" &&
+                    <TestSummary handler={handler} />
+                }
+                { !!status.length && view === "1" &&
+                    <TestAutomationHealth handler={handler} />
                 }
             </Stack>
 
@@ -437,7 +453,7 @@ export const TestAutomationView: React.FC = () => {
     return <Stack className={hordeClasses.horde} key="key_test_automation_hub">
                 <TopNav />
                 <Breadcrumbs items={[{ text: 'Test Automation Hub' }]} />
-                <Stack grow style={{ padding: 12, backgroundColor: modeColors.background, height: 'calc(100vh - 148px)', overflow: 'auto' } } horizontalAlign='center'>
+                <Stack grow style={{ padding: 12, backgroundColor: modeColors.background } } horizontalAlign='center'>
                     {initialized && !!jobId && !!stepId &&
                         <JobStepTestsSelectionModal jobId={jobId} stepId={stepId} />
                     }
@@ -451,7 +467,7 @@ export const TestAutomationView: React.FC = () => {
                     {initialized && !testDataId && !!handler.availableStreams.length &&
                         <Stack horizontal grow>
                             <TestAutomationSidebarLeft/>
-                            <Stack style={{ minWidth: 1024 }}>
+                            <Stack style={{ width: 'calc(100vw - 350px)', height: 'calc(100vh - 162px)', overflow: 'auto' }} grow>
                                 <TestAutomationPanel/>
                             </Stack>
                         </Stack>
diff --git a/Engine/Source/Programs/Horde/HordeDashboard/plugins/test-automation/testData.ts b/Engine/Source/Programs/Horde/HordeDashboard/plugins/test-automation/testData.ts
--- a/Engine/Source/Programs/Horde/HordeDashboard/plugins/test-automation/testData.ts
+++ b/Engine/Source/Programs/Horde/HordeDashboard/plugins/test-automation/testData.ts
@@ -18,6 +18,9 @@ type TagId = string;
 type TestSessionId = string;
 type TestDataId = string;
 type PhaseKey = string;
+type TestRecipeId = string;
+type ErrorFingerprintId = string;
+type WorkflowId = string;
 
 // TestData Version registration
 export interface TestVersionFetcher {
@@ -39,6 +42,12 @@ export interface TestVersionFetcher {
     /// fetch test data information from job and step id
     fetchJobStepTestData(jobId: string, stepId: string): Promise<JobStepTestDataItem[] | undefined>;
 
+    /// fetch test recipes and attached them to their associated test
+    fetchTestRecipes(testRef: TestNameRef, recipeIds: TestRecipeId[]): Promise<void>;
+
+    /// fetch test audits associated with a list of tests
+    fetchTestAudit(testRefs: TestNameRef[], auditMap: Map<TestNameRef, TestAuditRecords>, weeks?: number): Promise<void>;
+
 }
 
 export class TestDataVersionRegistrar {
@@ -86,7 +95,7 @@ export class PhaseSessionResult {
     jobId?: string;
     stepId?: string;
     hasWarning?: boolean;
-    errorFingerprint?: string;
+    errorFingerprint?: ErrorFingerprintId;
     tagIds?: TagId[];
 }
 
@@ -130,6 +139,8 @@ export class TestSessionResult {
     stepId?: string;
     testDataId?: string;
     tagIds?: TagId[];
+    recipeId?: TestRecipeId;
+    errorFingerprints?: Map<ErrorFingerprintId, PhaseKey[]>;
 
     private _artifactFactory?: ArtifactFactory;
 }
@@ -434,6 +445,7 @@ export class TestSessionDetails {
     session?: TestSessionResult;
     tags?: string[];
     preflightCommitId?: string;
+    commandline?: string;
 
     private _start: Date;
     private _duration: number;
@@ -467,12 +479,26 @@ export class TestNameRef {
         return toMatch.some(tag => this.tagRefs?.values().some(r => r.isMatch(tag)));
     }
 
+    updateRecipes(recipes: TestRecipeRef[]) {
+        const newRecipes = new Set<TestRecipeRef>(recipes);
+        this.recipeRefs = this.recipeRefs? this.recipeRefs.union(newRecipes) : newRecipes;
+    }
+
+    getRecipe(recipeId: TestRecipeId) {
+        return this.recipeRefs?.keys().find(r => r.id === recipeId);
+    }
+
     id: TestId;
     key: string;
     name: string;
     version: number;
     tagRefs?: Set<TestTagRef>;
     phases?: Map<PhaseKey, TestPhaseRef>;
+    team?: string;
+    harness?: string;
+    profile?: string;
+    intent?: string;
+    recipeRefs?: Set<TestRecipeRef>;
 
     private lowerKey: string;
     private lowerName: string;
@@ -562,6 +588,104 @@ export class MetadataRef {
     }
 }
 
+export class TestRecipeRef {
+    constructor(id: string, streamId: string, jobTemplateId: string, jobStep: string, nodeClass: string, testVersion: string, timeoutMinutes: number, version: number) {
+        this.id = id;
+        this.streamId = streamId;
+        this.jobTemplateId = jobTemplateId;
+        this.jobStep = jobStep;
+        this.nodeClass = nodeClass;
+        this.testVersion = testVersion;
+        this.timeoutMinutes = timeoutMinutes;
+        this.version = version;
+    }
+
+    isMatch(recipeId: string) {
+        return this.id === recipeId;
+    }
+
+    includeParams(paramKeys: string[]) {
+        return paramKeys.some(key => this.params?.has(key));
+    }
+
+    isParamValueMatch(key: string, value: any) {
+        return !!this.params && this.params.get(key) === value;
+    }
+
+    id: TestRecipeId;
+    version: number;
+    streamId: StreamId;
+    jobTemplateId: string;
+    jobStep: string;
+    nodeClass: string;
+    testVersion: string;
+    timeoutMinutes: number;
+
+    test?: TestNameRef;
+    workflowId?: string;
+    params?: Map<string, any>
+}
+
+export class TestAudit {
+    constructor(testId: string, userId: string, userInputDate: Date, enableSchedule: boolean) {
+        this.testId = testId;
+        this.userId = userId;
+        this.userInputDate = userInputDate;
+        this.enableSchedule = enableSchedule;
+    }
+
+    clone(): TestAudit {
+        return {
+            ...this,
+            customers: this.customers?.slice(),
+            notifications: this.notifications? new Map(this.notifications.entries()) : undefined
+        };
+    }
+
+    testId: TestId;
+    userId: string;
+    userInputDate: Date;
+    enableSchedule: boolean;
+    lastAuditDate?: Date;
+    isUnderAudit?: boolean;
+    owner?: string;
+    customers?: string[];
+    notifications?: Map<StreamId, WorkflowId>;
+    team?: string;
+    harness?: string;
+    intent?: string;
+    profile?: string;
+    notes?: string;
+}
+
+export abstract class TestAuditRecords {
+    constructor(test: TestNameRef) {
+        this.history = [];
+        this.test = test;
+    }
+
+    getLatest() {
+        if (!this.history.length) {
+            return;
+        }
+
+        return this.history[0];
+    }
+
+    appendRecord(audit: TestAudit) {
+        this.history.push(audit);
+    }
+
+    sortRecords() {
+        this.history.sort((a, b) => b.userInputDate.getTime() - a.userInputDate.getTime());
+    }
+
+    abstract post(entry: TestAudit): Promise<boolean>;
+
+    test: TestNameRef;
+    history: TestAudit[];
+}
+
 export class TestMetaStatus {
     constructor(test: TestNameRef) {
         this.sessions = new Map();
@@ -1478,6 +1602,72 @@ export class TestDataHandler {
         }
     }
 
+    async queryRecipes(test: TestNameRef, recipeIds: TestRecipeId[]) {
+
+        const missingRecipes = recipeIds.filter(r => !test.getRecipe(r));
+        if (missingRecipes.length === 0) {
+            this.setRecipeQueryLoading(false);
+            return;
+        }
+
+        this.setRecipeQueryLoading(true);
+
+        for (const fetcherVersion of TestDataVersionRegistrar.manifest) {
+            /// fetch only from the matching version
+            if (fetcherVersion.version === test.version) {
+                console.debug(`Fetching from version ${fetcherVersion.version} test recipes`);
+                try {
+                    await fetcherVersion.fetchTestRecipes(test, missingRecipes);
+                } catch (reason) {
+                    console.debug(`Failed to fetch recipes from version ${fetcherVersion.version}: ${reason}`);
+                }
+                break;
+            }
+        }
+
+        if (this.loaded) {
+            this.setRecipeQueryLoading(false);
+        }
+    }
+
+    async queryAudits(tests: TestNameRef[], withHistory?: boolean) {
+
+        const missingTests = tests.filter(t => !this.testAudits.get(t) || (withHistory && this.testAudits.get(t)!.history.length === 1));
+
+        if (missingTests.length === 0) {
+            this.setAuditQueryLoading(false);
+            return;
+        }
+
+        let queryWeeks: number | undefined;
+        if (withHistory) {
+            const state = this.filterState;
+            queryWeeks = state.weeks && state.weeks > defaultQueryWeeks ? state.weeks : defaultQueryWeeks;
+        }
+
+        this.setAuditQueryLoading(true);
+
+        for (const fetcherVersion of TestDataVersionRegistrar.manifest) {
+            console.debug(`Fetching from version ${fetcherVersion.version} test audits`);
+            try {
+                await fetcherVersion.fetchTestAudit(missingTests, this.testAudits, queryWeeks);
+            } catch (reason) {
+                console.debug(`Failed to fetch audits from version ${fetcherVersion.version}: ${reason}`);
+            }
+        }
+
+        if (withHistory) {
+            missingTests.forEach(test => {
+                const records = this.testAudits.get(test);
+                records?.sortRecords();
+            });
+        }
+
+        if (this.loaded) {
+            this.setAuditQueryLoading(false);
+        }
+    }
+
     // load immutable data for this view session
     private async load() {
 
@@ -1541,6 +1731,16 @@ export class TestDataHandler {
         this.subPhaseQueryLoading = value;
     }
 
+    @action
+    setRecipeQueryLoading(value: boolean) {
+        this.queryRecipesLoading = value;
+    }
+
+    @action
+    setAuditQueryLoading(value: boolean) {
+        this.queryAuditsLoading = value;
+    }
+
     get stream(): StreamId | undefined {
         return this.filterState.stream;
     }
@@ -1588,6 +1788,12 @@ export class TestDataHandler {
     @observable
     subPhaseQueryLoading: boolean = false;
 
+    @observable
+    queryRecipesLoading: boolean = false;
+
+    @observable
+    queryAuditsLoading: boolean = false;
+
     subscribeToFilter() {
         if (this.filterUpdated) { }
     }
@@ -1616,6 +1822,14 @@ export class TestDataHandler {
         if (this.subPhaseQueryLoading) { }
     }
 
+    subscribeToRecipesLoading() {
+        if (this.queryRecipesLoading) { }
+    }
+
+    subscribeToAuditsLoading() {
+        if (this.queryAuditsLoading) { }
+    }
+
     private searchCallback?: (search: string, replace?: boolean) => void;
     private search: TestDataFilterSearchParams = new TestDataFilterSearchParams();
     filterState: TestDataState = {};
@@ -1626,6 +1840,8 @@ export class TestDataHandler {
 
     testTestData: Map<TestDataId, TestSessionDetails> = new Map();
 
+    testAudits: Map<TestNameRef, TestAuditRecords> = new Map();
+
     // quantized common commits to land at same time in the views
     commitIdDates = new Map<string, Date>();
     
diff --git a/Engine/Source/Programs/Horde/HordeDashboard/plugins/test-automation/testDataV2Fetcher.ts b/Engine/Source/Programs/Horde/HordeDashboard/plugins/test-automation/testDataV2Fetcher.ts
--- a/Engine/Source/Programs/Horde/HordeDashboard/plugins/test-automation/testDataV2Fetcher.ts
+++ b/Engine/Source/Programs/Horde/HordeDashboard/plugins/test-automation/testDataV2Fetcher.ts
@@ -1,7 +1,7 @@
 // Copyright Epic Games, Inc. All Rights Reserved.
 
-import { getJobStepTestDataV2, GetPhaseResponse, getStreamTestsV2, getTestDataV2, GetTestMetadataResponse, GetTestNameResponse, GetTestPhaseSessionResponse, getTestPhaseSessionsV2, getTestPhasesV2, GetTestSessionResponse, getTestSessionsV2, GetTestTagResponse, JobStepTestDataItem, PhaseSessionDetailsResponse, TestData, TestOutcome, TestSessionDetailsResponse } from "./api";
-import { TestStreamRef, TestStatus, TestSessionResult, TestNameRef, MetadataRef, TestTagRef, PhaseSessionResult, TestPhaseRef, TestMetaStatus, TestSessionDetails, TestPhaseStatus, DeviceRef, TestVersionFetcher } from "./testData"
+import { getJobStepTestDataV2, GetPhaseResponse, getStreamTestsV2, getTestAuditsV2, getTestDataV2, GetTestMetadataResponse, GetTestNameResponse, getTestNamesV2, GetTestPhaseSessionResponse, getTestPhaseSessionsV2, getTestPhasesV2, getTestRecipesV2, GetTestSessionResponse, getTestSessionsV2, GetTestTagResponse, JobStepTestDataItem, PhaseSessionDetailsResponse, postTestAuditV2, TestAuditPost, TestAuditResponse, TestData, TestOutcome, TestRecipeResponse, TestSessionDetailsResponse } from "./api";
+import { TestStreamRef, TestStatus, TestSessionResult, TestNameRef, MetadataRef, TestTagRef, PhaseSessionResult, TestPhaseRef, TestMetaStatus, TestSessionDetails, TestPhaseStatus, DeviceRef, TestVersionFetcher, TestRecipeRef, TestAudit, TestAuditRecords } from "./testData"
 import { EventEntry } from "./testEventsModel";
 
 const minuteInWeek = 10080;
@@ -23,6 +23,39 @@ class TestPhaseStatusV2 extends TestPhaseStatus {
 
 }
 
+class TestAuditRecordsV2 extends TestAuditRecords {
+
+    async post(entry: TestAudit): Promise<boolean> {
+
+        const post: TestAuditPost = {
+            enableSchedule: entry.enableSchedule,
+            lastAuditDate: entry.lastAuditDate?.toISOString(),
+            isUnderAudit: entry.isUnderAudit,
+            team: entry.team,
+            harness: entry.harness,
+            profile: entry.profile,
+            intent: entry.intent,
+            notes: entry.notes,
+            owner: entry.owner,
+            customers: entry.customers,
+            notifications: entry.notifications?.entries().reduce((obj, [stream, workflow]) => { obj[stream] = workflow; return obj }, {})
+        };
+
+        try {
+            const response = await postTestAuditV2(entry.testId, post);
+            const newAuditEntry = TestDataV2.createTestAudit(response);
+            this.history.unshift(newAuditEntry);
+
+        } catch(reason) {
+            console.debug(`Failed to post audit: ${reason}`);
+            return false;
+        }
+        
+        return true;
+    }
+
+}
+
 export class TestDataV2 implements TestVersionFetcher {
     get version() { return 2 }
 
@@ -138,15 +171,81 @@ export class TestDataV2 implements TestVersionFetcher {
 
     }
 
+    async fetchTestRecipes(testRef: TestNameRef, recipeIds: string[]) {
+
+        const responses = await getTestRecipesV2(recipeIds);
+        const refs: TestRecipeRef[] = [];
+        responses.forEach((recipe) => {
+            if (recipe.nameRef !== testRef.id) {
+                return;
+            }
+
+            const recipeRef = this.createRecipeRef(recipe);
+            recipeRef.test = testRef;
+            refs.push(recipeRef);
+        });
+
+        testRef.updateRecipes(refs);
+
+    }
+
+    async fetchTestAudit(testRefs: TestNameRef[], auditMap: Map<TestNameRef, TestAuditRecords>, weeks?: number) {
+
+        const tests: Map<string, TestNameRef> = new Map(testRefs.filter((t) => t.version >= this.version).map((t) => [t.id, t]));
+        if (tests.size) {
+
+            // create an audit record collection for each test
+            tests.forEach((test) => {
+                if (!auditMap.get(test)) {
+                    auditMap.set(test, new TestAuditRecordsV2(test));
+                }
+            });
+
+            const testsAudited = new Set<string>();
+            const minQueryDate = weeks? new Date(new Date().valueOf() - (minuteInWeek * weeks * 60000)) : undefined;
+            // send the query
+            const responses = await getTestAuditsV2(testRefs.map(t => t.id), minQueryDate?.toISOString());
+            responses.forEach((audit) => {
+                const test = tests.get(audit.testId);
+                if (!!test) {
+                    testsAudited.add(audit.testId);
+                    const auditEntry = TestDataV2.createTestAudit(audit);
+                    const auditRecords = auditMap.get(test)!;
+                    auditRecords.appendRecord(auditEntry);
+                }
+            });
+
+            const noAuditTests = (new Set<string>(tests.keys())).difference(testsAudited);
+            if (noAuditTests.size) {
+                // Make sure we have complete data from unaudit tests
+                const responses2 = await getTestNamesV2(Array.from(noAuditTests));
+                responses2.forEach((nameRef) => {
+                    const test = tests.get(nameRef.id);
+                    if (!!test) {
+                        test.profile = nameRef.profile;
+                        test.team = nameRef.team;
+                        test.harness = nameRef.harness;
+                        test.intent = nameRef.intent;
+                    }
+                });
+            }
+        }
+
+    }
+
     private testNameRefMap: Map<string, TestNameRef> = new Map();
     private metadataRefMap: Map<string, MetadataRef> = new Map();
     private testTagRefMap: Map<string, TestTagRef> = new Map();
+    private testRecipeRefMap: Map<string, TestRecipeRef> = new Map();
 
     private createTestNameRef(test: GetTestNameResponse): TestNameRef {
         const globalKey = test.id;
         let item = this.testNameRefMap.get(globalKey);
         if (!item) {
             item = new TestNameRef(test.id, test.key, test.name, this.version);
+            item.team = test.team;
+            item.harness = test.harness;
+            item.profile = test.profile;
             this.testNameRefMap.set(globalKey, item);
         }
 
@@ -179,6 +278,37 @@ export class TestDataV2 implements TestVersionFetcher {
         return new TestPhaseRef(item.id, item.key, item. name, this.version);
     }
 
+    private createRecipeRef(item: TestRecipeResponse): TestRecipeRef {
+        const globalKey = item.id;
+        let recipe = this.testRecipeRefMap.get(globalKey);
+        if (!recipe) {
+            recipe = new TestRecipeRef(item.id, item.streamId, item.jobTemplateId, item.jobStep, item.nodeClass, item.version, item.timeoutMinutes, this.version);
+            recipe.workflowId = item.workflowId;
+            recipe.params = recipe.params;
+            this.testRecipeRefMap.set(globalKey, recipe);
+        }
+
+        return recipe;
+    }
+
+    static createTestAudit(item: TestAuditResponse): TestAudit {
+
+        const auditDate = new Date(item.userInputDate);
+        const audit = new TestAudit(item.testId, item.userId, auditDate, item.enableSchedule);
+        audit.isUnderAudit = item.isUnderAudit;
+        audit.lastAuditDate = item.lastAuditDate? new Date(item.lastAuditDate) : undefined;
+        audit.owner = item.owner;
+        audit.customers = item.customers;
+        audit.notifications = item.notifications ? new Map(Object.keys(item.notifications).map(stream => [stream, item.notifications![stream]])): undefined;
+        audit.team = item.team;
+        audit.harness = item.harness;
+        audit.profile = item.profile;
+        audit.intent = item.intent;
+        audit.notes = item.notes;
+
+        return audit;
+    }
+
     private createTestSession(test: GetTestSessionResponse): TestSessionResult {
         const start = new Date(test.startDateTime);
         const result = new TestSessionResult(test.id, start, test.duration, test.commitId, test.commitOrder, test.nameRef, test.metadataId, test.streamId, test.outcome, test.testDataId);
@@ -188,6 +318,8 @@ export class TestDataV2 implements TestVersionFetcher {
         result.jobId = test.jobId;
         result.stepId = test.stepId;
         result.tagIds = test.tagIds;
+        result.recipeId = test.recipeId;
+        result.errorFingerprints = !test.errorFingerprints? undefined : new Map(test.errorFingerprints.map(i => [i.key, i.phases]));
 
         return result;
     }
@@ -223,6 +355,7 @@ export class TestDataV2 implements TestVersionFetcher {
                                     TestOutcome.Success : TestOutcome.Skipped;
         result.preflightCommitId = testData.preflightCommitId;
         result.tags = testDetails.tags;
+        result.commandline = testDetails.summary.commandline;
 
         // search for known metadata
         const metaFilter = new Map(Object.entries(testDetails.metadata).map(([key, value]) => [key, [value]]));
@@ -230,7 +363,16 @@ export class TestDataV2 implements TestVersionFetcher {
         result.meta = metaRef;
 
         // search for known test name
-        const testNameRef = this.testNameRefMap.values().find((nameRef) => nameRef.key === testData.key) ?? this.createTestNameRef({id: testData.id, key: testData.key, name: result.name});
+        const testNameRef = this.testNameRefMap.values().find((nameRef) => nameRef.key === testData.key) ?? this.createTestNameRef(
+            {
+                id: testData.id,
+                key: testData.key,
+                name: result.name,
+                team: testDetails.description?.team,
+                harness: testDetails.description?.harness,
+                profile: testDetails.description?.profile,
+                intent: testDetails.description?.intent
+            });
         result.test = testNameRef;
 
         return result;
diff --git a/Engine/Source/Programs/Horde/HordeDashboard/plugins/test-automation/testHealthData.ts b/Engine/Source/Programs/Horde/HordeDashboard/plugins/test-automation/testHealthData.ts
new file mode 100644
--- /dev/null
+++ b/Engine/Source/Programs/Horde/HordeDashboard/plugins/test-automation/testHealthData.ts
@@ -0,0 +1,282 @@
+// Copyright Epic Games, Inc. All Rights Reserved.
+
+import { action, makeObservable, observable } from "mobx";
+import { MetadataRef, TestAudit, TestDataHandler, TestNameRef, TestRecipeRef } from "./testData";
+
+export enum TestHealth {
+    Healthy = "Healthy",
+    Fair = "Fair",
+    Unstable = "Unstable",
+    Bad = "Bad",
+    Broken = "Broken",
+    NA = "N/A",
+}
+
+export const healthStars: Map<TestHealth, number> = new Map([
+    [TestHealth.Healthy, 5],
+    [TestHealth.Fair, 4],
+    [TestHealth.Unstable, 3],
+    [TestHealth.Bad, 2],
+    [TestHealth.Broken, 1],
+    [TestHealth.NA, 0],
+]);
+
+export const healthColors: Map<TestHealth, string> = new Map([
+    [TestHealth.Healthy, "#67bc54"],
+    [TestHealth.Fair, "#cae120ff"],
+    [TestHealth.Unstable, "#c89320"],
+    [TestHealth.Bad, "#ce7253ff"],
+    [TestHealth.Broken, "#ce5353"],
+    [TestHealth.NA, "#b5b5b5"],
+]);
+
+export enum TestCadence {
+    Nightly = "Nightly",
+    Incremental = "Incremental",
+    Daily = "Daily",
+    Limited = "Limited",
+    Weekly = "Weekly",
+    NotRun = "Not Run",
+    NA = "N/A"
+}
+
+export const cadenceOrders: Map<TestCadence, number> = new Map([
+    [TestCadence.Incremental, 0],
+    [TestCadence.Daily, 1],
+    [TestCadence.Weekly, 2],
+    [TestCadence.Limited, 3],
+    [TestCadence.NotRun, 4],
+    [TestCadence.NA, 5],
+]);
+
+export const cadenceColors: Map<TestCadence, string> = new Map([
+    [TestCadence.Incremental, "#c8b72081"],
+    [TestCadence.Daily, "#53cea18d"],
+    [TestCadence.Weekly, "#1f49a579"],
+    [TestCadence.Limited, "#4a3e6c8b"],
+    [TestCadence.NotRun, "#b5b5b516"],
+    [TestCadence.NA, "#6e6e6e16"],
+]);
+
+export const booleanColors: Map<boolean, string> = new Map([
+    [true, "#cbc82248"],
+    [false, "#67676748"],
+]);
+
+class LazyString {
+    constructor(getter: () => string | undefined) {
+        this.getter = getter;
+        this.cached = null;
+    }
+
+    get value() {
+        if (this.cached === null) {
+            this.cached = this.getter();
+            this.getter = () => undefined;
+        }
+        return this.cached;
+    }
+
+    private getter: () => string | undefined;
+    private cached: string | undefined | null;
+}
+
+export type TestHealthExplain = {
+    catastrophicFailureRate: number;
+    failureRate: number;
+    successRate: number;
+    redundantErrorRate: number;
+}
+
+export type TestHealthTiming = {
+    averageToInterruptSecs: number;
+    averageToCompleteRunSecs: number;
+    latestCompletedRun: Date | undefined;
+}
+
+export class TestHealthItem {
+    constructor(test: TestNameRef, handler: TestDataHandler, stream: string) {
+        this.nameRef = test;
+        this.handler = handler;
+        this.metaRefs = [];
+        this.health = TestHealth.NA;
+        this.cadence = TestCadence.NA
+        this.stream = stream;
+        this.lazyDisplayMeta = new LazyString(() => undefined);
+        makeObservable(this)
+    }
+
+    private handler: TestDataHandler;
+    private lazyDisplayMeta: LazyString;
+    private transientCopy?: TestAudit;
+
+    nameRef: TestNameRef;
+    metaRefs: MetadataRef[];
+    health: TestHealth;
+    cadence: TestCadence;
+    stream: string;
+
+    healthExplain?: TestHealthExplain;
+    timings?: TestHealthTiming;
+
+    @observable
+    updated: number = 0;
+    @action
+    setUpdated() {
+        this.updated += 1;
+    }
+    subscribeUpdate() {
+        if (this.updated) { }
+    }
+
+    get displayMeta() {
+        return this.lazyDisplayMeta.value;
+    }
+    setLazyDisplayMeta(getter: () => string | undefined) {
+        this.lazyDisplayMeta = new LazyString(getter);
+    }
+
+    private getOrStartTransaction(): TestAudit {
+        if (!this.transientCopy) {
+            const auditRecord = this.handler.testAudits.get(this.nameRef);
+            let latestClone = auditRecord?.getLatest()?.clone();
+            if (!latestClone) {
+                latestClone = new TestAudit(this.nameRef.id, "", new Date(Date.now()), true);
+            }
+            this.transientCopy = latestClone;
+        }
+        return this.transientCopy;
+    }
+
+    private get audit(): TestAudit | undefined {
+        return this.transientCopy ?? this.handler.testAudits.get(this.nameRef)?.getLatest();
+    }
+
+    get isTransient(): boolean {
+        return this.transientCopy !== undefined;
+    }
+
+    async commitTransaction() {
+        if (!this.isTransient) return;
+
+        await this.handler.testAudits.get(this.nameRef)?.post(this.transientCopy!);
+
+        this.transientCopy = undefined;
+        this.setUpdated();
+    }
+
+    discardTransaction() {
+        if (this.isTransient) {
+            this.transientCopy = undefined;
+            this.setUpdated();
+        }
+    }
+
+    get name() {
+        return this.nameRef.name;
+    }
+
+    get profile() {
+        return this.audit?.profile ?? this.nameRef.profile;
+    }
+    
+    set profile(value: string | undefined) {
+        this.getOrStartTransaction().profile = value;
+        this.setUpdated();
+    }
+
+    get team() {
+        return this.audit?.team ?? this.nameRef.team;
+    }
+
+    set team(value: string | undefined) {
+        this.getOrStartTransaction().team = value;
+        this.setUpdated();
+    }
+
+    get harness() {
+        return this.audit?.harness ?? this.nameRef.harness;
+    }
+
+    set harness(value: string | undefined) {
+        this.getOrStartTransaction().harness = value;
+        this.setUpdated();
+    }
+
+    get intent() {
+        return this.audit?.intent ?? this.nameRef.intent;
+    }
+
+    set intent(value: string | undefined) {
+        this.getOrStartTransaction().intent = value;
+        this.setUpdated();
+    }
+
+    get underAudit() {
+        return !!this.audit?.isUnderAudit;
+    }
+
+    set underAudit(value: boolean) {
+        this.getOrStartTransaction().isUnderAudit = value;
+        if (value) {
+            this.getOrStartTransaction().lastAuditDate = new Date(Date.now());
+        }
+        this.setUpdated();
+    }
+
+    get lastAuditDate() {
+        return this.audit?.lastAuditDate;
+    }
+
+    set lastAuditDate(value: Date | undefined) {
+        this.getOrStartTransaction().lastAuditDate = value;
+        this.setUpdated();
+    }
+
+    get notes() {
+        return this.audit?.notes;
+    }
+
+    set notes(value: string | undefined) {
+        this.getOrStartTransaction().notes = value;
+        this.setUpdated();
+    }
+
+    get owner() {
+        return this.audit?.owner;
+    }
+
+    get customers() {
+        return this.audit?.customers;
+    }
+
+    get notification() {
+        return this.audit?.notifications?.get(this.stream);
+    }
+
+    linkRecipeIds(ids: Iterable<string>) {
+        this.recipeIds = Array.from(ids);
+    }
+
+    async fetchRecipes() {
+        if (this.recipeIds && this.recipeIds.length) {
+            await this.handler.queryRecipes(this.nameRef, this.recipeIds);
+            this.cached_recipes = undefined;
+        }
+    }
+
+    private cached_recipes?: TestRecipeRef[];
+    private recipeIds?: string[];
+    get recipes() {
+        if (!this.recipeIds) return;
+        if (!this.cached_recipes) {
+            this.cached_recipes = this.nameRef.recipeRefs?.keys().filter(r => r.streamId === this.stream).toArray();
+        }
+        if (!this.cached_recipes || !this.cached_recipes.length) return;
+        return this.cached_recipes;
+    }
+
+    get fetchingRecipes() {
+        return this.handler.queryRecipesLoading;
+    }
+}
