[Horde][Dashboard] Updating dev environment dependencies due to transient dependecy postinstall changes.

diff --git a/Engine/Source/Programs/Horde/HordeDashboard/package.json b/Engine/Source/Programs/Horde/HordeDashboard/package.json
--- a/Engine/Source/Programs/Horde/HordeDashboard/package.json
+++ b/Engine/Source/Programs/Horde/HordeDashboard/package.json
@@ -48,10 +48,11 @@
     "eslint-plugin-react-hooks": "^4.6.0",
     "eslint-plugin-react-refresh": "^0.4.6",
     "license-checker-rseidelsohn": "4.3.0",
+    "patch-package": "^8.0.1",
     "rimraf": "^3.0.0",
     "typescript": "^5.4.5",
     "vite": "^5.2.0",
     "vite-tsconfig-paths": "^4.3.2",
     "vitest": "^1.0.0"
   }
-}
+}
\ No newline at end of file
