[Horde][Server] Including ODBC reference

diff --git a/Engine/Source/Programs/Horde/HordeServer/HordeServer.csproj b/Engine/Source/Programs/Horde/HordeServer/HordeServer.csproj
--- a/Engine/Source/Programs/Horde/HordeServer/HordeServer.csproj
+++ b/Engine/Source/Programs/Horde/HordeServer/HordeServer.csproj
@@ -1,4 +1,4 @@
-<Project Sdk="Microsoft.NET.Sdk.Web">
+﻿<Project Sdk="Microsoft.NET.Sdk.Web">
 
   <PropertyGroup>
     <TargetFramework>net8.0</TargetFramework>
@@ -80,6 +80,7 @@
     <PackageReference Include="Swashbuckle.AspNetCore" Version="6.5.0" />
     <PackageReference Include="System.Linq.Async" Version="6.0.3" />
     <PackageReference Include="System.Text.Json" Version="8.0.6" />
+    <PackageReference Include="System.Data.Odbc" Version="9.0.10" />
     <PackageReference Include="T-Digest.NET" Version="1.0.3" />
     <PackageReference Include="TimeZoneConverter" Version="6.1.0" />
     <PackageReference Include="ZstdSharp.Port" Version="0.8.6" />
