Updated RetrieveBlockAsync to no longer enable verbose logging on retries of extracting the directory node. We suspect that verbose logging was causing a large amount of overhead resulting in steps timing out and cancelling. We still enable hash verification of files though.

diff --git a/Engine/Source/Programs/Horde/Drivers/JobDriver/Utility/TempStorage.cs b/Engine/Source/Programs/Horde/Drivers/JobDriver/Utility/TempStorage.cs
--- a/Engine/Source/Programs/Horde/Drivers/JobDriver/Utility/TempStorage.cs
+++ b/Engine/Source/Programs/Horde/Drivers/JobDriver/Utility/TempStorage.cs
@@ -789,7 +789,7 @@ namespace JobDriver.Utility
 					// Add all the files and flush the ref
 					ExtractOptions options = new ExtractOptions();
 					options.Progress = new ExtractStatsLogger(logger);
-					options.VerifyOutput = options.VerboseOutput = (attempt > 0);
+					options.VerifyOutput = (attempt > 0);
 					options.MaxBatchSize = maxBatchSize;
 
 					await rootDirEntry.Handle.ExtractAsync(rootDir.ToDirectoryInfo(), options, logger, cancellationToken);
@@ -819,7 +819,7 @@ namespace JobDriver.Utility
 					catch (FileNotFoundException ex) when (attempt == 0)
 					{
 						logger.LogWarning(ex, "File not found: {Message}", ex.Message);
-						logger.LogInformation("Running extract again with verbose logging.");
+						logger.LogInformation("Running extract again with VerifyOutput enabled to check hashes.");
 					}
 				}
 			}
