Horde: Add options for controlling sync threads and batches in PerforceMaterializer

Also add support for ModTime as a flag.

diff --git a/Engine/Source/Programs/Horde/Drivers/JobDriver/Execution/PerforceMaterializer.cs b/Engine/Source/Programs/Horde/Drivers/JobDriver/Execution/PerforceMaterializer.cs
--- a/Engine/Source/Programs/Horde/Drivers/JobDriver/Execution/PerforceMaterializer.cs
+++ b/Engine/Source/Programs/Horde/Drivers/JobDriver/Execution/PerforceMaterializer.cs
@@ -1,7 +1,9 @@
 ﻿// Copyright Epic Games, Inc. All Rights Reserved.
 
+using System.Collections.Specialized;
 using System.Diagnostics;
 using System.Text.Json;
+using System.Web;
 using EpicGames.Core;
 using EpicGames.Perforce;
 using HordeCommon.Rpc.Messages;
@@ -22,6 +24,7 @@ namespace JobDriver.Execution;
 /// <param name="SyncThreads">Number of threads to use when syncing a Perforce workspace</param>
 /// <param name="SyncBatch">How many files in each sync batch</param>
 /// <param name="SyncBatchSize">Size of each sync batch</param>
+/// <param name="UseModTime">Whether to use ModTime flag on the workspace</param>
 public record PerforceMaterializerOptions(
 	string DirPath,
 	RpcAgentWorkspace AgentWorkspace,
@@ -29,7 +32,55 @@ public record PerforceMaterializerOptions(
 	IReadOnlyDictionary<string, string>? EnvVars = null,
 	int SyncThreads = -1,
 	int SyncBatch = -1,
-	int SyncBatchSize = -1);
+	int SyncBatchSize = -1,
+	bool UseModTime = true)
+{
+	/// <summary>
+	/// Parse options from <see cref="RpcAgentWorkspace.Method" /> as HTTP query string
+	/// </summary>
+	/// <param name="dirPath">Base directory for this workspace</param>
+	/// <param name="agentWorkspace">Workspace options</param>
+	public static PerforceMaterializerOptions FromMethodString(string dirPath, RpcAgentWorkspace agentWorkspace)
+	{
+		const string NameKey = "name";
+		const string PerforceMaterializerValue = "perforce";
+		const string SyncThreadsKey = "syncThreads";
+		const string SyncBatchKey = "syncBatch";
+		const string SyncBatchSizeKey = "syncBatchSize";
+		const string UseModTimeKey = "useModTime";
+		
+		PerforceMaterializerOptions options = new (dirPath, agentWorkspace);
+		if (!String.IsNullOrEmpty(agentWorkspace.Method))
+		{
+			NameValueCollection nameValues = HttpUtility.ParseQueryString(agentWorkspace.Method);
+			if (String.Equals(nameValues[NameKey], PerforceMaterializerValue, StringComparison.OrdinalIgnoreCase))
+			{
+				if (Int32.TryParse(nameValues[SyncThreadsKey], out int threads))
+				{
+					options = options with { SyncThreads = threads };
+				}
+				if (Int32.TryParse(nameValues[SyncBatchKey], out int batch))
+				{
+					options = options with { SyncBatch = batch };
+				}
+				if (Int32.TryParse(nameValues[SyncBatchSizeKey], out int batchSize))
+				{
+					options = options with { SyncBatchSize = batchSize };
+				}
+				if (String.Equals(nameValues[UseModTimeKey], "true", StringComparison.OrdinalIgnoreCase))
+				{
+					options = options with { UseModTime = true };
+				}
+				if (String.Equals(nameValues[UseModTimeKey], "false", StringComparison.OrdinalIgnoreCase))
+				{
+					options = options with { UseModTime = false };
+				}
+			}
+		}
+
+		return options;
+	}
+}
 
 /// <summary>
 /// Materializer using standard Perforce client for syncing files
@@ -436,7 +487,7 @@ public sealed class PerforceMaterializer : IWorkspaceMaterializer
 			Host = _options.HostName ?? Environment.MachineName,
 			Type = "partitioned",
 			Root = _workspaceDir.FullName,
-			Options = ClientOptions.Clobber,
+			Options = _options.UseModTime ? ClientOptions.Clobber | ClientOptions.ModTime : ClientOptions.Clobber,
 		};
 		newClient.View.AddRange(view);
 		Directory.CreateDirectory(newClient.Root);
@@ -579,7 +630,7 @@ class PerforceMaterializerFactory(IServiceProvider serviceProvider) : IWorkspace
 			}
 			
 			ILogger<PerforceMaterializer> logger = serviceProvider.GetRequiredService<ILogger<PerforceMaterializer>>();
-			PerforceMaterializerOptions options = new (workspaceDir.FullName, workspaceInfo);
+			PerforceMaterializerOptions options = PerforceMaterializerOptions.FromMethodString(workspaceDir.FullName, workspaceInfo);
 			return new PerforceMaterializer(options, tracer, logger);
 		}
 		return null;
