Fix typo

#rnx

diff --git a/Engine/Source/Programs/Horde/HordeDashboard/plugins/test-automation/testAutomationHealth.tsx b/Engine/Source/Programs/Horde/HordeDashboard/plugins/test-automation/testAutomationHealth.tsx
--- a/Engine/Source/Programs/Horde/HordeDashboard/plugins/test-automation/testAutomationHealth.tsx
+++ b/Engine/Source/Programs/Horde/HordeDashboard/plugins/test-automation/testAutomationHealth.tsx
@@ -691,7 +691,7 @@ const TestHealthDetailsModal: React.FC<{ test: TestHealthItem, handler: TestData
                                 {!isCommitting &&
                                     <Stack tokens={{childrenGap: 6}}>
                                         <Text block>You have uncommitted changes to this audit.</Text>
-                                        <Text block>Do want to commit them?</Text>
+                                        <Text block>Do you want to commit them?</Text>
                                         <FocusZone handleTabKey={FocusZoneTabbableElements.all} isCircularNavigation>
                                             <Stack horizontal horizontalAlign="center">
                                             <PrimaryButton 
