Fixed up docker-compose.yml path in Horde docs.

#rnx

diff --git a/Engine/Source/Programs/Horde/Docs/Deployment/Server.md b/Engine/Source/Programs/Horde/Docs/Deployment/Server.md
--- a/Engine/Source/Programs/Horde/Docs/Deployment/Server.md
+++ b/Engine/Source/Programs/Horde/Docs/Deployment/Server.md
@@ -43,7 +43,7 @@ by providing a preconfigured group of Docker containers, which includes instance
 
 Similar to the MSI installer, this method is suitable for testing Horde or deploying it in small-scale environments.
 To access the prebuilt images, refer to the Docker section above.
-The necessary Docker Compose configuration can be found within the `Engine/Source/Programs/Horde.Server/docker-compose.yml` file.
+The necessary Docker Compose configuration can be found within the `Engine/Source/Programs/Horde/HordeServer/docker-compose.yml` file.
 
 From the same directory, start the containers with:
 ```bash
