Fix broken scrolling handle in Phase view

#rnx

diff --git a/Engine/Source/Programs/Horde/HordeDashboard/plugins/test-automation/testAutomationView.tsx b/Engine/Source/Programs/Horde/HordeDashboard/plugins/test-automation/testAutomationView.tsx
--- a/Engine/Source/Programs/Horde/HordeDashboard/plugins/test-automation/testAutomationView.tsx
+++ b/Engine/Source/Programs/Horde/HordeDashboard/plugins/test-automation/testAutomationView.tsx
@@ -453,7 +453,7 @@ export const TestAutomationView: React.FC = () => {
     return <Stack className={hordeClasses.horde} key="key_test_automation_hub">
                 <TopNav />
                 <Breadcrumbs items={[{ text: 'Test Automation Hub' }]} />
-                <Stack grow style={{ padding: 12, backgroundColor: modeColors.background } } horizontalAlign='center'>
+                <Stack grow style={{ padding: 12, backgroundColor: modeColors.background, height: 'calc(100vh - 136px)', overflow: 'auto'} } horizontalAlign='center'>
                     {initialized && !!jobId && !!stepId &&
                         <JobStepTestsSelectionModal jobId={jobId} stepId={stepId} />
                     }
