[Horde][BuildHealth][Dashboard] Skipped steps as a summarized metric & manual invocation metadata for CSV.

diff --git a/Engine/Source/Programs/Horde/HordeDashboard/src/components/buildhealth/stepoutcome/StepOutcomeD3Component.tsx b/Engine/Source/Programs/Horde/HordeDashboard/src/components/buildhealth/stepoutcome/StepOutcomeD3Component.tsx
--- a/Engine/Source/Programs/Horde/HordeDashboard/src/components/buildhealth/stepoutcome/StepOutcomeD3Component.tsx
+++ b/Engine/Source/Programs/Horde/HordeDashboard/src/components/buildhealth/stepoutcome/StepOutcomeD3Component.tsx
@@ -5,7 +5,7 @@ import * as d3 from "d3";
 import { getHordeStyling, preloadFonts } from "horde/styles/Styles";
 import { useEffect, useRef } from "react";
 import { StepOutcomeDataHandler } from "./StepOutcomeDataHandler";
-import { AbortedStr, ChangeHeader, ColumnHeader, encodeStepNameFromStepNameHeader, isChangeHeader, isDateHeader, isLabelOutcome, isStepOutcome, isSummary, isSummaryHeader, StepNameHeader, StepOutcomeTable, TableEntry } from "./StepOutcomeDataTypes";
+import { AbortedStr, ChangeHeader, ColumnHeader, encodeStepNameFromStepNameHeader, isChangeHeader, isDateHeader, isLabelOutcome, isStepOutcome, isSummary, isSummaryHeader, StepNameHeader, StepOutcomeTable, SummaryTableEntry, TableEntry } from "./StepOutcomeDataTypes";
 import { getCellAccents, getCellClasses, getCellStyle, getColorRecords, IncrementalRefreshIndicatorPanel, NOT_APPLICABLE_LABEL, resetCellClasses, sharedBorderStyle } from "./StepOutcomeSharedUIComponents";
 import { getHordeTheme } from "horde/styles/theme";
 import { observer } from "mobx-react-lite";
@@ -141,6 +141,96 @@ function GenerateD3Table(stepOutcomeTable: StepOutcomeTable, handler: StepOutcom
             .style("top", `${rect.bottom - tooltipHeight}px`);
     };
 
+    let generateSummaryTooltipRow = (tooltipContent: HTMLDivElement, labelString: string, valueString: string, explanationString?: string) => {
+        const ratio = document.createElement("div");
+        ratio.style.display = "flex";
+        ratio.style.justifyContent = "space-between";
+        ratio.style.alignItems = "flex-start";
+        ratio.style.paddingBottom = "2px";
+
+        const leftDiv = document.createElement("div");
+        leftDiv.style.display = "flex";
+        leftDiv.style.flexDirection = "column";
+
+        const label = document.createElement("strong");
+        label.textContent = labelString;
+        label.style.paddingLeft = "5px";
+        leftDiv.appendChild(label);
+
+        if (explanationString !== undefined) {
+            const desc = document.createElement("span");
+            desc.style.fontSize = "smaller";
+            desc.textContent = explanationString;
+            desc.style.paddingLeft = "10px";
+            leftDiv.appendChild(desc);
+        }
+
+        const rightDiv = document.createElement("span");
+        rightDiv.textContent = valueString;
+
+        // Append both sides to the main container
+        ratio.appendChild(leftDiv);
+        ratio.appendChild(rightDiv);
+
+        tooltipContent.appendChild(ratio);
+    }
+
+    let generateSummaryTooltipHeader = (tooltipContent: HTMLDivElement, header: string) => {
+        const sectionHeaderDiv = document.createElement("div");
+        sectionHeaderDiv.style.display = "flex";
+        sectionHeaderDiv.style.paddingTop = "10px";
+        sectionHeaderDiv.style.justifyContent = "space-between";
+        sectionHeaderDiv.innerHTML = `<strong>${header}`;
+        sectionHeaderDiv.className = tooltipClasses.toolTipSection;
+        tooltipContent.appendChild(sectionHeaderDiv);
+    }
+
+    let computePassRatio = (entry: SummaryTableEntry) => {
+        let ratio: number = -1;
+
+        if (entry.steps > 0 && entry.stepsIncluded > 0) {
+            let numerator: number = entry.stepsPass;
+
+            if (!stepOutcomeTable.viewModelOptions.warningsAsSummaryFailure) {
+                numerator += entry.stepsWarning;
+            }
+
+            ratio = numerator / entry.stepsIncluded;
+        }
+
+        let ratioStr: string = ratio === -1 ? NOT_APPLICABLE_LABEL : `${(ratio * 100).toFixed(1)}%`;
+        return { ratio, ratioStr };
+    }
+
+    let computeSkipRatio = (entry: SummaryTableEntry) => {
+        let skippedRatio: number = -1;
+
+        if (entry.steps > 0 && entry.stepsIncluded > 0) {
+            let numerator: number = entry.stepsSkipped;
+            skippedRatio = numerator / entry.stepsIncluded;
+        }
+
+        let skippedRatioStr: string = skippedRatio === -1 ? NOT_APPLICABLE_LABEL : `${(skippedRatio * 100).toFixed(1)}%`;
+        return { skippedRatio, skippedRatioStr };
+    }
+
+    let computeCompletedPassRatio = (entry: SummaryTableEntry) => {
+        let truePassRatio: number = -1;
+
+        if (entry.steps > 0 && entry.stepsIncluded > 0) {
+            let numerator: number = entry.stepsPass;
+
+            if (!stepOutcomeTable.viewModelOptions.warningsAsSummaryFailure) {
+                numerator += entry.stepsWarning;
+            }
+
+            truePassRatio = numerator / (entry.stepsIncluded - entry.stepsSkipped);
+        }
+
+        let truePassRatioStr: string = truePassRatio === -1 ? NOT_APPLICABLE_LABEL : `${(truePassRatio * 100).toFixed(1)}%`;
+        return { truePassRatio, truePassRatioStr };
+    }
+
     // #endregion -- Tooltip Helpers --
 
     const containerRef = useRef<HTMLDivElement>(null);
@@ -595,19 +685,9 @@ function GenerateD3Table(stepOutcomeTable: StepOutcomeTable, handler: StepOutcom
                                 })
                         getCellClasses(entry, handler.getJobResponseData(entry.jobId)).forEach(cls => div.classed(cls, true));
                     } else if (isSummary(entry)) { // SummaryTableEntry
-                        let ratio: number = -1;
-
-                        if (entry.steps > 0 && entry.stepsCompleted > 0) {
-                            let numerator: number = entry.stepsPass
-
-                            if (!stepOutcomeTable.viewModelOptions.warningsAsSummaryFailure) {
-                                numerator += entry.stepsWarning;
-                            }
-
-                            ratio = numerator / entry.stepsCompleted;
-                        }
-
-                        let ratioStr: string = ratio === -1 ? NOT_APPLICABLE_LABEL : `${(ratio * 100).toFixed(1)}%`;
+                        let { ratio, ratioStr }: { ratio: number; ratioStr: string; } = computePassRatio(entry);
+                        let { skippedRatio, skippedRatioStr }: { skippedRatio: number; skippedRatioStr: string; } = computeSkipRatio(entry);
+                        let { truePassRatio, truePassRatioStr }: { truePassRatio: number; truePassRatioStr: string; } = computeCompletedPassRatio(entry);
 
                         div.style("cursor", "default")
                             .style("background-color", summaryColorScale(ratio))
@@ -623,66 +703,25 @@ function GenerateD3Table(stepOutcomeTable: StepOutcomeTable, handler: StepOutcom
                                 const tooltipContent = document.createElement("div");
                                 tooltipContent.style.paddingLeft = "4px";
 
-                                // Show ratio in proper format
-                                const ratio = document.createElement("div");
-                                ratio.style.display = "flex";
-                                ratio.style.justifyContent = "space-between";
-                                ratio.style.alignItems = "flex-start";
-                                ratio.className = tooltipClasses.toolTipSection;
-
-                                const leftDiv = document.createElement("div");
-                                leftDiv.style.display = "flex";
-                                leftDiv.style.flexDirection = "column";
-
-                                const label = document.createElement("strong");
-                                label.textContent = "Success Ratio:";
-                                leftDiv.appendChild(label);
-
-                                const desc = document.createElement("span");
-                                desc.style.fontSize = "smaller";
-                                desc.textContent = "(Success / Completed Steps)";
-                                leftDiv.appendChild(desc);
-
-                                // Right side: ratio value
-                                const rightDiv = document.createElement("span");
-                                rightDiv.textContent = ratioStr;
-
-                                // Append both sides to the main container
-                                ratio.appendChild(leftDiv);
-                                ratio.appendChild(rightDiv);
-
-                                tooltipContent.appendChild(ratio);
+                                // Show ratios
+                                generateSummaryTooltipHeader(tooltipContent, "Success & Failure Ratios");
+                                generateSummaryTooltipRow(tooltipContent, "Overall Success Ratio:", ratioStr, "(Success / Included Steps)");
+                                generateSummaryTooltipRow(tooltipContent, "Completed Success Ratio:", truePassRatioStr, "(Success / (Included Steps - Skipped Steps))");
+                                generateSummaryTooltipRow(tooltipContent, "Upstream Failure Ratio:", skippedRatioStr, "(Skipped / Included Steps)");
 
                                 // Show totals
-                                const runningSteps = document.createElement("div");
-                                runningSteps.style.display = "flex";
-                                runningSteps.style.justifyContent = "space-between";
-                                runningSteps.innerHTML = `<strong>Pending Steps:</strong><span>${entry.stepsPending ?? "-"}</span>`;
-                                tooltipContent.appendChild(runningSteps);
-
-                                const succeededSteps = document.createElement("div");
-                                succeededSteps.style.display = "flex";
-                                succeededSteps.style.justifyContent = "space-between";
-                                succeededSteps.innerHTML = `<strong>Succeeded Steps:</strong><span>${entry.stepsPass ?? "-"}</span>`;
-                                tooltipContent.appendChild(succeededSteps);
-
-                                const warningSteps = document.createElement("div");
-                                warningSteps.style.display = "flex";
-                                warningSteps.style.justifyContent = "space-between";
-                                warningSteps.innerHTML = `<strong>Warning Steps:</strong><span>${entry.stepsWarning ?? "-"}</span>`;
-                                tooltipContent.appendChild(warningSteps);
-
-                                const failedSteps = document.createElement("div");
-                                failedSteps.style.display = "flex";
-                                failedSteps.style.justifyContent = "space-between";
-                                failedSteps.innerHTML = `<strong>Failed Steps:</strong><span>${entry.stepsFail ?? "-"}</span>`;
-                                tooltipContent.appendChild(failedSteps);
-
-                                const completedStepsDiv = document.createElement("div");
-                                completedStepsDiv.style.display = "flex";
-                                completedStepsDiv.style.justifyContent = "space-between";
-                                completedStepsDiv.innerHTML = `<strong>Completed Steps:</strong><span>${entry.stepsCompleted ?? "-"}</span>`;
-                                tooltipContent.appendChild(completedStepsDiv);
+                                // Included steps
+                                generateSummaryTooltipHeader(tooltipContent, "Included Steps");
+                                generateSummaryTooltipRow(tooltipContent, "Succeeded Steps", entry.stepsPass?.toString());
+                                generateSummaryTooltipRow(tooltipContent, "Warning Steps", entry.stepsWarning?.toString());
+                                generateSummaryTooltipRow(tooltipContent, "Skipped Steps", entry.stepsSkipped?.toString());
+                                generateSummaryTooltipRow(tooltipContent, "Failed Steps", (entry.stepsFail).toString());
+                                generateSummaryTooltipRow(tooltipContent, "Failed Completed Steps", (entry.stepsFail - entry.stepsSkipped).toString(), "(Failed Steps - Skipped Steps)");
+                                generateSummaryTooltipRow(tooltipContent, "Included Steps (Total)", entry.stepsIncluded?.toString());
+
+                                // Excluded steps
+                                generateSummaryTooltipHeader(tooltipContent, "Excluded Steps");
+                                generateSummaryTooltipRow(tooltipContent, "Pending Steps", entry.stepsPending?.toString());
 
                                 attachTargettedToolTip(this as HTMLElement, scaffold.current!.tooltip, tooltipContent);
                             })
diff --git a/Engine/Source/Programs/Horde/HordeDashboard/src/components/buildhealth/stepoutcome/StepOutcomeDataHandler.ts b/Engine/Source/Programs/Horde/HordeDashboard/src/components/buildhealth/stepoutcome/StepOutcomeDataHandler.ts
--- a/Engine/Source/Programs/Horde/HordeDashboard/src/components/buildhealth/stepoutcome/StepOutcomeDataHandler.ts
+++ b/Engine/Source/Programs/Horde/HordeDashboard/src/components/buildhealth/stepoutcome/StepOutcomeDataHandler.ts
@@ -109,7 +109,7 @@ export interface StepOutcomeFilters {
  * @todo Data sharing & caching would ideally occur between the BuildHealth View & Step Outcome.
  */
 export class StepOutcomeDataHandler extends PollBase {
-    private readonly STREAM_JOBS_FILTER = "id,streamId,name,change,templateId,state,createTime,updateTime,abortedByUserInfo,preflightCommitId,preflightDescription,cancellationReason,labels,defaultLabel";
+    private readonly STREAM_JOBS_FILTER = "id,streamId,name,change,templateId,state,createTime,updateTime,startedByUserInfo,abortedByUserInfo,preflightCommitId,preflightDescription,cancellationReason,labels,defaultLabel";
     private readonly STREAM_JOBS_BATCH_STEP_FILTER = "batches.error,batches.steps.id,batches.steps.name,batches.steps.state,batches.steps.outcome,batches.steps.startTime,batches.steps.finishTime,batches.steps.warning,batches.steps.error"
     private readonly STREAM_CHANGES_FILTER = "id,dateUtc";
     static readonly MAX_JOB_COUNT_DEFAULT = 400;
@@ -372,7 +372,8 @@ export class StepOutcomeDataHandler extends PollBase {
                 job.id,
                 {
                     ...job,
-                    summarizedOutcome: summarizeJob(job)
+                    summarizedOutcome: summarizeJob(job),
+                    isManualJobInvocation : job.startedByUserInfo !== null
                 }
             ]
         }
diff --git a/Engine/Source/Programs/Horde/HordeDashboard/src/components/buildhealth/stepoutcome/StepOutcomeDataTypes.ts b/Engine/Source/Programs/Horde/HordeDashboard/src/components/buildhealth/stepoutcome/StepOutcomeDataTypes.ts
--- a/Engine/Source/Programs/Horde/HordeDashboard/src/components/buildhealth/stepoutcome/StepOutcomeDataTypes.ts
+++ b/Engine/Source/Programs/Horde/HordeDashboard/src/components/buildhealth/stepoutcome/StepOutcomeDataTypes.ts
@@ -65,6 +65,7 @@ export type UpdatedChangeSummaryData = ChangeSummaryData & {
 */
 export type JobResponseData = GetJobResponse & {
     summarizedOutcome: JobStepOutcome;
+    isManualJobInvocation : boolean;
 }
 
 // #region -- Table Entry Types --
@@ -106,11 +107,12 @@ export function isSummary(entry: TableEntry): entry is SummaryTableEntry {
  */
 export class SummaryTableEntry {
     steps: number = 0;
-    stepsCompleted: number = 0;
+    stepsIncluded: number = 0;
     stepsPass: number = 0;
     stepsFail: number = 0;
     stepsWarning: number = 0;
     stepsPending: number = 0;
+    stepsSkipped: number = 0;
 }
 
 /**
@@ -336,14 +338,14 @@ export class LabelOutcomeTableEntry extends OutcomeTableEntry {
         let baseUrl = window.location.origin;
         let jobUrl = `${baseUrl}/job/${this.jobId}`;
         let labelUrl = `${baseUrl}/job/${this.jobId}?label=${this.labelIdx}`;
-        return `${this.change},${this.streamId},${this.jobName},${this.jobId},${this.labelResponse.dashboardCategory},${this.labelResponse.dashboardName},${this.id},${this.state},${this.outcome},${this.jobCreateTime},${relatedJobContext?.summarizedOutcome},${relatedJobContext?.state},${jobUrl},${labelUrl}\n`;
+        return `${this.change},${this.streamId},${this.jobName},${this.jobId},${this.labelResponse.dashboardCategory},${this.labelResponse.dashboardName},${this.id},${this.state},${this.outcome},${this.jobCreateTime},${relatedJobContext?.summarizedOutcome},${relatedJobContext?.state},${relatedJobContext?.isManualJobInvocation},${jobUrl},${labelUrl}\n`;
     }
 
     /**
      * @inheritdoc
      */
     getCSVHeader(): string {
-        return "Change,StreamId,JobName,JobId,LabelCategory,LabelName,LabelIdx,StepState,StepOutcome,JobStart,JobSummaryOutcome,JobState,JobUrl,LabelUrl\n";
+        return "Change,StreamId,JobName,JobId,LabelCategory,LabelName,LabelIdx,StepState,StepOutcome,JobStart,JobSummaryOutcome,JobState,JobManualInvocation,JobUrl,LabelUrl\n";
     }
 
     // #endregion -- Abstract API --
@@ -438,7 +440,7 @@ export class StepOutcomeTableEntry extends OutcomeTableEntry {
     // #region -- Private Members --
 
     private stepResponse: IJobStepAdapter;
-    private stepDurationMs : number;
+    private stepDurationMs: number;
 
     // #endregion -- Private Members --
 
@@ -507,14 +509,14 @@ export class StepOutcomeTableEntry extends OutcomeTableEntry {
         let baseUrl = window.location.origin;
         let jobUrl = `${baseUrl}/job/${this.jobId}`;
         let stepUrl = `${baseUrl}/job/${this.jobId}?step=${this.stepResponse.stepId}`;
-        return `${this.change},${this.streamId},${this.jobName},${this.jobId},${this.stepResponse.stepName!},${this.stepResponse.stepId},${this.stepResponse.getBase().state},${this.stepResponse.getBase().outcome},${this.stepDurationMs / 1000},${this.jobCreateTime},${this.stepResponse.getBase().startTime},${relatedJobContext?.summarizedOutcome},${relatedJobContext?.state},${jobUrl},${stepUrl}\n`;
+        return `${this.change},${this.streamId},${this.jobName},${this.jobId},${this.stepResponse.stepName!},${this.stepResponse.stepId},${this.stepResponse.getBase().state},${this.stepResponse.getBase().outcome},${this.stepDurationMs / 1000},${this.jobCreateTime},${this.stepResponse.getBase().startTime},${relatedJobContext?.summarizedOutcome},${relatedJobContext?.state},${relatedJobContext?.isManualJobInvocation},${jobUrl},${stepUrl}\n`;
     }
 
     /**
      * @inheritdoc
      */
     getCSVHeader(): string {
-        return "Change,StreamId,JobName,JobId,StepName,StepId,StepState,StepOutcome,Duration(Sec),JobStart,StepStart,JobSummaryOutcome,JobState,JobUrl,StepUrl\n";
+        return "Change,StreamId,JobName,JobId,StepName,StepId,StepState,StepOutcome,Duration(Sec),JobStart,StepStart,JobSummaryOutcome,JobState,JobManualInvocation,JobUrl,StepUrl\n";
     }
 
     // #endregion -- Abstract API --
@@ -854,7 +856,12 @@ export class StepOutcomeTable {
                     continue;
                 }
 
-                summaryTableEntry.stepsCompleted++;
+                summaryTableEntry.stepsIncluded++;
+
+                if (cell.state == SkippedStr) {
+                    summaryTableEntry.stepsSkipped++;
+                }
+
 
                 if (cell.outcome == WarningStr) {
                     summaryTableEntry.stepsWarning++;
