Fix Horde server build by adding and removing paths to reflect the migration of log matchers to their own project.

diff --git a/Engine/Source/Programs/Horde/HordeServer/Dockerfile b/Engine/Source/Programs/Horde/HordeServer/Dockerfile
--- a/Engine/Source/Programs/Horde/HordeServer/Dockerfile
+++ b/Engine/Source/Programs/Horde/HordeServer/Dockerfile
@@ -32,6 +32,7 @@ RUN wget https://repo.mongodb.org/apt/debian/dists/bookworm/mongodb-org/7.0/main
 COPY Source/Programs/Shared/EpicGames.Core/*.csproj ./Source/Programs/Shared/EpicGames.Core/
 COPY Source/Programs/Shared/EpicGames.Horde/*.csproj ./Source/Programs/Shared/EpicGames.Horde/
 COPY Source/Programs/Shared/EpicGames.IoHash/*.csproj ./Source/Programs/Shared/EpicGames.IoHash/
+COPY Source/Programs/Shared/EpicGames.LogMatchers/*.csproj ./Source/Programs/Shared/EpicGames.LogMatchers/
 COPY Source/Programs/Shared/EpicGames.Perforce/*.csproj ./Source/Programs/Shared/EpicGames.Perforce/
 COPY Source/Programs/Shared/EpicGames.Perforce.Managed/*.csproj ./Source/Programs/Shared/EpicGames.Perforce.Managed/
 COPY Source/Programs/Shared/EpicGames.Redis/*.csproj ./Source/Programs/Shared/EpicGames.Redis/
@@ -51,8 +52,6 @@ COPY Source/Programs/Horde/.editorconfig Source/Programs/Horde/.editorconfig
 COPY Source/Programs/Horde/Plugins/ Source/Programs/Horde/Plugins/
 COPY Source/Programs/Horde/HordeServer/ Source/Programs/Horde/HordeServer/
 COPY Source/Programs/Horde/HordeServer.Shared/ Source/Programs/Horde/HordeServer.Shared/
-COPY Source/Programs/AutomationTool/AutomationUtils/Matchers/ Source/Programs/AutomationTool/AutomationUtils/Matchers/
-COPY Source/Programs/UnrealBuildTool/Matchers/ Source/Programs/UnrealBuildTool/Matchers/
 ARG msbuild_args
 RUN dotnet publish Source/Programs/Horde/HordeServer/HordeServer.csproj -c Release -o out -p:WithHordeStorage=false $msbuild_args
 
