Add tags column to test health

#rnx

diff --git a/Engine/Source/Programs/Horde/HordeDashboard/plugins/test-automation/testAutomationHealth.tsx b/Engine/Source/Programs/Horde/HordeDashboard/plugins/test-automation/testAutomationHealth.tsx
--- a/Engine/Source/Programs/Horde/HordeDashboard/plugins/test-automation/testAutomationHealth.tsx
+++ b/Engine/Source/Programs/Horde/HordeDashboard/plugins/test-automation/testAutomationHealth.tsx
@@ -1,11 +1,11 @@
 // Copyright Epic Games, Inc. All Rights Reserved.
 
-import { DetailsList, FontIcon, IColumn, IconButton, Label, Modal, SelectionMode, Stack, Text, HoverCard, HoverCardType, IPlainCardProps, DirectionalHint, DetailsListLayoutMode, Spinner, SpinnerSize, Toggle, DatePicker, ComboBox, IComboBoxOption, FocusZone, FocusZoneTabbableElements, FocusTrapCallout, PrimaryButton, DefaultButton } from "@fluentui/react";
+import { DetailsList, FontIcon, IColumn, IconButton, Label, Modal, SelectionMode, Stack, Text, HoverCard, HoverCardType, IPlainCardProps, DirectionalHint, DetailsListLayoutMode, Spinner, SpinnerSize, Toggle, DatePicker, ComboBox, IComboBoxOption, FocusZone, FocusZoneTabbableElements, FocusTrapCallout, PrimaryButton, DefaultButton, TooltipHost, TooltipOverflowMode } from "@fluentui/react";
 import { MetadataRef, TestDataHandler, TestNameRef, TestSessionResult } from "./testData";
 import { useCallback, useEffect, useMemo, useState } from "react";
 import { observer } from "mobx-react-lite";
 import { TestOutcome } from "./api";
-import { DomainNameColors, EditableTextField, getTestUniqueMetaIdentifiers, StreamSelector } from "./testAutomationCommon";
+import { DomainNameColors, EditableTextField, getTestUniqueMetaIdentifiers, StreamSelector, tagColors } from "./testAutomationCommon";
 import { TestHistoryGraphWidget } from "./testAutomationSummary";
 import { StoreKey, userStorage } from "./testAutomationUserStorage";
 import { projectStore } from "horde/backend/ProjectStore";
@@ -85,7 +85,7 @@ const getHealthColumns = (items: TestHealthItem[], columns: TestHealthColumn[],
             columns: newColumns
         })
     }
-    
+
     const newColumns: TestHealthColumn[] = [
         {
             key: 'name',
@@ -172,10 +172,35 @@ const getHealthColumns = (items: TestHealthItem[], columns: TestHealthColumn[],
             onColumnClick: (_, column) => onClickColumn(column),
             compare: (a: TestHealthItem, b: TestHealthItem) => a.notification?.localeCompare(b.notification ?? "") ?? -1
         },
+        {
+            key: 'tags',
+            name: 'Tags',
+            minWidth: 32,
+            maxWidth: 300,
+            isResizable: true,
+            isCollapsible: true,
+            isSorted: false,
+            onRender: (item: TestHealthItem) => {
+                const content = <Stack horizontal verticalAlign="center" tokens={{childrenGap: 4}} wrap style={{maxWidth: 350}}>
+                                    {item.nameRef?.tagRefs &&
+                                        item.nameRef.tagRefs.keys().map(t => t.name).toArray().sort().map(tag =>
+                                            <Stack style={{backgroundColor: tagColors.getTagColor(tag), borderRadius: 5, padding: '1px 4px', opacity: 0.8}}>
+                                                <Text>#{tag}</Text>
+                                            </Stack>
+                                        )
+                                    }
+                                </Stack>;
+
+                return <Stack style={{cursor: 'default', paddingTop: 2, height: 24, maskImage: 'linear-gradient(to right, white 70%, transparent)'}}>
+                            <TooltipHost content={content} overflowMode={TooltipOverflowMode.Self}>{content}</TooltipHost>
+                        </Stack>
+            },
+            onColumnClick: (_, column) => onClickColumn(column),
+        },
         {
             key: 'cadence',
             name: 'Cadence',
-            minWidth: 100,
+            minWidth: 75,
             maxWidth: 100,
             isResizable: true,
             isCollapsible: true,
@@ -192,7 +217,7 @@ const getHealthColumns = (items: TestHealthItem[], columns: TestHealthColumn[],
         {
             key: 'lastCompletedDate',
             name: 'Last Uninterrupted Run',
-            minWidth: 32,
+            minWidth: 75,
             maxWidth: 100,
             isResizable: true,
             isCollapsible: true,
@@ -480,6 +505,8 @@ export const TestHealthView: React.FC<{ test: TestHealthItem, handler: TestDataH
     const teamOptions: IComboBoxOption[] = teamColors.getDomain().map(p => ({key: p, text: p, styles: {optionText: {backgroundColor: teamColors.get(p), padding: 3, borderRadius: 4, minWidth: 72}}}));
     const harnessOptions: IComboBoxOption[] = harnessColors.getDomain().map(p => ({key: p, text: p}));
 
+    const tags = useMemo(() => test.nameRef?.tagRefs?.keys().map(t => t.name).toArray().sort(), [test]);
+
     return <Stack tokens={{childrenGap: 16}} style={{position: 'relative'}}>
                 <Stack tokens={{childrenGap: 8}} style={{padding: 14, borderRadius: 4}} >
                     <Stack horizontal tokens={{childrenGap: 12}} verticalAlign="center">
@@ -525,6 +552,17 @@ export const TestHealthView: React.FC<{ test: TestHealthItem, handler: TestDataH
                         />
                     </Stack>
 
+                    {tags &&
+                        <Stack horizontal verticalAlign="center" tokens={{childrenGap: 4}} wrap grow style={{paddingLeft: 12, paddingTop: 6}}>
+                            {tags.map(tag =>
+                                    <Stack style={{backgroundColor: tagColors.getTagColor(tag), borderRadius: 5, padding: '1px 4px', opacity: 0.8}}>
+                                        <Text>#{tag}</Text>
+                                    </Stack>
+                                )
+                            }
+                        </Stack>
+                    }
+
                     <Stack horizontal style={{paddingLeft: 12}}>
                         <Stack horizontal verticalAlign="center" tokens={{childrenGap: 6}} style={{cursor: 'default'}} title="edit disabled for now">
                             <Label>Worflow Notification: </Label>
