[Horde][BuildHealth][Server] Add lastSuccessTime and firstFailureTime to FindSuspectsForSpanAsync and use them for time delta gate.
Fetching these times via GetAsync was causing exceptions to be raised when commits were out of stream and thus exiting suspect finding altogether.

diff --git a/Engine/Source/Programs/Horde/Plugins/Build/HordeServer.Build/Issues/IssueService.cs b/Engine/Source/Programs/Horde/Plugins/Build/HordeServer.Build/Issues/IssueService.cs
--- a/Engine/Source/Programs/Horde/Plugins/Build/HordeServer.Build/Issues/IssueService.cs
+++ b/Engine/Source/Programs/Horde/Plugins/Build/HordeServer.Build/Issues/IssueService.cs
@@ -1038,7 +1038,7 @@ namespace HordeServer.Issues
 				if (prevJob != null)
 				{
 					spanData.LastSuccess = new NewIssueStepData(prevJob);
-					spanData.Suspects = await FindSuspectsForSpanAsync(streamConfig, spanData.Fingerprint, spanData.LastSuccess.CommitId, spanData.FirstFailure.CommitId, groupAuditLogger, cancellationToken);
+					spanData.Suspects = await FindSuspectsForSpanAsync(streamConfig, spanData.Fingerprint, spanData.LastSuccess.CommitId, spanData.LastSuccess.StepTime, spanData.FirstFailure.CommitId, spanData.FirstFailure.StepTime, groupAuditLogger, cancellationToken);
 				}
 
 				IJobStepRef? nextJob = await _jobStepRefs.GetNextStepForNodeAsync(job.StreamId, job.TemplateId, node.Name, job.CommitId, JobStepOutcome.Success, true);
@@ -1373,11 +1373,13 @@ namespace HordeServer.Issues
 		/// <param name="streamConfig">The stream name</param>
 		/// <param name="fingerprint">The fingerprint for the span</param>
 		/// <param name="lastSuccess">The last successful commit</param>
+		/// <param name="lastSuccessTime">The time of the last successful commit</param>
 		/// <param name="firstFailure">The first failing commit</param>
+		/// <param name="firstFailureTime">The time of the first failing commit</param>
 		/// <param name="issueLogger">The audit logger for the issue that the span is associated with</param>
 		/// <param name="cancellationToken">Cancellation token for the operation</param>
 		/// <returns>List of suspects</returns>
-		async Task<List<NewIssueSpanSuspectData>> FindSuspectsForSpanAsync(StreamConfig streamConfig, IIssueFingerprint fingerprint, CommitId lastSuccess, CommitId firstFailure, ILogger? issueLogger, CancellationToken cancellationToken)
+		async Task<List<NewIssueSpanSuspectData>> FindSuspectsForSpanAsync(StreamConfig streamConfig, IIssueFingerprint fingerprint, CommitId lastSuccess, DateTime lastSuccessTime, CommitId firstFailure, DateTime firstFailureTime, ILogger? issueLogger, CancellationToken cancellationToken)
 		{
 			using TelemetrySpan span = _tracer.StartActiveSpan($"{nameof(IssueService)}.{nameof(FindSuspectsForSpanAsync)}");
 			
@@ -1387,21 +1389,19 @@ namespace HordeServer.Issues
 			{
 				if (!String.IsNullOrEmpty(fingerprint.ChangeFilter))
 				{
-					ICommit? firstFailureCommit = await _commitService.GetCollection(streamConfig).GetAsync(firstFailure, cancellationToken);
-					ICommit? lastSuccessCommit = await _commitService.GetCollection(streamConfig).GetAsync(lastSuccess, cancellationToken);
-
-					if (firstFailureCommit == null || lastSuccessCommit == null)
+					if (lastSuccessTime == default || firstFailureTime == default)
 					{
-						LoggerUtils.MultiLog([(_logger, LogLevel.Error), (issueLogger, LogLevel.Information)], "Failed to get commits for {StreamName}: {LastSuccess} and {FirstFailure}", streamConfig.Name, lastSuccess, firstFailure);
-						return suspects;
+						LoggerUtils.MultiLog([(_logger, LogLevel.Warning), (issueLogger, LogLevel.Debug)], "Skipping IssueFindSuspectsMaxSuccessDays check, due to default time value for commit {DefaultTimeCommit}.", lastSuccessTime == default ? lastSuccess.Name : firstFailure.Name);
 					}
-
-					TimeSpan delta = firstFailureCommit.DateUtc - lastSuccessCommit.DateUtc;
-					
-					if (Math.Abs(delta.TotalDays) >= _buildConfig.CurrentValue.IssueFindSuspectsMaxSuccessDays)
+					else
 					{
-						LoggerUtils.MultiLog([(_logger, LogLevel.Warning), (issueLogger, LogLevel.Debug)], "Ignoring commits for {StreamName}: {LastSuccess} and {FirstFailure}, delta is {DeltaDays}", streamConfig.Name, lastSuccess, firstFailure.Name, delta.TotalDays);
-						return suspects;
+						TimeSpan delta = firstFailureTime - lastSuccessTime;
+
+						if (Math.Abs(delta.TotalDays) >= _buildConfig.CurrentValue.IssueFindSuspectsMaxSuccessDays)
+						{
+							LoggerUtils.MultiLog([(_logger, LogLevel.Warning), (issueLogger, LogLevel.Debug)], "Ignoring commits for {StreamName}: {LastSuccess} and {FirstFailure}, delta is {DeltaDays}", streamConfig.Name, lastSuccess.Name, firstFailure.Name, delta.TotalDays);
+							return suspects;
+						}
 					}
 
 					LoggerUtils.MultiLog([(_logger, LogLevel.Debug), (issueLogger, LogLevel.Debug)], "Querying for changes in {StreamName} between {LastSuccess} and {FirstFailure}", streamConfig.Name, lastSuccess, firstFailure);
@@ -1767,7 +1767,7 @@ namespace HordeServer.Issues
 				{
 					NewIssueStepData newLastSuccess = new NewIssueStepData(job.CommitId, IssueSeverity.Unspecified, job.Name, job.Id, batch.Id, step.Id, step.StartTimeUtc ?? default, step.LogId, null, false);
 					ILogger issueLogger = _issueCollection.GetLogger(span.IssueId);
-					List<NewIssueSpanSuspectData> newSuspects = await FindSuspectsForSpanAsync(streamConfig, span.Fingerprint, job.CommitId, span.FirstFailure.CommitId, issueLogger, cancellationToken);
+					List<NewIssueSpanSuspectData> newSuspects = await FindSuspectsForSpanAsync(streamConfig, span.Fingerprint, newLastSuccess.CommitId, newLastSuccess.StepTime, span.FirstFailure.CommitId, span.FirstFailure.StepTime, issueLogger, cancellationToken);
 
 					if (await _issueCollection.TryUpdateSpanAsync(span, newLastSuccess: newLastSuccess, newSuspects: newSuspects, cancellationToken: cancellationToken) == null)
 					{
