Horde: Fix issue with case-sensitive materializer name comparison for conforms

diff --git a/Engine/Source/Programs/Horde/Drivers/JobDriver/Execution/ConformExecutor.cs b/Engine/Source/Programs/Horde/Drivers/JobDriver/Execution/ConformExecutor.cs
--- a/Engine/Source/Programs/Horde/Drivers/JobDriver/Execution/ConformExecutor.cs
+++ b/Engine/Source/Programs/Horde/Drivers/JobDriver/Execution/ConformExecutor.cs
@@ -64,8 +64,13 @@ namespace JobDriver.Execution
 				bool isPerforceExecutor = _driverSettings.Executor.Equals(PerforceExecutor.Name, StringComparison.OrdinalIgnoreCase);
 				bool isWorkspaceExecutor = _driverSettings.Executor.Equals(WorkspaceExecutor.Name, StringComparison.OrdinalIgnoreCase);
 
-				List<RpcAgentWorkspace> managedWorkspaces = pendingWorkspaces.Where(x => PerforceExecutor.GetMaterializerName(x.Method) != PerforceExecutor.Name).ToList();
-				List<RpcAgentWorkspace> perforceWorkspaces = pendingWorkspaces.Where(x => PerforceExecutor.GetMaterializerName(x.Method) == PerforceExecutor.Name).ToList();
+				List<RpcAgentWorkspace> managedWorkspaces = pendingWorkspaces
+					.Where(x => !String.Equals(PerforceExecutor.GetMaterializerName(x.Method), PerforceExecutor.Name, StringComparison.OrdinalIgnoreCase))
+					.ToList();
+
+				List<RpcAgentWorkspace> perforceWorkspaces = pendingWorkspaces
+					.Where(x => String.Equals(PerforceExecutor.GetMaterializerName(x.Method), PerforceExecutor.Name, StringComparison.OrdinalIgnoreCase))
+					.ToList();
 				
 				// When using WorkspaceExecutor, only job options can override exact materializer to use
 				// It will default to ManagedWorkspaceMaterializer, which is compatible with the conform call below
diff --git a/Engine/Source/Programs/Horde/Drivers/JobDriver/Execution/PerforceExecutor.cs b/Engine/Source/Programs/Horde/Drivers/JobDriver/Execution/PerforceExecutor.cs
--- a/Engine/Source/Programs/Horde/Drivers/JobDriver/Execution/PerforceExecutor.cs
+++ b/Engine/Source/Programs/Horde/Drivers/JobDriver/Execution/PerforceExecutor.cs
@@ -254,7 +254,7 @@ namespace JobDriver.Execution
 			span.SetAttribute("horde.perforce.remove_untracked", removeUntrackedFiles);
 
 			// Print out all the workspaces we're going to sync
-			logger.LogInformation("Workspaces:");
+			logger.LogInformation("Workspaces to conform using ManagedWorkspace:");
 			foreach (RpcAgentWorkspace pendingWorkspace in pendingWorkspaces)
 			{
 				logger.LogInformation("  Identifier={Identifier}, Stream={StreamName}, Incremental={Incremental} Method={Method} Partitioned={Partitioned}",
@@ -462,7 +462,7 @@ namespace JobDriver.Execution
 			List<IWorkspaceMaterializer> materializers = [];
 			
 			// Set up all the materializers
-			logger.LogInformation("Workspaces:");
+			logger.LogInformation("Workspaces to conform as materializers:");
 			foreach (RpcAgentWorkspace workspace in workspaces)
 			{
 				logger.LogInformation("  Identifier={Identifier}, Stream={StreamName}, Incremental={Incremental} Method={Method} Partitioned={Partitioned}",
