Horde: Expose lease outcome reasons in REST API

diff --git a/Engine/Source/Programs/Horde/Plugins/Compute/HordeServer.Compute/Agents/AgentsController.cs b/Engine/Source/Programs/Horde/Plugins/Compute/HordeServer.Compute/Agents/AgentsController.cs
--- a/Engine/Source/Programs/Horde/Plugins/Compute/HordeServer.Compute/Agents/AgentsController.cs
+++ b/Engine/Source/Programs/Horde/Plugins/Compute/HordeServer.Compute/Agents/AgentsController.cs
@@ -176,7 +176,24 @@ namespace HordeServer.Agents
 
 		internal static GetAgentLeaseResponse CreateGetAgentLeaseResponse(ILease lease, Dictionary<string, string>? details, double? agentRate)
 		{
-			return new GetAgentLeaseResponse(lease.Id, lease.ParentId, lease.AgentId, agentRate, lease.Name, lease.LogId, lease.StartTime, lease.FinishTime, lease.FinishTime == null, details, lease.Outcome, null);
+			string[] outcomeReasons = ToFlagNames(lease.OutcomeReason);
+			return new GetAgentLeaseResponse(
+				lease.Id, lease.ParentId, lease.AgentId, agentRate, lease.Name, lease.LogId, lease.StartTime,
+				lease.FinishTime, lease.FinishTime == null, details, lease.Outcome, outcomeReasons, (long)lease.OutcomeReason, null);
+		}
+		
+		private static string[] ToFlagNames(LeaseOutcomeReason reason)
+		{
+			if (reason == LeaseOutcomeReason.None)
+			{
+				return [];
+			}
+    
+			return Enum.GetValues(typeof(LeaseOutcomeReason))
+				.Cast<LeaseOutcomeReason>()
+				.Where(flag => flag != LeaseOutcomeReason.None && reason.HasFlag(flag))
+				.Select(flag => flag.ToString())
+				.ToArray();
 		}
 		
 		private static List<string> FilterProperties(List<string> properties, string? filter)
diff --git a/Engine/Source/Programs/Shared/EpicGames.Horde/Agents/AgentMessages.cs b/Engine/Source/Programs/Shared/EpicGames.Horde/Agents/AgentMessages.cs
--- a/Engine/Source/Programs/Shared/EpicGames.Horde/Agents/AgentMessages.cs
+++ b/Engine/Source/Programs/Shared/EpicGames.Horde/Agents/AgentMessages.cs
@@ -72,8 +72,10 @@ namespace EpicGames.Horde.Agents
 	/// <param name="Executing"> Whether this lease has started executing on the agent yet </param>
 	/// <param name="Details"> Details of the payload being executed </param>
 	/// <param name="Outcome"> Outcome of the lease </param>
+	/// <param name="OutcomeReasons">List of outcome reasons, in string form</param>
+	/// <param name="OutcomeReasonsFlags">Outcome reasons as a uint64 (bit flags)</param>
 	/// <param name="State"> State of the lease (for AgentLeases) </param>
-	public record GetAgentLeaseResponse(LeaseId Id, LeaseId? ParentId, AgentId? AgentId, double? AgentRate, string Name, LogId? LogId, DateTime StartTime, DateTime? FinishTime, bool Executing, Dictionary<string, string>? Details, LeaseOutcome? Outcome, LeaseState? State);
+	public record GetAgentLeaseResponse(LeaseId Id, LeaseId? ParentId, AgentId? AgentId, double? AgentRate, string Name, LogId? LogId, DateTime StartTime, DateTime? FinishTime, bool Executing, Dictionary<string, string>? Details, LeaseOutcome? Outcome, string[] OutcomeReasons, long? OutcomeReasonsFlags, LeaseState? State);
 
 	/// <summary>
 	/// Information about a workspace synced on an agent
