[Horde][Analytics] IssueTelemetry event name fix

#Synced-CL 47915125

diff --git a/Engine/Source/Programs/Horde/Plugins/Build/HordeServer.Build/Issues/IssueTelemetry.cs b/Engine/Source/Programs/Horde/Plugins/Build/HordeServer.Build/Issues/IssueTelemetry.cs
--- a/Engine/Source/Programs/Horde/Plugins/Build/HordeServer.Build/Issues/IssueTelemetry.cs
+++ b/Engine/Source/Programs/Horde/Plugins/Build/HordeServer.Build/Issues/IssueTelemetry.cs
@@ -19,7 +19,7 @@ namespace HordeServer.Issues
 		/// <summary>
 		/// Default event name for the IssueTelemetry.
 		/// </summary>
-		public static readonly string DefaultEventName = "State.Issue.Dev";
+		public static readonly string DefaultEventName = "State.Issue";
 
 		/// <summary>
 		/// Default context for the Horde Issue Telemetry.
