Horde: Root Cause Summary Escalation Notifications to Slack

#Synced-CL Unknown

diff --git a/Engine/Source/Programs/Horde/Plugins/Build/HordeServer.Build/Issues/IIssue.cs b/Engine/Source/Programs/Horde/Plugins/Build/HordeServer.Build/Issues/IIssue.cs
--- a/Engine/Source/Programs/Horde/Plugins/Build/HordeServer.Build/Issues/IIssue.cs
+++ b/Engine/Source/Programs/Horde/Plugins/Build/HordeServer.Build/Issues/IIssue.cs
@@ -113,6 +113,11 @@ namespace HordeServer.Issues
 		/// </summary>
 		public CommitId? RootCommitId { get; }
 
+		/// <summary>
+		/// User to provide the root cause analysis for this issue
+		/// </summary>
+		public UserId? RootCauseOwnerId { get; }
+
 		/// <summary>
 		/// Root cause category for this issue
 		/// </summary>
diff --git a/Engine/Source/Programs/Horde/Plugins/Build/HordeServer.Build/Issues/IIssueCollection.cs b/Engine/Source/Programs/Horde/Plugins/Build/HordeServer.Build/Issues/IIssueCollection.cs
--- a/Engine/Source/Programs/Horde/Plugins/Build/HordeServer.Build/Issues/IIssueCollection.cs
+++ b/Engine/Source/Programs/Horde/Plugins/Build/HordeServer.Build/Issues/IIssueCollection.cs
@@ -238,6 +238,7 @@ namespace HordeServer.Issues
 	/// <param name="FixCommitId">Fix changelist for the issue. Pass a default value to clear the fix changelist.</param>
 	/// <param name="FixSystemic">Whether the issue should be marked fixed as a systemic issue</param>
 	/// <param name="RootCommitId">Root changelist for the issue. Pass an empty value to clear the root changelist.</param>
+	/// <param name="RootCauseOwnerId">New owner of the root cause</param>
 	/// <param name="RootCauseCategory">Root cause category</param>
 	/// <param name="RootCauseSummary">Root cause summary.</param>
 	/// <param name="DuplicateIssueId">Duplicate issue id</param>
@@ -262,6 +263,7 @@ namespace HordeServer.Issues
 		CommitId? FixCommitId = null,
 		bool? FixSystemic = null,
 		CommitId? RootCommitId = null,
+		UserId? RootCauseOwnerId = null,
 		string? RootCauseCategory = null,
 		string? RootCauseSummary = null,
 		int? DuplicateIssueId = null,
diff --git a/Engine/Source/Programs/Horde/Plugins/Build/HordeServer.Build/Issues/IIssueDetails.cs b/Engine/Source/Programs/Horde/Plugins/Build/HordeServer.Build/Issues/IIssueDetails.cs
--- a/Engine/Source/Programs/Horde/Plugins/Build/HordeServer.Build/Issues/IIssueDetails.cs
+++ b/Engine/Source/Programs/Horde/Plugins/Build/HordeServer.Build/Issues/IIssueDetails.cs
@@ -30,6 +30,11 @@ namespace HordeServer.Issues
 		/// </summary>
 		IUser? ResolvedBy { get; }
 
+		/// <summary>
+		/// The user that provided the root cause analysis
+		/// </summary>
+		IUser? RootCauseOwner { get; }
+
 		/// <summary>
 		/// List of spans for the issue
 		/// </summary>
diff --git a/Engine/Source/Programs/Horde/Plugins/Build/HordeServer.Build/Issues/IssueCollection.cs b/Engine/Source/Programs/Horde/Plugins/Build/HordeServer.Build/Issues/IssueCollection.cs
--- a/Engine/Source/Programs/Horde/Plugins/Build/HordeServer.Build/Issues/IssueCollection.cs
+++ b/Engine/Source/Programs/Horde/Plugins/Build/HordeServer.Build/Issues/IssueCollection.cs
@@ -125,6 +125,9 @@ namespace HordeServer.Issues
 			[BsonElement("RootCommit"), BsonIgnoreIfNull]
 			public string? RootCommitName { get; set; }
 
+			[BsonIgnoreIfNull]
+			public UserId? RootCauseOwnerId { get; set; }
+
 			[BsonIgnoreIfNull]
 			public string? RootCauseCategory { get; set; }
 
@@ -1315,6 +1318,7 @@ namespace HordeServer.Issues
 			if (options.ResolvedById != null)
 			{
 				options = options with { OwnerId = options.OwnerId ?? options.ResolvedById, Acknowledged = options.Acknowledged ?? true };
+				options = options with { RootCauseOwnerId = options.RootCauseOwnerId ?? options.OwnerId ?? options.ResolvedById };
 			}
 			if (options.OwnerId != null)
 			{
@@ -1391,6 +1395,18 @@ namespace HordeServer.Issues
 				}
 			}
 
+			if (options.RootCauseOwnerId != null)
+			{
+				if (options.RootCauseOwnerId.Value == UserId.Empty)
+				{
+					updates.Add(Builders<IssueDocument>.Update.Unset(x => x.RootCauseOwnerId));
+				}
+				else
+				{
+					updates.Add(Builders<IssueDocument>.Update.Set(x => x.RootCauseOwnerId, options.RootCauseOwnerId.Value));
+				}
+			}
+
 			if (options.RootCauseCategory != null)
 			{
 				if (options.RootCauseCategory.Length == 0)
diff --git a/Engine/Source/Programs/Horde/Plugins/Build/HordeServer.Build/Issues/IssueService.cs b/Engine/Source/Programs/Horde/Plugins/Build/HordeServer.Build/Issues/IssueService.cs
--- a/Engine/Source/Programs/Horde/Plugins/Build/HordeServer.Build/Issues/IssueService.cs
+++ b/Engine/Source/Programs/Horde/Plugins/Build/HordeServer.Build/Issues/IssueService.cs
@@ -44,6 +44,9 @@ namespace HordeServer.Issues
 		/// <inheritdoc/>
 		public IUser? ResolvedBy { get; }
 
+		/// <inheritdoc/>
+		public IUser? RootCauseOwner { get; }
+
 		/// <inheritdoc/>
 		public IUser? QuarantinedBy { get; set; }
 
@@ -83,13 +86,14 @@ namespace HordeServer.Issues
 		/// <param name="resolvedBy"></param>
 		/// <param name="quarantinedBy"></param>
 		/// <param name="forceClosedBy"></param>
+		/// <param name="rootCauseOwner"></param>
 		/// <param name="spans">Spans for this issue</param>
 		/// <param name="steps">Steps for this issue</param>
 		/// <param name="suspects"></param>
 		/// <param name="suspectUsers"></param>
 		/// <param name="showDesktopAlerts"></param>
 		/// <param name="externalIssueKey"></param>
-		public IssueDetails(IIssue issue, IUser? owner, IUser? nominatedBy, IUser? resolvedBy, IUser? quarantinedBy, IUser? forceClosedBy, IReadOnlyList<IIssueSpan> spans, IReadOnlyList<IIssueStep> steps, IReadOnlyList<IIssueSuspect> suspects, IReadOnlyList<IUser> suspectUsers, bool showDesktopAlerts, string? externalIssueKey)
+		public IssueDetails(IIssue issue, IUser? owner, IUser? nominatedBy, IUser? resolvedBy, IUser? quarantinedBy, IUser? forceClosedBy, IUser? rootCauseOwner, IReadOnlyList<IIssueSpan> spans, IReadOnlyList<IIssueStep> steps, IReadOnlyList<IIssueSuspect> suspects, IReadOnlyList<IUser> suspectUsers, bool showDesktopAlerts, string? externalIssueKey)
 		{
 			Issue = issue;
 			Owner = owner;
@@ -98,6 +102,7 @@ namespace HordeServer.Issues
 			QuarantinedBy = quarantinedBy;
 			QuarantineTimeUtc = issue.QuarantineTimeUtc;
 			ForceClosedBy = forceClosedBy;
+			RootCauseOwner = rootCauseOwner;
 			Spans = spans;
 			Steps = steps;
 			Suspects = suspects;
@@ -424,6 +429,7 @@ namespace HordeServer.Issues
 			IUser? resolvedBy = (issue.ResolvedById.HasValue && issue.ResolvedById != IIssue.ResolvedByTimeoutId && issue.ResolvedById != IIssue.ResolvedByUnknownId) ? await _userCollection.GetCachedUserAsync(issue.ResolvedById.Value, cancellationToken) : null;
 			IUser? quarantinedBy = issue.QuarantinedByUserId.HasValue ? await _userCollection.GetCachedUserAsync(issue.QuarantinedByUserId.Value, cancellationToken) : null;
 			IUser? forceClosedBy = issue.ForceClosedByUserId.HasValue ? await _userCollection.GetCachedUserAsync(issue.ForceClosedByUserId.Value, cancellationToken) : null;
+			IUser? rootCauseOwner = issue.RootCauseOwnerId.HasValue ? await _userCollection.GetCachedUserAsync(issue.RootCauseOwnerId.Value, cancellationToken) : null;
 
 			IReadOnlyList<IIssueSpan> spans = await _issueCollection.FindSpansAsync(issue.Id, cancellationToken);
 			IReadOnlyList<IIssueStep> steps = await _issueCollection.FindStepsAsync(spans.Select(x => x.Id), cancellationToken);
@@ -439,7 +445,7 @@ namespace HordeServer.Issues
 				}
 			}
 
-			return new IssueDetails(issue, owner, nominatedBy, resolvedBy, quarantinedBy, forceClosedBy, spans, steps, suspects, suspectUsers, ShowDesktopAlertsForIssue(issue, spans), issue.ExternalIssueKey);
+			return new IssueDetails(issue, owner, nominatedBy, resolvedBy, quarantinedBy, forceClosedBy, rootCauseOwner, spans, steps, suspects, suspectUsers, ShowDesktopAlertsForIssue(issue, spans), issue.ExternalIssueKey);
 		}
 
 		/// <inheritdoc/>
@@ -480,6 +486,7 @@ namespace HordeServer.Issues
 		/// <param name="fixCommitId">Fix changelist for the issue</param>
 		/// <param name="fixSystemic">There was a systemic fix for this issue</param>
 		/// <param name="rootCommitId">Root changelist for the issue</param>
+		/// <param name="rootCauseOwnerId">New owner of the root cause</param>
 		/// <param name="rootCauseCategory">Root cause category for the issue</param>
 		/// <param name="rootCauseSummary">Root cause summary for the issue</param>
 		/// <param name="duplicateIssueId">Duplicate issue id</param>
@@ -490,10 +497,10 @@ namespace HordeServer.Issues
 		/// <param name="quarantinedById">User who has quarantined the issue</param>
 		/// <param name="forceClosedById">User who has force closed the issue</param>
 		/// <param name="initiatedById">User initiating the changes, for auditing purposes</param>
-		/// <param name="workflowThreadUrl">The workfloe thread created for the issue</param>
+		/// <param name="workflowThreadUrl">The workflow thread created for the issue</param>
 		/// <param name="cancellationToken">Cancellation token for the operation</param>
 		/// <returns>True if the issue was updated</returns>
-		public async Task<bool> UpdateIssueAsync(int id, string? summary = null, string? description = null, bool? promoted = null, UserId? ownerId = null, UserId? nominatedById = null, bool? acknowledged = null, UserId? declinedById = null, CommitId? fixCommitId = null, bool? fixSystemic = null, CommitId? rootCommitId = null, string? rootCauseCategory = null, string? rootCauseSummary = null, int? duplicateIssueId = null, UserId? resolvedById = null, List<ObjectId>? addSpanIds = null, List<ObjectId>? removeSpanIds = null, string? externalIssueKey = null, UserId? quarantinedById = null, UserId? forceClosedById = null, UserId? initiatedById = null, Uri? workflowThreadUrl = null, CancellationToken cancellationToken = default)
+		public async Task<bool> UpdateIssueAsync(int id, string? summary = null, string? description = null, bool? promoted = null, UserId? ownerId = null, UserId? nominatedById = null, bool? acknowledged = null, UserId? declinedById = null, CommitId? fixCommitId = null, bool? fixSystemic = null, CommitId? rootCommitId = null, UserId? rootCauseOwnerId = null, string? rootCauseCategory = null, string? rootCauseSummary = null, int? duplicateIssueId = null, UserId? resolvedById = null, List<ObjectId>? addSpanIds = null, List<ObjectId>? removeSpanIds = null, string? externalIssueKey = null, UserId? quarantinedById = null, UserId? forceClosedById = null, UserId? initiatedById = null, Uri? workflowThreadUrl = null, CancellationToken cancellationToken = default)
 		{
 			using TelemetrySpan tracerSpan = _tracer.StartActiveSpan($"{nameof(IssueService)}.{nameof(UpdateIssueAsync)}");
 			
@@ -507,7 +514,7 @@ namespace HordeServer.Issues
 					return false;
 				}
 
-				issue = await _issueCollection.TryUpdateIssueAsync(issue, initiatedById, new UpdateIssueOptions { UserSummary = summary, Description = description, Promoted = promoted, OwnerId = ownerId ?? resolvedById, NominatedById = nominatedById, DeclinedById = declinedById, Acknowledged = acknowledged, FixCommitId = fixCommitId, FixSystemic = fixSystemic, RootCommitId = rootCommitId, RootCauseCategory = rootCauseCategory, RootCauseSummary = rootCauseSummary, DuplicateIssueId = duplicateIssueId, ResolvedById = resolvedById, ExcludeSpanIds = removeSpanIds, ExternalIssueKey = externalIssueKey, QuarantinedById = quarantinedById, ForceClosedById = forceClosedById, WorkflowThreadUrl = workflowThreadUrl }, cancellationToken: cancellationToken);
+				issue = await _issueCollection.TryUpdateIssueAsync(issue, initiatedById, new UpdateIssueOptions { UserSummary = summary, Description = description, Promoted = promoted, OwnerId = ownerId ?? resolvedById, NominatedById = nominatedById, DeclinedById = declinedById, Acknowledged = acknowledged, FixCommitId = fixCommitId, FixSystemic = fixSystemic, RootCommitId = rootCommitId, RootCauseOwnerId = rootCauseOwnerId, RootCauseCategory = rootCauseCategory, RootCauseSummary = rootCauseSummary, DuplicateIssueId = duplicateIssueId, ResolvedById = resolvedById, ExcludeSpanIds = removeSpanIds, ExternalIssueKey = externalIssueKey, QuarantinedById = quarantinedById, ForceClosedById = forceClosedById, WorkflowThreadUrl = workflowThreadUrl }, cancellationToken: cancellationToken);
 				if (issue != null)
 				{
 					break;
diff --git a/Engine/Source/Programs/Horde/Plugins/Build/HordeServer.Build/Issues/IssueTelemetry.cs b/Engine/Source/Programs/Horde/Plugins/Build/HordeServer.Build/Issues/IssueTelemetry.cs
--- a/Engine/Source/Programs/Horde/Plugins/Build/HordeServer.Build/Issues/IssueTelemetry.cs
+++ b/Engine/Source/Programs/Horde/Plugins/Build/HordeServer.Build/Issues/IssueTelemetry.cs
@@ -14,7 +14,7 @@ namespace HordeServer.Issues
 	/// Record used for Issue telemetry. Models a <see cref="IIssue"/>.
 	/// </summary>
 	/// <remarks>Modifications to this record's structure should be versioned appropriately through <see cref="SchemaVersion"/> via <see cref="CurrentSchemaVersion"/>, and call sites should continue to publish in a backwards compatible manner.</remarks>
-	public record IssueTelemetry(string EventName, int SchemaVersion, string Context, int Id, DateTime? AcknowledgedAt, DateTime CreatedAt, UserId? OwnerId, string? ExternalIssueKey, CommitId? FixChange, bool FixedSystemic, CommitId? RootChange, DateTime LastSeenAt, DateTime? NominatedAt, UserId? NominatedById, DateTime? ResolvedAt, UserId? ResolvedById, IssueSeverity Severity, string Summary, DateTime? VerifiedAt, int? DuplicateIssueId, string? RootCauseCategory, string? RootCauseSummary)
+	public record IssueTelemetry(string EventName, int SchemaVersion, string Context, int Id, DateTime? AcknowledgedAt, DateTime CreatedAt, UserId? OwnerId, string? ExternalIssueKey, CommitId? FixChange, bool FixedSystemic, CommitId? RootChange, DateTime LastSeenAt, DateTime? NominatedAt, UserId? NominatedById, DateTime? ResolvedAt, UserId? ResolvedById, IssueSeverity Severity, string Summary, DateTime? VerifiedAt, int? DuplicateIssueId, UserId? RootCauseOwnerId, string? RootCauseCategory, string? RootCauseSummary)
 	{
 		/// <summary>
 		/// Default event name for the IssueTelemetry.
@@ -29,7 +29,7 @@ namespace HordeServer.Issues
 		/// <summary>
 		/// Current schema version.
 		/// </summary>
-		public static readonly int CurrentSchemaVersion = 3;
+		public static readonly int CurrentSchemaVersion = 4;
 
 		/// <summary>
 		/// Telemetry payload for <see cref="IIssue"/>.
@@ -56,6 +56,7 @@ namespace HordeServer.Issues
 					issue.Summary,
 					issue.VerifiedAt,
 					issue.DuplicateIssueId,
+					issue.RootCauseOwnerId,
 					issue.RootCauseCategory,
 					issue.RootCauseSummary
 				  )
diff --git a/Engine/Source/Programs/Horde/Plugins/Build/HordeServer.Build/Issues/IssuesController.cs b/Engine/Source/Programs/Horde/Plugins/Build/HordeServer.Build/Issues/IssuesController.cs
--- a/Engine/Source/Programs/Horde/Plugins/Build/HordeServer.Build/Issues/IssuesController.cs
+++ b/Engine/Source/Programs/Horde/Plugins/Build/HordeServer.Build/Issues/IssuesController.cs
@@ -155,6 +155,7 @@ namespace HordeServer.Issues
 						IUser? nominatedBy = null;
 						IUser? resolvedBy = null;
 						IUser? quarantinedBy = null;
+						IUser? rootCauseOwner = null;
 
 						if (issue.OwnerId != null)
 						{
@@ -176,7 +177,12 @@ namespace HordeServer.Issues
 							quarantinedBy = await _userCollection.GetCachedUserAsync(issue.QuarantinedByUserId.Value, cancellationToken);
 						}
 
-						FindIssueResponse response = NewFindIssueResponse(issue, owner, nominatedBy, resolvedBy, quarantinedBy, streamSeverity, spanResponses, openWorkflowIds.ToList());
+						if (issue.RootCauseOwnerId != null)
+						{
+							rootCauseOwner = await _userCollection.GetCachedUserAsync(issue.RootCauseOwnerId.Value, cancellationToken);
+						}
+
+						FindIssueResponse response = NewFindIssueResponse(issue, owner, nominatedBy, resolvedBy, quarantinedBy, rootCauseOwner, streamSeverity, spanResponses, openWorkflowIds.ToList());
 						responses.Add(PropertyFilter.Apply(response, filter));
 					}
 				}
@@ -229,10 +235,11 @@ namespace HordeServer.Issues
 		/// <param name="nominatedBy">User that nominated the current fixer</param>
 		/// <param name="resolvedBy">User that resolved the issue</param>
 		/// <param name="quarantinedBy">User that quarantined the issue</param>
+		/// <param name="rootCauseOwner">Owner of the root cause analysis</param>
 		/// <param name="streamSeverity">The current severity in the stream</param>
 		/// <param name="spans">Spans for this issue</param>
 		/// <param name="openWorkflows">List of workflows for which this issue is open</param>
-		static FindIssueResponse NewFindIssueResponse(IIssue issue, IUser? owner, IUser? nominatedBy, IUser? resolvedBy, IUser? quarantinedBy, IssueSeverity? streamSeverity, List<FindIssueSpanResponse> spans, List<WorkflowId> openWorkflows)
+		static FindIssueResponse NewFindIssueResponse(IIssue issue, IUser? owner, IUser? nominatedBy, IUser? resolvedBy, IUser? quarantinedBy, IUser? rootCauseOwner, IssueSeverity? streamSeverity, List<FindIssueSpanResponse> spans, List<WorkflowId> openWorkflows)
 		{
 			FindIssueResponse response = new FindIssueResponse();
 			response.Id = issue.Id;
@@ -255,6 +262,10 @@ namespace HordeServer.Issues
 			response.FixCommitId = issue.FixCommitId;
 			response.FixSystemic = issue.FixSystemic;
 			response.RootCommitId = issue.RootCommitId;
+			if (rootCauseOwner != null)
+			{
+				response.RootCauseOwner = rootCauseOwner.ToThinApiResponse();
+			}
 			response.RootCauseCategory = issue.RootCauseCategory;
 			response.RootCauseSummary = issue.RootCauseSummary;
 			response.DuplicateIssueId = issue.DuplicateIssueId;
@@ -507,6 +518,9 @@ namespace HordeServer.Issues
 			response.FixCommitId = issue.FixCommitId;
 			response.FixSystemic = issue.FixSystemic;
 			response.RootCommitId = issue.RootCommitId;
+			response.RootCauseOwner = details.RootCauseOwner?.Login;
+			response.RootCauseOwnerId = details.RootCauseOwner?.Id.ToString();
+			response.RootCauseOwnerInfo = details.RootCauseOwner?.ToThinApiResponse();
 			response.RootCauseCategory = issue.RootCauseCategory;
 			response.RootCauseSummary = issue.RootCauseSummary;
 			response.DuplicateIssueId = issue.DuplicateIssueId;
@@ -799,6 +813,12 @@ namespace HordeServer.Issues
 				newOwnerId = request.OwnerId.Length == 0 ? UserId.Empty : UserId.Parse(request.OwnerId);
 			}
 
+			UserId? newRootCauseOwnerId = null;
+			if (request.RootCauseOwnerId != null)
+			{
+				newRootCauseOwnerId = request.RootCauseOwnerId.Length == 0 ? UserId.Empty : UserId.Parse(request.RootCauseOwnerId);
+			}
+
 			UserId? newNominatedById = null;
 			if (request.NominatedById != null)
 			{
@@ -845,7 +865,7 @@ namespace HordeServer.Issues
 				removeSpans = request.RemoveSpans.ConvertAll(x => ObjectId.Parse(x));
 			}
 
-			if (!await _issueService.UpdateIssueAsync(issueId, request.Summary, request.Description, request.Promoted, newOwnerId, newNominatedById, request.Acknowledged, newDeclinedById, request.FixCommitId, request.FixSystemic, request.RootCommitId, request.RootCauseCategory, request.RootCauseSummary, request.DuplicateIssueId, newResolvedById, addSpans, removeSpans, request.ExternalIssueKey, newQuarantinedById, newForceClosedById, initiatedById: User.GetUserId()))
+			if (!await _issueService.UpdateIssueAsync(issueId, request.Summary, request.Description, request.Promoted, newOwnerId, newNominatedById, request.Acknowledged, newDeclinedById, request.FixCommitId, request.FixSystemic, request.RootCommitId, newRootCauseOwnerId, request.RootCauseCategory, request.RootCauseSummary, request.DuplicateIssueId, newResolvedById, addSpans, removeSpans, request.ExternalIssueKey, newQuarantinedById, newForceClosedById, initiatedById: User.GetUserId()))
 			{
 				return NotFound();
 			}
diff --git a/Engine/Source/Programs/Horde/Plugins/Build/HordeServer.Build/Issues/WorkflowConfig.cs b/Engine/Source/Programs/Horde/Plugins/Build/HordeServer.Build/Issues/WorkflowConfig.cs
--- a/Engine/Source/Programs/Horde/Plugins/Build/HordeServer.Build/Issues/WorkflowConfig.cs
+++ b/Engine/Source/Programs/Horde/Plugins/Build/HordeServer.Build/Issues/WorkflowConfig.cs
@@ -65,6 +65,11 @@ namespace HordeServer.Issues
 
 		IReadOnlyList<int> IWorkflow.EscalateTimes => EscalateTimes;
 
+		/// <inheritdoc cref="IWorkflow.EscalateRootCauseAnalysisTimes"/>
+		public List<int> EscalateRootCauseAnalysisTimes { get; set; } = [60 * 24];
+
+		IReadOnlyList<int> IWorkflow.EscalateRootCauseAnalysisTimes => EscalateRootCauseAnalysisTimes;
+
 		/// <inheritdoc/>
 		public int MaxMentions { get; set; } = 5;
 
diff --git a/Engine/Source/Programs/Horde/Plugins/Build/HordeServer.Build/Notifications/Sinks/SlackNotificationSink.cs b/Engine/Source/Programs/Horde/Plugins/Build/HordeServer.Build/Notifications/Sinks/SlackNotificationSink.cs
--- a/Engine/Source/Programs/Horde/Plugins/Build/HordeServer.Build/Notifications/Sinks/SlackNotificationSink.cs
+++ b/Engine/Source/Programs/Horde/Plugins/Build/HordeServer.Build/Notifications/Sinks/SlackNotificationSink.cs
@@ -227,6 +227,7 @@ namespace HordeServer.Notifications.Sinks
 		readonly JsonSerializerOptions _jsonSerializerOptions;
 		readonly ITicker _escalateTicker;
 		static readonly RedisSortedSetKey<int> s_escalateIssues = "slack/escalate";
+		static readonly RedisSortedSetKey<int> s_escalateRootCauseAnalysis = "slack/escalate-rca";
 		readonly BackgroundTask _backgroundTask;
 		readonly IServerInfo _serverInfo;
 		readonly IOptionsMonitor<BuildConfig> _buildConfig;
@@ -1317,6 +1318,15 @@ namespace HordeServer.Notifications.Sinks
 					}
 				}
 
+				if (issue.ResolvedAt != null && workflow.EscalateRootCauseAnalysisTimes.Count > 0)
+				{
+					DateTime escalateTime = issue.ResolvedAt.Value.AddMinutes(workflow.EscalateRootCauseAnalysisTimes[0]);
+					if (await _redisService.GetDatabase().SortedSetAddAsync(s_escalateRootCauseAnalysis, issue.Id, (escalateTime - DateTime.UnixEpoch).TotalSeconds, SortedSetWhen.NotExists))
+					{
+						_logger.LogInformation("First escalation time for root cause analysis {IssueId} is {Time}", issue.Id, escalateTime);
+					}
+				}
+
 				if ((workflow.TriageAlias != null || workflow.TriageTypeAliases != null) && issue.OwnerId == null && (suspects.All(x => x.DeclinedAt != null) || notifyTriageAlias) && !closed)
 				{
 					string? triageAlias;
@@ -2902,6 +2912,50 @@ namespace HordeServer.Notifications.Sinks
 					}
 				}
 			}
+
+			await EscalateIssuesForKeyAsync(s_escalateRootCauseAnalysis, time, utcNow, "root cause analysis",
+				EscalateSingleRootCauseAnalysisAsync, cancellationToken);
+		}
+
+		async ValueTask EscalateIssuesForKeyAsync(RedisSortedSetKey<int> key, double time, DateTime utcNow, string notificationName,
+			Func<int, DateTime, string, CancellationToken, Task<double?>> escalateSingleIssueAsync, CancellationToken cancellationToken)
+		{
+			int[] issueIds = await _redisService.GetDatabase().SortedSetRangeByScoreAsync(key, 0, time);
+			if (issueIds.Length > 0)
+			{
+				_logger.LogInformation("Escalating {NotificationName} for {Time} ({TimeSecs})", notificationName, utcNow, time);
+				foreach (int issueId in issueIds)
+				{
+					cancellationToken.ThrowIfCancellationRequested();
+					try
+					{
+						double? nextTime = await escalateSingleIssueAsync(issueId, utcNow, notificationName, cancellationToken);
+						if (nextTime == null)
+						{
+							_logger.LogInformation("Cancelling escalation for {NotificationName} {IssueId}", notificationName, issueId);
+							await _redisService.GetDatabase().SortedSetRemoveAsync(key, issueId);
+						}
+						else
+						{
+							_logger.LogInformation("Next escalation for {NotificationName} {IssueId} is at timestamp {Time}", notificationName, issueId, nextTime.Value);
+							await _redisService.GetDatabase().SortedSetAddAsync(key, issueId, nextTime.Value);
+						}
+					}
+					catch (SlackException ex)
+					{
+						if (ex.Message.Equals("channel_not_found", StringComparison.Ordinal))
+						{
+							_logger.LogWarning(ex, "Slack channel not found while escalating {NotificationName} {IssueId}; cancelling.", notificationName, issueId);
+						}
+						else
+						{
+							_logger.LogError(ex, "Slack exception while escalating {NotificationName} {IssueId}; cancelling.", notificationName, issueId);
+						}
+
+						await _redisService.GetDatabase().SortedSetRemoveAsync(key, issueId);
+					}
+				}
+			}
 		}
 
 		async Task<double?> EscalateSingleIssueAsync(BuildConfig buildConfig, int issueId, DateTime utcNow, CancellationToken cancellationToken)
@@ -3016,6 +3070,106 @@ namespace HordeServer.Notifications.Sinks
 			return (nextEscalationTime - DateTime.UnixEpoch).TotalSeconds;
 		}
 
+		async Task<double?> EscalateSingleRootCauseAnalysisAsync(int issueId, DateTime utcNow, string notificationName, CancellationToken cancellationToken)
+		{
+			IIssue? issue = await _issueService.Collection.GetIssueAsync(issueId, cancellationToken);
+			if (issue == null)
+			{
+				return null;
+			}
+
+			IReadOnlyList<IIssueSpan> spans = await _issueService.Collection.FindSpansAsync(issueId, cancellationToken);
+			if (spans.Count == 0)
+			{
+				return null;
+			}
+
+			if (issue.ResolvedAt == null)
+			{
+				return null;
+			}
+
+			UserId? ownerId = issue.RootCauseOwnerId ?? issue.OwnerId ?? issue.ResolvedById;
+			if (ownerId == null || ownerId == UserId.Empty)
+			{
+				return null;
+			}
+
+			// RootCommitId, RootCauseCategory and RootCauseSummary are optional inputs in the UI because the issue closer may not have all the
+			// information when marking as fixed; they are all required to have a value for Root Cause Analysis to be considered closed.
+			if (issue.RootCommitId != null && !String.IsNullOrEmpty(issue.RootCauseCategory) && !String.IsNullOrEmpty(issue.RootCauseSummary))
+			{
+				// RCA is completed
+				return null;
+			}
+
+			IIssueSpan span = spans[0];
+
+			WorkflowId? workflowId = span.LastFailure.Annotations.WorkflowId;
+			if (workflowId == null)
+			{
+				return null;
+			}
+
+			if (!_buildConfig.CurrentValue.TryGetStream(span.StreamId, out StreamConfig? streamConfig))
+			{
+				return null;
+			}
+			if (!streamConfig.TryGetWorkflow(workflowId.Value, out WorkflowConfig? workflow))
+			{
+				return null;
+			}
+			if (String.IsNullOrEmpty(workflow.TriageChannel) || workflow.EscalateRootCauseAnalysisTimes.Count == 0)
+			{
+				return null;
+			}
+
+			if (!issue.IsMarkedFixed())
+			{
+				return null;
+			}
+
+			MessageStateDocument? state = await GetMessageStateAsync(workflow.TriageChannel, GetTriageThreadEventId(issueId), cancellationToken);
+			if (state == null)
+			{
+				return null;
+			}
+
+			string mention = await FormatMentionAsync(ownerId.Value, workflow.AllowMentions, cancellationToken);
+			Uri issueUrl = GetIssueUrl(issue, span.FirstFailure);
+			string message = $"{mention} this root cause analysis requires your attention - Issue <{issueUrl}|{issue.Id}> - Select ...More > Modify Issue to complete the root cause analysis.";
+			await _slackClient.PostMessageToThreadAsync(state.MessageId, message, cancellationToken);
+
+			return GetNextEscalationTime(issue, issue.ResolvedAt.Value, utcNow, workflow.EscalateRootCauseAnalysisTimes, notificationName);
+		}
+
+		double GetNextEscalationTime(IIssue issue, DateTime whence, DateTime utcNow, List<int> escalationTimes, string notificationName)
+		{
+			DateTime nextEscalationTime = whence;
+			for (int idx = 0; ; idx++)
+			{
+				if (idx >= escalationTimes.Count)
+				{
+					nextEscalationTime = utcNow + TimeSpan.FromMinutes(escalationTimes[^1]);
+					_logger.LogInformation("Next escalation time for {NotificationName} {IssueId} is {NextEscalationTime} (now + {Time}m [{Idx}/{Count}])",
+						notificationName, issue.Id, nextEscalationTime, escalationTimes[^1], escalationTimes.Count, escalationTimes.Count);
+					break;
+				}
+
+				DateTime prevEscalationTime = nextEscalationTime;
+				nextEscalationTime += TimeSpan.FromMinutes(escalationTimes[idx]);
+
+				if (nextEscalationTime > utcNow + TimeSpan.FromMinutes(5.0))
+				{
+					_logger.LogInformation("Next escalation time for {NotificationName} {IssueId} is {NextEscalationTime} ({PrevEscalationTime} + {Time}m [{Idx}/{Count}])",
+						notificationName, issue.Id, nextEscalationTime, prevEscalationTime, escalationTimes[idx], idx + 1, escalationTimes.Count);
+					break;
+				}
+			}
+
+			return (nextEscalationTime - DateTime.UnixEpoch).TotalSeconds;
+		}
+
 		/// <inheritdoc/>
 		async Task ExecuteAsync(CancellationToken stoppingToken)
 		{
@@ -3250,6 +3404,15 @@ namespace HordeServer.Notifications.Sinks
 								_logger.LogInformation("Marking issue {IssueId} fixed by user {UserId} in {Change}", issueId, resolvedById ?? userId, fixCommitId);
 							}
 
+							UserId? rootCauseOwnerId = null;
+							if (payload.View.State.TryGetValue("rootcause_owner", "rootcause_owner_action", out string? rootCauseOwnerStr))
+							{
+								if (UserId.TryParse(rootCauseOwnerStr, out UserId rcaId))
+								{
+									rootCauseOwnerId = rcaId;
+								}
+							}
+
 							string? rootChange;
 							CommitId rootCommit = new CommitId("");
 							if (payload.View.State.TryGetValue("rootcause_cl", "rootcause_cl_action", out rootChange))
@@ -3342,7 +3505,7 @@ namespace HordeServer.Notifications.Sinks
 								summary = curSummary;
 							}
 
-							await _issueService.UpdateIssueAsync(issueId, summary: summary, fixCommitId: fixCommitId, resolvedById: resolvedById, initiatedById: userId, duplicateIssueId: dupeIssueId, rootCommitId: rootCommit, rootCauseCategory: rootCategory, rootCauseSummary: rootSummary, cancellationToken: cancellationToken);
+							await _issueService.UpdateIssueAsync(issueId, summary: summary, fixCommitId: fixCommitId, resolvedById: resolvedById, initiatedById: userId, duplicateIssueId: dupeIssueId, rootCommitId: rootCommit, rootCauseOwnerId: rootCauseOwnerId, rootCauseCategory: rootCategory, rootCauseSummary: rootSummary, cancellationToken: cancellationToken);
 
 						}
 						else if (TryMatch(payload.View.CallbackId, @"^issue_(\d+)_ack_([a-fA-F0-9]{24})$", out match))
@@ -3537,6 +3700,40 @@ namespace HordeServer.Notifications.Sinks
 
 				if (issue.DuplicateIssueId == null)
 				{
+					List<SlackOption> rootCauseOwnerOptions = [];
+					if (issue.RootCauseOwnerId != null && issue.RootCauseOwnerId != user.Id)
+					{
+						IUser? rootCauseOwner = await _userCollection.GetCachedUserAsync(issue.RootCauseOwnerId.Value, cancellationToken);
+						if (rootCauseOwner != null)
+						{
+							rootCauseOwnerOptions.Add(new SlackOption($"Current root cause owner ({rootCauseOwner.Name})", rootCauseOwner.Id.ToString()));
+						}
+					}
+					if (issue.OwnerId != null && issue.OwnerId != user.Id)
+					{
+						IUser? owner = await _userCollection.GetCachedUserAsync(issue.OwnerId.Value, cancellationToken);
+						if (owner != null)
+						{
+							rootCauseOwnerOptions.Add(new SlackOption($"Current owner ({owner.Name})", owner.Id.ToString()));
+						}
+					}
+					if (issue.ResolvedById != null && issue.ResolvedById != user.Id)
+					{
+						IUser? resolvedBy = await _userCollection.GetCachedUserAsync(issue.ResolvedById.Value, cancellationToken);
+						if (resolvedBy != null)
+						{
+							rootCauseOwnerOptions.Add(new SlackOption($"Fixed by ({resolvedBy.Name})", resolvedBy.Id.ToString()));
+						}
+					}
+					rootCauseOwnerOptions.Add(new SlackOption("Me", user.Id.ToString()));
+					//rootCauseOwnerOptions.Add(new SlackOption("Unassigned", UserId.Empty.Id.ToString()));
+
+					RadioButtonGroupElement rootCauseOwnership = new("rootcause_owner_action", rootCauseOwnerOptions);
+					//rootCauseOwnership.InitialOption = rootCauseOwnerOptions[0];
+					InputBlock rootCauseOwnerBlock = view.AddInput("Root Cause Owner", rootCauseOwnership);
+					rootCauseOwnerBlock.BlockId = "rootcause_owner";
+					rootCauseOwnerBlock.Optional = true;
+
 					InputBlock selectBlock = view.AddInput("Root Cause Category:", categorySelect);
 					selectBlock.BlockId = "rootcause_category";
 					selectBlock.Optional = true;
@@ -3548,7 +3745,6 @@ namespace HordeServer.Notifications.Sinks
 					InputBlock changeBlock = view.AddInput("Root Cause Changelist:", new PlainTextInputElement("rootcause_cl_action", placeholder: "Root Cause Changelist", initialValue: issue?.RootCommitId?.Name));
 					changeBlock.BlockId = "rootcause_cl";
 					changeBlock.Hint = "Change which caused this issue, please only enter if certain";
-
 					changeBlock.Optional = true;
 				}
 
diff --git a/Engine/Source/Programs/Shared/EpicGames.Horde/Issues/IWorkflow.cs b/Engine/Source/Programs/Shared/EpicGames.Horde/Issues/IWorkflow.cs
--- a/Engine/Source/Programs/Shared/EpicGames.Horde/Issues/IWorkflow.cs
+++ b/Engine/Source/Programs/Shared/EpicGames.Horde/Issues/IWorkflow.cs
@@ -91,6 +91,11 @@ namespace EpicGames.Horde.Issues
 		/// </summary>
 		IReadOnlyList<int> EscalateTimes { get; }
 
+		/// <summary>
+		/// Times after a root cause analysis has been opened to escalate to the owner, in minutes. Continues to notify on the last interval once reaching the end of the list.
+		/// </summary>
+		IReadOnlyList<int> EscalateRootCauseAnalysisTimes { get; }
+
 		/// <summary>
 		/// Maximum number of people to mention on a triage thread
 		/// </summary>
diff --git a/Engine/Source/Programs/Shared/EpicGames.Horde/Issues/IssueMessages.cs b/Engine/Source/Programs/Shared/EpicGames.Horde/Issues/IssueMessages.cs
--- a/Engine/Source/Programs/Shared/EpicGames.Horde/Issues/IssueMessages.cs
+++ b/Engine/Source/Programs/Shared/EpicGames.Horde/Issues/IssueMessages.cs
@@ -408,6 +408,21 @@ namespace EpicGames.Horde.Issues
 		}
 		CommitId? _rootCommitId;
 
+		/// <summary>
+		/// User that provided the root cause analysis
+		/// </summary>
+		public string? RootCauseOwnerId { get; set; }
+
+		/// <summary>
+		/// User that provided the root cause analysis
+		/// </summary>
+		public string? RootCauseOwner { get; set; }
+
+		/// <summary>
+		/// User that provided the root cause analysis
+		/// </summary>
+		public GetThinUserInfoResponse? RootCauseOwnerInfo { get; set; }
+
 		/// <summary>
 		/// Root Cause Category for this issue
 		/// </summary>
@@ -640,6 +655,11 @@ namespace EpicGames.Horde.Issues
 		}
 		CommitId? _fixCommitId;
 
+		/// <summary>
+		/// User that provided the root cause analysis
+		/// </summary>
+		public GetThinUserInfoResponse? RootCauseOwner { get; set; }
+
 		/// <summary>
 		/// Root Cause Category for this issue
 		/// </summary>
@@ -807,6 +827,11 @@ namespace EpicGames.Horde.Issues
 		}
 		CommitId? _rootCommitId;
 
+		/// <summary>
+		/// New user id for owner of the root cause analysis, can be cleared by passing empty string
+		/// </summary>
+		public string? RootCauseOwnerId { get; set; }
+
 		/// <summary>
 		/// Root Cause Category for this issue
 		/// </summary>
diff --git a/Engine/Source/Programs/Shared/EpicGames.Horde/Issues/WorkflowMessages.cs b/Engine/Source/Programs/Shared/EpicGames.Horde/Issues/WorkflowMessages.cs
--- a/Engine/Source/Programs/Shared/EpicGames.Horde/Issues/WorkflowMessages.cs
+++ b/Engine/Source/Programs/Shared/EpicGames.Horde/Issues/WorkflowMessages.cs
@@ -65,6 +65,11 @@ namespace EpicGames.Horde.Issues
 		public List<int> EscalateTimes { get; set; }
 		IReadOnlyList<int> IWorkflow.EscalateTimes => EscalateTimes;
 
+		/// <inheritdoc cref="IWorkflow.EscalateRootCauseAnalysisTimes"/>
+		public List<int> EscalateRootCauseAnalysisTimes { get; set; }
+
+		IReadOnlyList<int> IWorkflow.EscalateRootCauseAnalysisTimes => EscalateRootCauseAnalysisTimes;
+
 		/// <inheritdoc/>
 		public int MaxMentions { get; set; }
 
@@ -99,6 +104,7 @@ namespace EpicGames.Horde.Issues
 		{
 			ReportTimes = new List<TimeSpan>();
 			EscalateTimes = new List<int>();
+			EscalateRootCauseAnalysisTimes = [];
 		}
 
 		/// <summary>
@@ -122,6 +128,7 @@ namespace EpicGames.Horde.Issues
 			TriageErrors = workflow.TriageErrors;
 			EscalateAlias = workflow.EscalateAlias;
 			EscalateTimes = workflow.EscalateTimes.ToList();
+			EscalateRootCauseAnalysisTimes = workflow.EscalateRootCauseAnalysisTimes.ToList();
 			MaxMentions = workflow.MaxMentions;
 			AllowMentions = workflow.AllowMentions;
 			InviteRestrictedUsers = workflow.InviteRestrictedUsers;
