Allow filtering by test metadata (Team, Harness, Profile)

#rnx

#Synced-CL 48144241

diff --git a/Engine/Source/Programs/Horde/HordeDashboard/plugins/test-automation/testAutomationView.tsx b/Engine/Source/Programs/Horde/HordeDashboard/plugins/test-automation/testAutomationView.tsx
--- a/Engine/Source/Programs/Horde/HordeDashboard/plugins/test-automation/testAutomationView.tsx
+++ b/Engine/Source/Programs/Horde/HordeDashboard/plugins/test-automation/testAutomationView.tsx
@@ -174,6 +174,11 @@ const TagChooser: React.FC = observer(() => {
         </Stack>
 });
 
+enum MetaType {
+    Session = "Session",
+    Test = "Test",
+}
+
 const MetaKeyChooser: React.FC<{ onUpdateKey: (key: string) => void }> = observer(({ onUpdateKey }) => {
 
     handler.subscribeToFilter();
@@ -182,17 +187,31 @@ const MetaKeyChooser: React.FC<{ onUpdateKey: (key: string) => void }> = observe
     const cmetakeys: Set<string> = new Set(handler.filterState.metadata?.keys() ?? []);
     const streamMetadataKeys = handler.allStreamMetadataKeys;
 
+    const ctestmetakeys: Set<string> = new Set(handler.filterState.testmeta?.keys() ?? []);
+    const streamTestMetaKeys = handler.allStreamTestMetaKeys;
+
     const options: IContextualMenuItem[] = [];
     streamMetadataKeys.forEach(a => {
         if (cmetakeys.has(a)) {
             return;
         }
         options.push({
-            key: `metakey_${a}`, text: a, onClick: () => {
+            key: `metakey_${a}`, text: a, data: MetaType.Session, onClick: () => {
+                onUpdateKey(a);
+            }
+        });
+    });
+    streamTestMetaKeys.forEach(a => {
+        if (ctestmetakeys.has(a)) {
+            return;
+        }
+        options.push({
+            key: `tmetakey_${a}`, text: a.charAt(0).toUpperCase()+a.slice(1), data: MetaType.Test, onClick: () => {
                 onUpdateKey(a);
             }
         });
     });
+    options.sort((a, b) => a.text!.localeCompare(b.text!));
 
     const menuProps: IContextualMenuProps = {
         shouldFocusOnMount: true,
@@ -210,8 +229,14 @@ const MetaValueChooser: React.FC<{ metakey: string, onRemove: (key: string) => v
 
     handler.subscribeToFilter();
 
-    const cmetavalues: Set<string> = new Set(handler.filterState.metadata?.get(metakey) ?? []);
-    const metavalues = handler.getAllStreamMetadataValues(metakey);
+    const metaType = handler.allStreamTestMetaKeys.includes(metakey)? MetaType.Test : MetaType.Session;
+    const state = metaType === MetaType.Test? handler.filterState.testmeta : handler.filterState.metadata;
+
+    const cmetavalues: Set<string> = new Set(state?.get(metakey) ?? []);
+    const metavalues = metaType === MetaType.Test? handler.getAllStreamTestMetaValues(metakey) : handler.getAllStreamMetadataValues(metakey);
+
+    const addMeta = metaType === MetaType.Test? (key: string, value: string) => handler.addTestMeta(key, value) : (key: string, value: string) => handler.addMetadata(key,value);
+    const removeMeta = metaType === MetaType.Test? (key: string, value?: string) => handler.removeTestMeta(key, value) : (key: string, value: string) => handler.removeMetadata(key, value);
 
     let options: IComboBoxOption[] = !!metavalues.length ? metavalues.map(t => { return { key: t, text: t } }) : cmetavalues.keys().toArray().map(t => { return { key: t, text: `! ${t}` } });
     const diffmetavalues = cmetavalues.difference(new Set(metavalues));
@@ -224,22 +249,22 @@ const MetaValueChooser: React.FC<{ metakey: string, onRemove: (key: string) => v
             content={noMatch ? "meta key or value(s) missing in this stream" : undefined}
             directionalHint={DirectionalHint.rightCenter}>
                 <Stack style={{ paddingTop: 12, paddingBottom: 4 }} horizontal>
-                <Label style={{ paddingTop: 5, paddingRight: 5 }}>{metakey}</Label>
+                <Label style={{ paddingTop: 5, paddingRight: 5 }}>{metakey.charAt(0).toUpperCase()+metakey.slice(1)}</Label>
                 <Stack style={{ width: 270, borderWidth: 1, borderStyle: "solid", borderColor: noMatch ? getStatusColors().get(StatusColor.Failure) : "transparent" }}>
                     <MultiOptionChooser style={{height: 28}} options={options} initialSelection={cmetavalues.keys().toArray()} disabled={!metavalues.length}
                         updateKeys={(keys) => {
                             keys.forEach(t => {
-                                if (!cmetavalues.has(t))  handler.addMetadata(metakey, t);
+                                if (!cmetavalues.has(t)) addMeta(metakey, t);
                             });
                             cmetavalues.forEach(t => {
-                                if (!keys.includes(t))  handler.removeMetadata(metakey, t);
+                                if (!keys.includes(t)) removeMeta(metakey, t);
                             });
                         }} />
                 </Stack>
                 <DefaultButton style={{ minWidth: 0, fontSize: 11, padding: 5, height: 30 }} title="remove filter"
                         onClick={() => {
-                            if (handler.filterState?.metadata?.keys().some((key) => key === metakey)) {
-                                handler.removeMetadata(metakey);
+                            if (state?.keys().some((key) => key === metakey)) {
+                                removeMeta(metakey);
                             }
                             onRemove(metakey);
                         }}>
@@ -252,7 +277,9 @@ const MetaValueChooser: React.FC<{ metakey: string, onRemove: (key: string) => v
 const TestAutomationSidebarLeft: React.FC = observer(() => {
 
     handler.subscribeToFilter();
-    const filterMetaKeys = handler.filterState.metadata?.keys().toArray().sort((a, b) => a.localeCompare(b)) ?? [];
+    let filterMetaKeys = handler.filterState.metadata?.keys().toArray() ?? [];
+    filterMetaKeys = filterMetaKeys.concat(handler.filterState.testmeta?.keys().toArray() ?? []);
+    filterMetaKeys.sort((a, b) => a.localeCompare(b));
 
     const view = (handler.getSearchParam("view") as string | undefined) ?? "0";
 
diff --git a/Engine/Source/Programs/Horde/HordeDashboard/plugins/test-automation/testData.ts b/Engine/Source/Programs/Horde/HordeDashboard/plugins/test-automation/testData.ts
--- a/Engine/Source/Programs/Horde/HordeDashboard/plugins/test-automation/testData.ts
+++ b/Engine/Source/Programs/Horde/HordeDashboard/plugins/test-automation/testData.ts
@@ -214,6 +214,7 @@ type TestDataState = {
     tests?: string[];
     tags?: string[];
     metadata?: MetadataFilter;
+    testmeta?: MetadataFilter;
     weeks?: number;
 }
 
@@ -479,6 +480,21 @@ export class TestNameRef {
         return toMatch.some(tag => this.tagRefs?.values().some(r => r.isMatch(tag)));
     }
 
+     isMetaMatch(metaFilter: MetadataFilter): boolean {
+        for (const [key, values] of metaFilter.entries()) {
+            const thisValue = this[key];
+            if (!thisValue) {
+                // missing key
+                return false;
+            }
+            const toValidate = values.some((v) => v.indexOf('+') > -1) ? [thisValue] : thisValue.split('+');
+            if (!toValidate.some((v) => values.includes(v))) {
+                return false;
+            }
+        }
+        return true;
+    }
+
     updateRecipes(recipes: TestRecipeRef[]) {
         const newRecipes = new Set<TestRecipeRef>(recipes);
         this.recipeRefs = this.recipeRefs? this.recipeRefs.union(newRecipes) : newRecipes;
@@ -824,11 +840,26 @@ export class TestStreamRef {
         this.testNames = new Set();
         this.metadataMap = new Map();
         this.tagNames = new Set();
+        this.testmetaMap = new Map([
+            ['team', new Set()],
+            ['profile', new Set()],
+            ['harness', new Set()]
+        ]);
     }
 
     addTest(item: TestNameRef) {
         this.tests.push(item);
         this.testNames.add(item.name);
+        // test metadata
+        if (item.team) {
+            this.testmetaMap.get('team')!.add(item.team);
+        }
+        if (item.harness) {
+            this.testmetaMap.get('harness')!.add(item.harness);
+        }
+        if (item.profile) {
+            this.testmetaMap.get('profile')!.add(item.profile);
+        }
     }
 
     addMetadata(item: MetadataRef) {
@@ -860,18 +891,27 @@ export class TestStreamRef {
         return this.metadataMap.get(key)?.keys().toArray();
     }
 
+    getAllTestMetaKeys(): string[] {
+        return this.testmetaMap.keys().toArray();
+    }
+
+    getAllTestMetaValues(key: string): string[] | undefined {
+        return this.testmetaMap.get(key)?.keys().toArray();
+    }
+
     getAllTagNames(): string[] {
         return  this.tagNames.keys().toArray();
     }
 
-    getMatchingTests(toMatchNames?: string[], toMatchTags?: string[]): Set<TestNameRef> {
+    getMatchingTests(toMatchNames?: string[], toMatchTags?: string[], toMatchMeta?: MetadataFilter): Set<TestNameRef> {
         const lowerToMatchNames = toMatchNames?.map((i) => i.toLowerCase());
         const lowerToMatchTags = toMatchTags?.map((i) => i.toLowerCase());
         const matches: Set<TestNameRef> = new Set();
         this.tests.forEach((test) => {
             const isMatchNames = !lowerToMatchNames || lowerToMatchNames.some((key) => test.isMatch(key));
             const isMatchTags = !lowerToMatchTags || test.includesTags(lowerToMatchTags);
-            if (isMatchNames && isMatchTags) {
+            const isMatchMeta = !toMatchMeta || test.isMetaMatch(toMatchMeta);
+            if (isMatchNames && isMatchTags && isMatchMeta) {
                 matches.add(test);
             }
         })
@@ -910,6 +950,7 @@ export class TestStreamRef {
     private testNames: Set<string>;
     private metadataMap: Map<string, Set<string>>;
     private tagNames: Set<string>;
+    private testmetaMap: Map<string, Set<string>>;
 }
 
 export class TestDataFilterSearchParams {
@@ -1145,6 +1186,49 @@ export class TestDataHandler {
         this.setFilterUpdated();
     }
 
+    addTestMeta(key: string, value: string) {
+
+        const state = this.filterState;
+        const meta = state.testmeta?.get(key);
+
+        if (!!meta && meta!.includes(value)) {
+            return;
+        }
+
+        if (!state.testmeta) {
+            state.testmeta = new Map();
+        }
+
+        if (!meta) {
+            state.testmeta!.set(key, [value]);
+        } else {
+            state.testmeta!.get(key)!.push(value);
+        }
+
+        this.setFilterUpdated();
+    }
+
+    removeTestMeta(key: string, value?: string) {
+
+        const state = this.filterState;
+        const meta = state.testmeta?.get(key);
+        if (!!value) {
+            const idx = meta?.indexOf(value!);
+
+            if (!meta || !isNumber(idx) || idx! < 0) {
+                return;
+            }
+
+            meta!.splice(idx!, 1);
+        }
+
+        if (!value || !meta!.length) {
+            state.testmeta!.delete(key);
+        }
+
+        this.setFilterUpdated();
+    }
+
     addTest(test: string) {
 
         const state = this.filterState;
@@ -1302,6 +1386,38 @@ export class TestDataHandler {
         return testStreamRef.getAllMetadataValues(key)?.sort((a, b) => a.localeCompare(b)) ?? [];
     }
 
+    /// All test metadata keys from current stream
+    get allStreamTestMetaKeys(): string[] {
+
+        const state = this.filterState;
+        if (!state.stream) {
+            return [];
+        }
+
+        const testStreamRef = this.testStreams.get(state.stream);
+        if (!testStreamRef) {
+            return [];
+        }
+
+        return testStreamRef.getAllTestMetaKeys().sort((a, b) => a.localeCompare(b));
+    }
+
+    /// All test metadata values for a key from current stream
+    getAllStreamTestMetaValues(key: string): string[] {
+
+        const state = this.filterState;
+        if (!state.stream) {
+            return [];
+        }
+
+        const testStreamRef = this.testStreams.get(state.stream);
+        if (!testStreamRef) {
+            return [];
+        }
+
+        return testStreamRef.getAllTestMetaValues(key)?.sort((a, b) => a.localeCompare(b)) ?? [];
+    }
+
     /// Tests that match current filters
     get filteredTests(): TestNameRef[] {
 
@@ -1321,7 +1437,7 @@ export class TestDataHandler {
         }
 
         const metaRefs = testStreamRef!.getMatchingMeta(state.metadata);
-        const testNameRefs = testStreamRef!.getMatchingTests(state.tests, state.tags);
+        const testNameRefs = testStreamRef!.getMatchingTests(state.tests, state.tags, state.testmeta);
 
         return testStatus!.getMatchingTests(testNameRefs, metaRefs).sort((a, b) => a.name.localeCompare(b.name));
     }
@@ -1390,7 +1506,8 @@ export class TestDataHandler {
         const state = this.filterState;
         state.weeks = state.weeks ? state.weeks : undefined;
         state.stream = state.stream ? state.stream : undefined;
-        state.metadata = state.metadata && state.metadata.size > 0 ? state.metadata : undefined
+        state.metadata = state.metadata && state.metadata.size > 0 ? state.metadata : undefined;
+        state.testmeta = state.testmeta && state.testmeta.size > 0 ? state.testmeta : undefined;
         state.tests = state.tests?.sort((a, b) => a.localeCompare(b));
         state.tags = state.tags?.sort((a, b) => a.localeCompare(b));
 
@@ -1423,6 +1540,12 @@ export class TestDataHandler {
             search.set("meta", key, state.metadata?.get(key));
         });
 
+        const testMetaKeys = new Set<string>(search.getGroup("tmeta")?.keys());
+        testMetaKeys.union(new Set(state.testmeta?.keys())).keys().forEach(key => {
+            needUpdate = needUpdate || search.get('tmeta', key)?.toString() !== state.testmeta?.get(key)?.toString();
+            search.set("tmeta", key, state.testmeta?.get(key));
+        });
+
         if (needUpdate) {
             this.query(state.stream);
             return true;
@@ -1439,12 +1562,14 @@ export class TestDataHandler {
 
         const stream = search.get("stream") ?? undefined;
         const metadata = search.getGroup("meta");
+        const testmeta = search.getGroup("tmeta");
         const tests = search.get("tests") ?? undefined;
         const tags = search.get("tags") ?? undefined;
         const weeks = search.get("weeks") ?? undefined;
 
         state.stream = stream && stream.length > 0? stream[0] : undefined;
         state.metadata = metadata && metadata.size > 0 ? metadata : undefined;
+        state.testmeta = testmeta && testmeta.size > 0 ? testmeta : undefined;
         state.tests = tests;
         state.tags = tags;
         state.weeks = weeks && weeks.length > 0 ? parseInt(weeks[0]) : defaultQueryWeeks;
