Add owner and customers user picker
+ tweak audit commit visual feedback
+ stash transaction when view gets destroy and load it back when view is back up

#rnx

#Synced-CL 48144241

diff --git a/Engine/Source/Programs/Horde/HordeDashboard/plugins/test-automation/api.ts b/Engine/Source/Programs/Horde/HordeDashboard/plugins/test-automation/api.ts
--- a/Engine/Source/Programs/Horde/HordeDashboard/plugins/test-automation/api.ts
+++ b/Engine/Source/Programs/Horde/HordeDashboard/plugins/test-automation/api.ts
@@ -1,6 +1,7 @@
 // Copyright Epic Games, Inc. All Rights Reserved.
 
 import backend from "horde/backend";
+import { GetUserResponse } from "horde/backend/Api";
 
 /// A test data response
 export type GetTestDataResponse = {
@@ -741,3 +742,12 @@ export async function getTestAuditsV2(testIds: string[], minCreateTime?: string)
         });
     });
 }
+
+export async function getUsersInfo(userIds: string[]): Promise<GetUserResponse[]> {
+    const results = await backend.getUsers({
+        ids: userIds,
+        includeAvatar: true
+    });
+
+    return results;
+}
diff --git a/Engine/Source/Programs/Horde/HordeDashboard/plugins/test-automation/testAutomationCommon.tsx b/Engine/Source/Programs/Horde/HordeDashboard/plugins/test-automation/testAutomationCommon.tsx
--- a/Engine/Source/Programs/Horde/HordeDashboard/plugins/test-automation/testAutomationCommon.tsx
+++ b/Engine/Source/Programs/Horde/HordeDashboard/plugins/test-automation/testAutomationCommon.tsx
@@ -1,6 +1,6 @@
 // Copyright Epic Games, Inc. All Rights Reserved.
 
-import { mergeStyles, mergeStyleSets, Text, ComboBox, IComboBoxOption, SelectableOptionMenuItemType, IContextualMenuProps, IContextualMenuItem, IconButton, IContextualMenuListProps, Stack, Checkbox, FontIcon, Modal, DefaultButton, ScrollablePane, SelectionZone, DetailsList, Selection, DetailsListLayoutMode, SelectionMode, IColumn, PrimaryButton, Label, TextField } from "@fluentui/react";
+import { mergeStyles, mergeStyleSets, Text, ComboBox, IComboBoxOption, SelectableOptionMenuItemType, IContextualMenuProps, IContextualMenuItem, IconButton, IContextualMenuListProps, Stack, Checkbox, FontIcon, Modal, DefaultButton, ScrollablePane, SelectionZone, DetailsList, Selection, DetailsListLayoutMode, SelectionMode, IColumn, PrimaryButton, Label, TextField, MessageBar, MessageBarType } from "@fluentui/react";
 import { MetadataRef, PhaseSessionResult, TestDataHandler, TestNameRef, TestPhaseStatus, TestSessionDetails, TestSessionResult } from "./testData";
 import dashboard, { StatusColor } from "horde/backend/Dashboard";
 import { memo, useCallback, useEffect, useMemo, useState } from "react";
@@ -9,6 +9,9 @@ import { PhaseHistoryGraph } from "./phaseHistoryGraph";
 import { projectStore } from "horde/backend/ProjectStore";
 import { getHordeStyling } from "horde/styles/Styles";
 
+import { GetUserResponse } from "horde/backend/Api";
+import { UserSelect } from "horde/components/UserSelect";
+
 import * as d3 from "d3";
 
 export type PhaseSessionCallback = (session: PhaseSessionResult) => void;
@@ -930,3 +933,58 @@ export const EditableTextField: React.FC<{ label: string, onChange: (value: stri
             </Stack>    
 }
 
+export const UserPickerModal: React.FC<{ defaultUser?: GetUserResponse, onSelect: (userId?: string) => void, onClose: () => void }> = ({ defaultUser, onSelect, onClose }) => {
+
+    const [state, setState] = useState<{ error?: string, userId?: string }>({});
+
+    const { hordeClasses } = getHordeStyling();
+
+    const onValidate = async () => {
+        const userId = state.userId;
+
+        if (userId && userId === defaultUser?.id) {
+            state.error = `Already selected this user`;
+            setState({ ...state });
+            return;
+        }
+
+        onSelect(userId);
+        onClose();
+    }
+
+    return <Modal isOpen={true} className={hordeClasses.modal} styles={{ main: { padding: 8, width: 540 } }} onDismiss={() => { onClose() }}>
+            <Stack horizontal styles={{ root: { padding: 8 } }}>
+                <Stack.Item grow={2}>
+                    <Text variant="mediumPlus">Select a user</Text>
+                </Stack.Item>
+                <Stack.Item grow={0}>
+                    <IconButton
+                        iconProps={{ iconName: 'Cancel' }}
+                        ariaLabel="Close popup modal"
+                        onClick={() => { onClose(); }}
+                    />
+                </Stack.Item>
+            </Stack>
+
+            {!!state.error && <MessageBar
+                messageBarType={MessageBarType.error}
+                isMultiline={false}> {state.error} </MessageBar>}
+
+            <Stack styles={{ root: { padding: 8 } }}>
+                <Stack grow>
+                    <Stack tokens={{ childrenGap: 4 }}>
+                    <UserSelect autoFocus handleSelection={(id => {
+                        setState({ ...state, userId: id })
+                    })} userIdHints={defaultUser? [defaultUser.id] : undefined} defaultUser={defaultUser} noResultsFoundText="No users found" />
+                    <Text variant="small" style={{ paddingTop: 4 }}>Leave blank to unassign</Text>
+                    </Stack>
+                </Stack>
+                
+                <Stack horizontal tokens={{ childrenGap: 16 }} styles={{ root: { paddingTop: 24, paddingLeft: 8, paddingBottom: 8 } }}>
+                    <Stack grow />
+                    <PrimaryButton text="Select" onClick={() => { onValidate(); }} />
+                    <DefaultButton text="Cancel" onClick={() => { onClose(); }} />
+                </Stack>
+            </Stack>
+        </Modal>;
+};
diff --git a/Engine/Source/Programs/Horde/HordeDashboard/plugins/test-automation/testAutomationHealth.tsx b/Engine/Source/Programs/Horde/HordeDashboard/plugins/test-automation/testAutomationHealth.tsx
--- a/Engine/Source/Programs/Horde/HordeDashboard/plugins/test-automation/testAutomationHealth.tsx
+++ b/Engine/Source/Programs/Horde/HordeDashboard/plugins/test-automation/testAutomationHealth.tsx
@@ -2,16 +2,17 @@
 
 import { DetailsList, FontIcon, IColumn, IconButton, Label, Modal, SelectionMode, Stack, Text, HoverCard, HoverCardType, IPlainCardProps, DirectionalHint, DetailsListLayoutMode, Spinner, SpinnerSize, Toggle, DatePicker, ComboBox, IComboBoxOption, FocusZone, FocusZoneTabbableElements, FocusTrapCallout, PrimaryButton, DefaultButton, TooltipHost, TooltipOverflowMode } from "@fluentui/react";
 import { MetadataRef, TestDataHandler, TestNameRef, TestSessionResult } from "./testData";
-import { useCallback, useEffect, useMemo, useState } from "react";
+import { useCallback, useEffect, useMemo, useRef, useState } from "react";
 import { observer } from "mobx-react-lite";
-import { TestOutcome } from "./api";
-import { DomainNameColors, EditableTextField, getTestUniqueMetaIdentifiers, StreamSelector, tagColors } from "./testAutomationCommon";
+import { TestOutcome, getUsersInfo } from "./api";
+import { DomainNameColors, EditableTextField, getTestUniqueMetaIdentifiers, StreamSelector, tagColors, UserPickerModal } from "./testAutomationCommon";
 import { TestHistoryGraphWidget } from "./testAutomationSummary";
 import { StoreKey, userStorage } from "./testAutomationUserStorage";
 import { projectStore } from "horde/backend/ProjectStore";
 import dashboard from "horde/backend/Dashboard";
 import { getShortNiceTime, msecToElapsed } from "horde/base/utilities/timeUtils";
 import { booleanColors, cadenceColors, cadenceOrders, healthColors, healthStars, TestCadence, TestHealth, TestHealthExplain, TestHealthItem, TestHealthTiming } from "./testHealthData";
+import { GetUserResponse } from "horde/backend/Api";
 
 /// Colors
 const profileColors = new DomainNameColors();
@@ -23,6 +24,11 @@ const harnessColors = new DomainNameColors();
 const healthBadgeBackground = () => dashboard.darktheme? "#67676748" : "#8f8f8f9e";
 const textEmphasisColor = () => dashboard.darktheme? "#ffffff75" : "#00000075";
 
+enum UserType {
+    Owner = "Owner",
+    Customers = "Customers",
+}
+
 const onRenderExplainCard = (explainData: TestHealthExplain): JSX.Element => {
     return <Stack style={{pointerEvents: 'none', padding: 10}} tokens={{childrenGap: 5}}>
                 {!!explainData.successRate && <Text>{explainData.successRate}% Success</Text>}
@@ -464,6 +470,30 @@ export const TestHealthView: React.FC<{ test: TestHealthItem, handler: TestDataH
             toCompare: userStorage.getItem(StoreKey.HistoryToCompare, () => undefined)
         });
     const [testHealthToCompare, setTestHealthToCompare] = useState<TestHealthItem | undefined>(getTestHealthItems(handler, userStorage.getItem(StoreKey.HistoryToCompare), test.nameRef).find(() => true));
+    const [userPicker, setUserPicker] = useState<{target?: UserType, defaultUser?: GetUserResponse}>({});
+    const userCache = useRef<Map<string | undefined, GetUserResponse> | undefined>(undefined);
+
+    useEffect(() => {
+        if (!userCache.current) {
+            userCache.current = new Map<string, GetUserResponse>();
+        }
+        let userIds: string[] = [];
+        if (test.owner) {
+            userIds.push(test.owner);
+        }
+        if (test.customers) {
+            userIds = userIds.concat(test.customers);
+        }
+        userIds = userIds.filter(id => !userCache.current!.has(id));
+        if (userIds.length > 0) {
+            getUsersInfo(userIds).then((users) => {
+                users.forEach(user => {
+                    userCache.current!.set(user.id, user);
+                });
+                test.setUpdated();
+            });
+        }
+    }, [test.owner, test.customers]);
 
     handler.subscribeToAuditsLoading();
     handler.subscribeToRecipesLoading();
@@ -507,176 +537,226 @@ export const TestHealthView: React.FC<{ test: TestHealthItem, handler: TestDataH
 
     const tags = useMemo(() => test.nameRef?.tagRefs?.keys().map(t => t.name).toArray().sort(), [test]);
 
-    return <Stack tokens={{childrenGap: 16}} style={{position: 'relative'}}>
-                <Stack tokens={{childrenGap: 8}} style={{padding: 14, borderRadius: 4}} >
-                    <Stack horizontal tokens={{childrenGap: 12}} verticalAlign="center">
-                        <ComboBox
-                            key={`profile-${test.updated}`}
-                            styles={{input: {backgroundColor: profileColors.get(test.profile)}, root: {backgroundColor: profileColors.get(test.profile), width: 150}}}
-                            placeholder="Profile"
-                            allowFreeform
-                            autoComplete="on"
-                            options={profileOptions}
-                            text={test.profile}
-                            onChange={(_1, _2, _3, value) => { test.profile = value; onUpdateItem(test); }}
-                        />
-                        <ComboBox
-                            key={`team-${test.updated}`}
-                            styles={{input: {backgroundColor: teamColors.get(test.team)}, root: {backgroundColor: teamColors.get(test.team), width: 150}}}
-                            placeholder="Team"
-                            allowFreeform
-                            autoComplete="on"
-                            options={teamOptions}
-                            text={test.team}
-                            onChange={(_1, _2, _3, value) => { test.team = value; onUpdateItem(test); }}
-                        />
-                        <ComboBox
-                            key={`harness-${test.updated}`}
-                            styles={{root: {width: 170}}}
-                            placeholder="Harness"
-                            allowFreeform
-                            autoComplete="on"
-                            options={harnessOptions}
-                            text={test.harness}
-                            onChange={(_1, _2, _3, value) => { test.harness = value; }}
-                        />
-                    </Stack>
-                    <Stack horizontal tokens={{childrenGap: 12}} style={{paddingLeft: 12, paddingTop: 6}} grow>
-                        <EditableTextField
-                            key={`intent-${test.updated}`}
-                            label="Intent"
-                            text={test.intent}
-                            onChange={(value) => test.intent = value}
-                            placeholder="a short description"
-                            loaded={handler.queryAuditsLoading}
-                        />
-                    </Stack>
-
-                    {tags &&
-                        <Stack horizontal verticalAlign="center" tokens={{childrenGap: 4}} wrap grow style={{paddingLeft: 12, paddingTop: 6}}>
-                            {tags.map(tag =>
-                                    <Stack style={{backgroundColor: tagColors.getTagColor(tag), borderRadius: 5, padding: '1px 4px', opacity: 0.8}}>
-                                        <Text>#{tag}</Text>
-                                    </Stack>
-                                )
+    return <Stack>
+                {userPicker.target &&
+                    <UserPickerModal
+                        defaultUser={userPicker.defaultUser}
+                        onSelect={(userId?: string) => {
+                            switch(userPicker.target) {
+                                case UserType.Owner :
+                                    test.owner = userId;
+                                    break;
+                                case UserType.Customers :
+                                    if (userId) {
+                                        const customers = new Set(test.customers ?? []);
+                                        customers.add(userId);
+                                        test.customers = customers.keys().toArray();
+                                    }
+                                    break;
                             }
+                        }}
+                        onClose={() => setUserPicker({})}
+                    />
+                }
+                <Stack tokens={{childrenGap: 16}} style={{position: 'relative'}}>
+                    <Stack tokens={{childrenGap: 8}} style={{padding: 14, borderRadius: 4}} >
+                        <Stack horizontal tokens={{childrenGap: 12}} verticalAlign="center">
+                            <ComboBox
+                                key={`profile-${test.updated}`}
+                                styles={{input: {backgroundColor: profileColors.get(test.profile)}, root: {backgroundColor: profileColors.get(test.profile), width: 150}}}
+                                placeholder="Profile"
+                                allowFreeform
+                                autoComplete="on"
+                                options={profileOptions}
+                                text={test.profile}
+                                onChange={(_1, _2, _3, value) => { test.profile = value; onUpdateItem(test); }}
+                            />
+                            <ComboBox
+                                key={`team-${test.updated}`}
+                                styles={{input: {backgroundColor: teamColors.get(test.team)}, root: {backgroundColor: teamColors.get(test.team), width: 150}}}
+                                placeholder="Team"
+                                allowFreeform
+                                autoComplete="on"
+                                options={teamOptions}
+                                text={test.team}
+                                onChange={(_1, _2, _3, value) => { test.team = value; onUpdateItem(test); }}
+                            />
+                            <ComboBox
+                                key={`harness-${test.updated}`}
+                                styles={{root: {width: 170}}}
+                                placeholder="Harness"
+                                allowFreeform
+                                autoComplete="on"
+                                options={harnessOptions}
+                                text={test.harness}
+                                onChange={(_1, _2, _3, value) => { test.harness = value; }}
+                            />
                         </Stack>
-                    }
-
-                    <Stack horizontal style={{paddingLeft: 12}}>
-                        <Stack horizontal verticalAlign="center" tokens={{childrenGap: 6}} style={{cursor: 'default'}} title="edit disabled for now">
-                            <Label>Worflow Notification: </Label>
-                            <Text style={{color:textEmphasisColor()}}>{!!test.notification? test.notification : "none"}</Text>
+                        <Stack horizontal tokens={{childrenGap: 12}} style={{paddingLeft: 12, paddingTop: 6}} grow>
+                            <EditableTextField
+                                key={`intent-${test.updated}`}
+                                label="Intent"
+                                text={test.intent}
+                                onChange={(value) => test.intent = value}
+                                placeholder="a short description"
+                                loaded={handler.queryAuditsLoading}
+                            />
                         </Stack>
-                    </Stack>
 
-                    <Stack horizontal tokens={{childrenGap: 12}} style={{paddingLeft: 12}} title="edit disabled for now">
-                        <Stack style={{borderRadius: 4, cursor: 'default'}} verticalAlign="center">
-                            <Stack  horizontal tokens={{childrenGap: 6}} verticalAlign="center">
-                                <Label>Owner: </Label>
-                            </Stack>
-                            <Text style={{color:textEmphasisColor()}}>{test.owner ?? "N/A"}</Text>
-                        </Stack>
-                        <Stack style={{borderRadius: 4, cursor: 'default'}} verticalAlign="center">
-                            <Stack  horizontal tokens={{childrenGap: 6}} verticalAlign="center">
-                                <Label>Customers: </Label>
+                        {tags &&
+                            <Stack horizontal verticalAlign="center" tokens={{childrenGap: 4}} wrap grow style={{paddingLeft: 12, paddingTop: 6}}>
+                                {tags.map(tag =>
+                                        <Stack style={{backgroundColor: tagColors.getTagColor(tag), borderRadius: 5, padding: '1px 4px', opacity: 0.8}}>
+                                            <Text>#{tag}</Text>
+                                        </Stack>
+                                    )
+                                }
                             </Stack>
-                            <Text style={{color:textEmphasisColor()}}>{test.customers?.join(', ')  ?? "N/A"}</Text>
-                        </Stack>
-                    </Stack>
+                        }
 
-                    <Stack style={{padding: '4px 12px', backgroundColor: booleanColors.get(test.underAudit)}}>
-                        <Stack horizontal tokens={{childrenGap: 6}} verticalAlign="center">
-                            <Toggle
-                                key={`under-audit-${test.updated}`}
-                                label="Under Audit"
-                                checked={test.underAudit}
-                                onChange={(_, checked) => {test.underAudit = !!checked}}
-                                inlineLabel
-                                onText="Yes" offText="No"
-                            />
-                            <Stack horizontal horizontalAlign="start" tokens={{childrenGap: 6}} grow>
-                                <Label>Last date: </Label>
-                                <DatePicker
-                                    key={`last-audit-${test.updated}`}
-                                    styles={{root: { width: 150 }}}
-                                    placeholder="Select a date... "
-                                    value={test.lastAuditDate}
-                                    onSelectDate={(date) => test.lastAuditDate = (date ?? undefined)}
-                                />
+                        <Stack horizontal style={{paddingLeft: 12}}>
+                            <Stack horizontal verticalAlign="center" tokens={{childrenGap: 6}} style={{cursor: 'default'}} title="edit disabled for now">
+                                <Label>Workflow Notification: </Label>
+                                <Text style={{color:textEmphasisColor()}}>{!!test.notification? test.notification : "none"}</Text>
                             </Stack>
                         </Stack>
-                        <EditableTextField
-                            key={`notes-${test.updated}`}
-                            label="Notes"
-                            text={test.notes}
-                            onChange={(value) => test.notes = value}
-                            multiline
-                            placeholder="add notes"
-                            loaded={handler.queryAuditsLoading}
-                        />
-                    </Stack>
 
-                </Stack>
-                <Stack style={{paddingLeft: 16}} tokens={{childrenGap: 16}}>
-                    <Stack style={{paddingLeft: 12}} tokens={{childrenGap: 16}}>
-                        <Stack horizontal tokens={{childrenGap: 12}} verticalAlign="center">
-                            <Stack style={{width: 200, backgroundColor: "#8e8e8e48", padding:  '3px 12px', borderRadius: 3}}>
-                                <Text>{projectStore.streamById(streams.current)?.fullname ?? streams.current}</Text>
+                        <Stack horizontal tokens={{childrenGap: 12}} style={{paddingLeft: 12}} title="edit disabled for now">
+                            <Stack style={{borderRadius: 4}} verticalAlign="center">
+                                <Stack  horizontal tokens={{childrenGap: 6}} verticalAlign="center">
+                                    <Label>Owner: </Label>
+                                </Stack>
+                                <Stack style={{backgroundColor: booleanColors.get(false), padding: '2px 9px', borderRadius: 11}} horizontalAlign="center" verticalAlign="center">
+                                    {!!userCache.current?.get(test.owner) &&
+                                        <Stack horizontal verticalAlign="center">
+                                            <Text style={{cursor: 'pointer', color:textEmphasisColor()}}
+                                                onClick={() => setUserPicker({target: UserType.Owner, defaultUser: userCache.current?.get(test.owner)})}
+                                            >{userCache.current?.get(test.owner)?.name ?? <Spinner size={SpinnerSize.xSmall} />}</Text>
+                                            <FontIcon iconName="Cancel" style={{cursor: 'pointer', fontSize: 11, paddingLeft: 6}} onClick={() => {test.owner = undefined}}/>
+                                        </Stack>
+                                    }
+                                    {!userCache.current?.get(test.owner) &&
+                                        <FontIcon iconName="AddTo" style={{cursor: 'pointer', fontSize: 15, padding: 2}}
+                                            onClick={() => setUserPicker({target: UserType.Owner})}
+                                        />
+                                    }
+                                </Stack>
                             </Stack>
-                            <Stack style={{cursor: 'default'}}>{renderHealthBadge(test, 75)}</Stack>
-                            <Stack style={{backgroundColor: cadenceColors.get(test.cadence), padding:  '3px 12px', borderRadius: 4, cursor: 'default'}} verticalAlign="center">
-                                <Text title="Cadence">{test.cadence}</Text>
+                            <Stack style={{borderRadius: 4}} verticalAlign="center">
+                                <Stack  horizontal tokens={{childrenGap: 6}} verticalAlign="center">
+                                    <Label>Customers: </Label>
+                                </Stack>
+                                <Stack horizontal tokens={{childrenGap: 6}}>
+                                    {!!test.customers &&
+                                        test.customers.map(user => {
+                                            return <Stack style={{backgroundColor: booleanColors.get(false), padding: '2px 9px', borderRadius: 11}} verticalAlign="center" horizontal>
+                                                        <Text style={{color:textEmphasisColor()}}>{userCache.current?.get(user)?.name ?? <Spinner size={SpinnerSize.xSmall} />}</Text>
+                                                        <FontIcon iconName="Cancel" style={{cursor: 'pointer', fontSize: 11, paddingLeft: 6}} onClick={() => {test.customers = test.customers?.filter(c => c !== user)}}/>
+                                                    </Stack>
+                                        })
+                                    }
+                                    <Stack style={{backgroundColor: booleanColors.get(false), padding: '2px 9px', borderRadius: 11, width: 46}} verticalAlign="center" horizontalAlign="center">
+                                        <FontIcon iconName="AddTo" style={{cursor: 'pointer', fontSize: 15, padding: 2}}
+                                            onClick={() => setUserPicker({target: UserType.Customers})}
+                                        />
+                                    </Stack>
+                                </Stack>
                             </Stack>
-                            {test.timings && <TestHealthTimes test={test} />}
                         </Stack>
-                        <Stack horizontal tokens={{childrenGap: 16}}>
-                            <Stack style={{backgroundColor: booleanColors.get(false), width: 6, borderRadius: 3}}></Stack>
-                            <TestHistoryGraphWidget test={test.nameRef} stream={streams.current} commonMetaKeys={commonMetaKeys} handler={handler} onClick={onDetailsClick} />
+
+                        <Stack style={{padding: '4px 12px', backgroundColor: booleanColors.get(test.underAudit)}}>
+                            <Stack horizontal tokens={{childrenGap: 6}} verticalAlign="center">
+                                <Toggle
+                                    key={`under-audit-${test.updated}`}
+                                    label="Under Audit"
+                                    checked={test.underAudit}
+                                    onChange={(_, checked) => {test.underAudit = !!checked}}
+                                    inlineLabel
+                                    onText="Yes" offText="No"
+                                />
+                                <Stack horizontal horizontalAlign="start" tokens={{childrenGap: 6}} grow>
+                                    <Label>Last date: </Label>
+                                    <DatePicker
+                                        key={`last-audit-${test.updated}`}
+                                        styles={{root: { width: 150 }}}
+                                        placeholder="Select a date... "
+                                        value={test.lastAuditDate}
+                                        onSelectDate={(date) => test.lastAuditDate = (date ?? undefined)}
+                                    />
+                                </Stack>
+                            </Stack>
+                            <EditableTextField
+                                key={`notes-${test.updated}`}
+                                label="Notes"
+                                text={test.notes}
+                                onChange={(value) => test.notes = value}
+                                multiline
+                                placeholder="add notes"
+                                loaded={handler.queryAuditsLoading}
+                            />
                         </Stack>
+
                     </Stack>
-                    <Stack style={{paddingLeft: 12}} tokens={{childrenGap: 16}}>
-                        <Stack horizontal tokens={{childrenGap: 12}} verticalAlign="center">
-                            <Stack style={{width: 200}}>
-                                <StreamSelector streams={availableStreams} selected={streams.toCompare} disabled={availableStreams.length === 0} onClick={onClickStream} />
-                            </Stack>
-                            {!!testHealthToCompare && <Stack style={{cursor: 'default'}}>{renderHealthBadge(testHealthToCompare, 75)}</Stack>}
-                            {!!testHealthToCompare && 
-                                <Stack style={{backgroundColor: cadenceColors.get(testHealthToCompare.cadence), padding:  '3px 12px', borderRadius: 4, cursor: 'default'}} verticalAlign="center">
-                                    <Text title="Cadence">{testHealthToCompare.cadence}</Text>
+                    <Stack style={{paddingLeft: 16}} tokens={{childrenGap: 16}}>
+                        <Stack style={{paddingLeft: 12}} tokens={{childrenGap: 16}}>
+                            <Stack horizontal tokens={{childrenGap: 12}} verticalAlign="center">
+                                <Stack style={{width: 200, backgroundColor: "#8e8e8e48", padding:  '3px 12px', borderRadius: 3}}>
+                                    <Text>{projectStore.streamById(streams.current)?.fullname ?? streams.current}</Text>
                                 </Stack>
-                            }
-                            {!!testHealthToCompare && testHealthToCompare.timings && <TestHealthTimes test={testHealthToCompare} />}
-                        </Stack>
-                        <Stack horizontal tokens={{childrenGap: 16}}>
-                            <Stack style={{backgroundColor: booleanColors.get(false), width: 6, borderRadius: 3}}></Stack>
-                            {!streams.toCompare && !handler.subQueryLoading &&
-                                <Stack style={{minHeight: 40}} verticalAlign="center"><Text>select a stream</Text></Stack>
-                            }
-                            {!!streams.toCompare && !!testHealthToCompare && 
-                                <TestHistoryGraphWidget test={test.nameRef} stream={streams.toCompare} commonMetaKeys={commonMetaKeys} handler={handler} onClick={onDetailsClick} />
-                            }
-                            {!!streams.toCompare && !testHealthToCompare && 
-                                <Stack horizontalAlign='center' style={{ padding: 4 }}>
-                                    <Text style={{ fontSize: 12 }}>No Results</Text>
+                                <Stack style={{cursor: 'default'}}>{renderHealthBadge(test, 75)}</Stack>
+                                <Stack style={{backgroundColor: cadenceColors.get(test.cadence), padding:  '3px 12px', borderRadius: 4, cursor: 'default'}} verticalAlign="center">
+                                    <Text title="Cadence">{test.cadence}</Text>
                                 </Stack>
-                            }
-                            {handler.subQueryLoading &&
-                                <Stack horizontalAlign='center' style={{ padding: 12}} tokens={{ childrenGap: 10 }} horizontal>
-                                    <Text style={{ fontSize: 12 }}>Loading Data</Text>
-                                    <Spinner size={SpinnerSize.small} />
+                                {test.timings && <TestHealthTimes test={test} />}
+                            </Stack>
+                            <Stack horizontal tokens={{childrenGap: 16}}>
+                                <Stack style={{backgroundColor: booleanColors.get(false), width: 6, borderRadius: 3}}></Stack>
+                                <TestHistoryGraphWidget test={test.nameRef} stream={streams.current} commonMetaKeys={commonMetaKeys} handler={handler} onClick={onDetailsClick} />
+                            </Stack>
+                        </Stack>
+                        <Stack style={{paddingLeft: 12}} tokens={{childrenGap: 16}}>
+                            <Stack horizontal tokens={{childrenGap: 12}} verticalAlign="center">
+                                <Stack style={{width: 200}}>
+                                    <StreamSelector streams={availableStreams} selected={streams.toCompare} disabled={availableStreams.length === 0} onClick={onClickStream} />
                                 </Stack>
-                            }
+                                {!!testHealthToCompare && <Stack style={{cursor: 'default'}}>{renderHealthBadge(testHealthToCompare, 75)}</Stack>}
+                                {!!testHealthToCompare && 
+                                    <Stack style={{backgroundColor: cadenceColors.get(testHealthToCompare.cadence), padding:  '3px 12px', borderRadius: 4, cursor: 'default'}} verticalAlign="center">
+                                        <Text title="Cadence">{testHealthToCompare.cadence}</Text>
+                                    </Stack>
+                                }
+                                {!!testHealthToCompare && testHealthToCompare.timings && <TestHealthTimes test={testHealthToCompare} />}
+                            </Stack>
+                            <Stack horizontal tokens={{childrenGap: 16}}>
+                                <Stack style={{backgroundColor: booleanColors.get(false), width: 6, borderRadius: 3}}></Stack>
+                                {!streams.toCompare && !handler.subQueryLoading &&
+                                    <Stack style={{minHeight: 40}} verticalAlign="center"><Text>select a stream</Text></Stack>
+                                }
+                                {!!streams.toCompare && !!testHealthToCompare && 
+                                    <TestHistoryGraphWidget test={test.nameRef} stream={streams.toCompare} commonMetaKeys={commonMetaKeys} handler={handler} onClick={onDetailsClick} />
+                                }
+                                {!!streams.toCompare && !testHealthToCompare && 
+                                    <Stack horizontalAlign='center' style={{ padding: 4 }}>
+                                        <Text style={{ fontSize: 12 }}>No Results</Text>
+                                    </Stack>
+                                }
+                                {handler.subQueryLoading &&
+                                    <Stack horizontalAlign='center' style={{ padding: 12}} tokens={{ childrenGap: 10 }} horizontal>
+                                        <Text style={{ fontSize: 12 }}>Loading Data</Text>
+                                        <Spinner size={SpinnerSize.small} />
+                                    </Stack>
+                                }
+                            </Stack>
                         </Stack>
                     </Stack>
-                </Stack>
 
-                { (handler.queryAuditsLoading || handler.queryRecipesLoading) &&
-                    <Stack horizontalAlign='center' style={{ width: "100%", position: "absolute", top: 0 }} tokens={{childrenGap: 8}}>
-                        <Spinner size={SpinnerSize.large} />
-                    </Stack>
-                }
-            </Stack>
+                    { (handler.queryAuditsLoading || handler.queryRecipesLoading) &&
+                        <Stack horizontalAlign='center' style={{ width: "100%", position: "absolute", top: 0 }} tokens={{childrenGap: 8}}>
+                            <Spinner size={SpinnerSize.large} />
+                        </Stack>
+                    }
+                </Stack>
+        </Stack>
 });
 
 const TestHealthDetailsModal: React.FC<{ test: TestHealthItem, handler: TestDataHandler, onDismiss: () => void, onUpdateItem: (item: TestHealthItem) => void }> = observer(({ test, handler, onDismiss, onUpdateItem }) => {
@@ -684,8 +764,11 @@ const TestHealthDetailsModal: React.FC<{ test: TestHealthItem, handler: TestData
     const [isCommitting, setCommitting] = useState(false);
 
     useEffect(() => {
-        // discard changes on clean up
-        return () => { if(test.isTransient) { test.discardTransaction(); onUpdateItem(test); } }
+        if (test.loadStashedTransaction()) {
+            onUpdateItem(test);
+        }
+        // stash changes on clean up
+        return () => { if(test.isTransient) { test.stashTransaction(); test.discardTransaction(); onUpdateItem(test); } }
     }, [test]);
 
     test.subscribeUpdate();
@@ -695,25 +778,30 @@ const TestHealthDetailsModal: React.FC<{ test: TestHealthItem, handler: TestData
                 <Stack horizontal verticalAlign="center">
                     <Stack><Text style={{ paddingLeft: 8, fontSize: 16, fontWeight: 600 }}>Health Details - {`${test.name}`}</Text></Stack>
                     <Stack grow />
-                    {test.isTransient &&
-                        <Stack style={{ paddingBottom: 4, paddingRight: 6 }} title="Revert changes">
+                    {!isCommitting && test.isTransient &&
+                        <Stack style={{ paddingRight: 6 }} title="Revert changes">
                             <IconButton
-                                iconProps={{ iconName: 'History', style: { fontSize: 12 } }}
+                                iconProps={{ iconName: 'History', style: { fontSize: 14 } }}
                                 ariaLabel="Revert changes"
                                 onClick={() => { test.discardTransaction(); onUpdateItem(test); }}
                             />
                         </Stack>
                     }
                     {test.isTransient &&
-                        <Stack style={{ paddingBottom: 4, paddingRight: 6 }} title="Commit changes and close">
-                            <IconButton
-                                iconProps={{ iconName: 'Save', style: { fontSize: 16 } }}
-                                ariaLabel="Commit and close"
-                                onClick={() => { test.commitTransaction(); onDismiss(); }}
-                            />
+                        <Stack style={{ paddingRight: 6 }} title="Commit changes and close">
+                            {!isCommitting &&
+                                <IconButton
+                                    iconProps={{ iconName: 'Save', style: { fontSize: 16 } }}
+                                    ariaLabel="Commit and close"
+                                    onClick={() => { setCommitting(true); test.commitTransaction().then(() => { setCommitting(false) }); }}
+                                />
+                            }
+                            {isCommitting &&
+                                <Spinner style={{width: 32}} size={SpinnerSize.small} />
+                            }
                         </Stack>
                     }
-                    <Stack style={{ paddingBottom: 4 }} className="Close">
+                    <Stack className="Close">
                         <IconButton
                             iconProps={{ iconName: 'Cancel' }}
                             ariaLabel="Close popup modal"
@@ -731,25 +819,26 @@ const TestHealthDetailsModal: React.FC<{ test: TestHealthItem, handler: TestData
                                         <Text block>You have uncommitted changes to this audit.</Text>
                                         <Text block>Do you want to commit them?</Text>
                                         <FocusZone handleTabKey={FocusZoneTabbableElements.all} isCircularNavigation>
-                                            <Stack horizontal horizontalAlign="center">
-                                            <PrimaryButton 
-                                                iconProps={{ iconName: 'Save', style: { fontSize: 16 } }}
-                                                onClick={() => { setCommitting(true); test.commitTransaction().then(() => { setCommitting(false); setCalloutVisible(false); onDismiss(); }) }}
-                                            >Commit</PrimaryButton>
-                                            <DefaultButton
-                                                iconProps={{ iconName: 'History', style: { fontSize: 12 } }}
-                                                onClick={() => { setCalloutVisible(false); onDismiss(); }}
-                                            >Discard</DefaultButton>
-                                            <DefaultButton
-                                                onClick={() => { setCalloutVisible(false); }}
-                                            >Cancel</DefaultButton>
+                                            <Stack horizontal horizontalAlign="center" verticalAlign="center">
+                                                <PrimaryButton 
+                                                    iconProps={{ iconName: 'Save', style: { fontSize: 16 } }}
+                                                    onClick={() => { setCommitting(true); test.commitTransaction().then(() => { setCommitting(false); setCalloutVisible(false); onDismiss(); }) }}
+                                                >Commit</PrimaryButton>
+                                                <DefaultButton
+                                                    iconProps={{ iconName: 'History', style: { fontSize: 14 } }}
+                                                    onClick={() => { test.discardTransaction(); onUpdateItem(test); setCalloutVisible(false); onDismiss(); }}
+                                                >Discard</DefaultButton>
+                                                <DefaultButton
+                                                    onClick={() => { setCalloutVisible(false); }}
+                                                >Cancel</DefaultButton>
                                             </Stack>
                                         </FocusZone>
                                     </Stack>
                                 }
                                 {isCommitting &&
-                                    <Stack>
+                                    <Stack horizontal>
                                         <Text block>Committing changes...</Text>
+                                        <Spinner style={{width: 32}} size={SpinnerSize.small} />
                                     </Stack>
                                 }
                             </FocusTrapCallout>
diff --git a/Engine/Source/Programs/Horde/HordeDashboard/plugins/test-automation/testAutomationUserStorage.ts b/Engine/Source/Programs/Horde/HordeDashboard/plugins/test-automation/testAutomationUserStorage.ts
--- a/Engine/Source/Programs/Horde/HordeDashboard/plugins/test-automation/testAutomationUserStorage.ts
+++ b/Engine/Source/Programs/Horde/HordeDashboard/plugins/test-automation/testAutomationUserStorage.ts
@@ -24,6 +24,7 @@ export const userStorage = new TransientStorage();
 
 export enum StoreKey {
     HistoryToCompare = "historyToCompare",
-    SummaryExpandedCards = "summaryExpandedCards"
+    SummaryExpandedCards = "summaryExpandedCards",
+    HealthDetailsTransaction = "healthDetailsTransaction"
 }
 
diff --git a/Engine/Source/Programs/Horde/HordeDashboard/plugins/test-automation/testHealthData.ts b/Engine/Source/Programs/Horde/HordeDashboard/plugins/test-automation/testHealthData.ts
--- a/Engine/Source/Programs/Horde/HordeDashboard/plugins/test-automation/testHealthData.ts
+++ b/Engine/Source/Programs/Horde/HordeDashboard/plugins/test-automation/testHealthData.ts
@@ -2,6 +2,7 @@
 
 import { action, makeObservable, observable } from "mobx";
 import { MetadataRef, TestAudit, TestDataHandler, TestNameRef, TestRecipeRef } from "./testData";
+import { StoreKey, userStorage } from "./testAutomationUserStorage";
 
 export enum TestHealth {
     Healthy = "Healthy",
@@ -172,6 +173,22 @@ export class TestHealthItem {
         }
     }
 
+    stashTransaction() {
+        if (this.isTransient) {
+            userStorage.setItem(StoreKey.HealthDetailsTransaction, this.transientCopy);
+        }
+    }
+
+    loadStashedTransaction() {
+        const stash = userStorage.getItem(StoreKey.HealthDetailsTransaction) as TestAudit | undefined;
+        if (stash && stash.testId === this.nameRef.id) {
+            this.transientCopy = stash;
+            userStorage.setItem(StoreKey.HealthDetailsTransaction, undefined);
+            return true;
+        }
+        return false;
+    }
+
     get name() {
         return this.nameRef.name;
     }
@@ -246,10 +263,21 @@ export class TestHealthItem {
         return this.audit?.owner;
     }
 
+    set owner(value: string | undefined) {
+        this.getOrStartTransaction().owner = value;
+        this.setUpdated()
+    }
+
     get customers() {
         return this.audit?.customers;
     }
 
+    set customers(value: string[] | undefined) {
+        if (value && value.length === 0) value = undefined;
+        this.getOrStartTransaction().customers = value;
+        this.setUpdated()
+    }
+
     get notification() {
         return this.audit?.notifications?.get(this.stream);
     }
