[Horde][BuildHealth][Dashboard] JobStepRef parity changes.

#Synced-CL 48176906

diff --git a/Engine/Source/Programs/Horde/HordeDashboard/src/components/buildhealth/stepoutcome/StepOutcomeDataHandler.ts b/Engine/Source/Programs/Horde/HordeDashboard/src/components/buildhealth/stepoutcome/StepOutcomeDataHandler.ts
--- a/Engine/Source/Programs/Horde/HordeDashboard/src/components/buildhealth/stepoutcome/StepOutcomeDataHandler.ts
+++ b/Engine/Source/Programs/Horde/HordeDashboard/src/components/buildhealth/stepoutcome/StepOutcomeDataHandler.ts
@@ -1,12 +1,12 @@
 // Copyright Epic Games, Inc. All Rights Reserved.
 
-import { GetJobResponse, GetJobStepRefResponse, IssueQuery, JobStepState, JobStreamQuery, LabelState } from "horde/backend/Api";
+import { GetJobResponse, GetJobStepRefResponse, IssueQuery, JobStepOutcome, JobStepState, JobStreamQuery, LabelState } from "horde/backend/Api";
 import { PollBase } from "horde/backend/PollBase";
 import { GetJobStepRefResponseWrapper, GetStepResponseWrapper, JobResponseData, LabelOutcomeTableEntry, StepOutcomeTable, StepOutcomeTableEntry, UpdatedChangeSummaryData } from "./StepOutcomeDataTypes";
 import backend from "horde/backend";
-import { decodeLabelKey, decodeStepKey, decodeTemplateKey, encodeLabelKeyFromStrings, encodeStepNameFromStrings, encodeStepNameFromStringsPreserveCase } from "../BuildHealthUtilities";
+import { decodeLabelKey, decodeStepKey, decodeTemplateKey, encodeLabelKeyFromStrings, encodeStepNameFromStrings, encodeStepNameFromStringsPreserveCase, encodeTemplateKeyFromStrings } from "../BuildHealthUtilities";
 import { action, computed, makeObservable, observable, runInAction } from "mobx";
-import { summarizeJob } from "./StepOutcomeUtilities";
+import { summarizeJob, summarizeJobFromJobStepRef } from "./StepOutcomeUtilities";
 
 // #region -- API Data Types & Utilities --
 
@@ -111,6 +111,7 @@ export interface StepOutcomeFilters {
 export class StepOutcomeDataHandler extends PollBase {
     private readonly STREAM_JOBS_FILTER = "id,streamId,name,change,templateId,state,createTime,updateTime,startedByUserInfo,abortedByUserInfo,preflightCommitId,preflightDescription,cancellationReason,labels,defaultLabel";
     private readonly STREAM_JOBS_BATCH_STEP_FILTER = "batches.error,batches.steps.id,batches.steps.name,batches.steps.state,batches.steps.outcome,batches.steps.startTime,batches.steps.finishTime,batches.steps.warning,batches.steps.error"
+    private readonly STREAM_JOBS_BATCH_ERROR_FILTER = "batches.error"
     private readonly STREAM_CHANGES_FILTER = "id,dateUtc";
     static readonly MAX_JOB_COUNT_DEFAULT = 400;
 
@@ -147,7 +148,7 @@ export class StepOutcomeDataHandler extends PollBase {
     // #region -- Private API --
 
     private constructStreamJobQueries(templateId: string, jobCount: number, modifiedAfter: string, queries: { streamId: string; query: JobStreamQuery; }[], streamId: string, debugMode: boolean) {
-        let completeFilter: string = !debugMode ? `${this.STREAM_JOBS_FILTER},${this.STREAM_JOBS_BATCH_STEP_FILTER}` : `${this.STREAM_JOBS_FILTER}`
+        let completeFilter: string = !debugMode ? `${this.STREAM_JOBS_FILTER},${this.STREAM_JOBS_BATCH_STEP_FILTER}` : `${this.STREAM_JOBS_FILTER},${this.STREAM_JOBS_BATCH_ERROR_FILTER}`
         const query: JobStreamQuery = {
             template: [templateId],
             count: jobCount,
@@ -158,7 +159,7 @@ export class StepOutcomeDataHandler extends PollBase {
         queries.push({ streamId, query });
     }
 
-    private constructJobStepHistoryQueries(stepMap: Map<string, Map<string, Set<string>>>, streamId: string, templateId: string, jobStepHistoryBatchQueries: { streamId: string; stepNames: string[]; count: number; templateId: string; }[]) {
+    private constructJobStepHistoryQueries(stepMap: Map<string, Map<string, Set<string>>>, streamId: string, templateId: string, jobStepHistoryBatchQueries: { streamId: string; stepNames: string[]; templateId: string; }[]) {
         let templateMap: Map<string, Set<string>> = stepMap.get(streamId)!;
         let stepSet = templateMap?.get(templateId);
         if (stepSet) {
@@ -166,7 +167,7 @@ export class StepOutcomeDataHandler extends PollBase {
                 streamId: streamId, stepNames: Array.from(stepSet).map(x => {
                     let { streamId, templateId, stepName } = decodeStepKey(x);
                     return stepName;
-                }), count: 100, templateId: templateId
+                }), templateId: templateId
             });
         }
     }
@@ -244,15 +245,24 @@ export class StepOutcomeDataHandler extends PollBase {
     }
 
     private processJobStepHistoryResults(jobStepHistoryResults: GetJobStepRefResponse[][]) {
+        let returnMap: Map<string, GetJobStepRefResponse[]> = new Map<string, GetJobStepRefResponse[]>();
         for (const item of jobStepHistoryResults) {
             for (const subitem of item) {
                 if (this.jobData.has(subitem.jobId)) {
                     let jobName: string = this.jobData.get(subitem.jobId)?.name!;
                     let streamId: string = this.jobData.get(subitem.jobId)?.streamId!;
                     this.stepOutcomeTableData.addEntry(new StepOutcomeTableEntry(new GetJobStepRefResponseWrapper(subitem), subitem.jobId, jobName, streamId, subitem.jobStartTime.toString(), subitem.change, this.changeData.get(subitem.change!)?.dateUtc));
+                    if (!returnMap.has(subitem.jobId)) {
+                        returnMap.set(subitem.jobId, []);
+                    }
+                    returnMap.get(subitem.jobId)?.push(subitem);
                 }
             }
         }
+
+        for (const item of returnMap) {
+            this.jobData.get(item[0])!.summarizedOutcome = summarizeJobFromJobStepRef(item[1], this.jobData.get(item[0])!);
+        }
     }
 
     // #endregion -- Private API --
@@ -332,7 +342,7 @@ export class StepOutcomeDataHandler extends PollBase {
 
         const changeQueries: { streamId: string; query: any }[] = [];
         const streamJobQueries: { streamId: string; query: JobStreamQuery }[] = [];
-        const jobStepHistoryBatchQueries: { streamId: string; stepNames: string[], count: number, templateId: string }[] = [];
+        const jobStepHistoryBatchQueries: { streamId: string; stepNames: string[], templateId: string }[] = [];
 
         // Construct queries
         for (const streamId of streamList) {
@@ -344,45 +354,59 @@ export class StepOutcomeDataHandler extends PollBase {
             }
         }
 
-        const jobStepHistoryResults = this.activeFilter.debugMode ? await this.issueAllJobStepHistoryQueries(jobStepHistoryBatchQueries) : undefined;
         const streamJobResults = await this.issueAllStreamJobQueries(streamJobQueries);
 
         // Annotate the changes that we need to request in getChangeSummaries.
-        let jobHistoryCount: Map<string, number[]> = new Map<string, number[]>();
+        let streamToChangeRequestList: Map<string, number[]> = new Map<string, number[]>();
+        let streamJobToRequestSize: Map<string, number> = new Map<string, number>();
 
         streamJobResults.forEach((jobs, i) => {
             const streamId = streamJobQueries[i].streamId;
+            const newChanges: number[] = [];
+
+            for (let idx: number = 0; idx < jobs.length; ++idx) {
+                let change: number = jobs[idx].change!;
+                let fullyQualifiedId = encodeTemplateKeyFromStrings(streamId, jobs[idx].templateId!);
+
+                if (!this.changeData.has(jobs[idx].change!)) {
+                    newChanges.push(change);
+                }
 
-            const newChanges = jobs
-                .map(job => job.change!)
-                .filter(change => !this.changeData.has(change));
+                if (!streamJobToRequestSize.has(fullyQualifiedId)) {
+                    streamJobToRequestSize.set(fullyQualifiedId, 0);
+                }
+
+                streamJobToRequestSize.set(fullyQualifiedId, streamJobToRequestSize.get(fullyQualifiedId)! + 1);
+            }
 
             if (newChanges.length > 0) {
-                if (!jobHistoryCount.has(streamId)) {
-                    jobHistoryCount.set(streamId, []);
+                if (!streamToChangeRequestList.has(streamId)) {
+                    streamToChangeRequestList.set(streamId, []);
                 }
 
-                jobHistoryCount.get(streamId)!.push(...newChanges);
+                streamToChangeRequestList.get(streamId)!.push(...newChanges);
             }
         });
 
         // Flatten all jobData & memoize the summarized outcome
+        const jobStepHistoryResults = this.activeFilter.debugMode ? await this.issueAllJobStepHistoryQueries(jobStepHistoryBatchQueries.map(x => ({ ...x, count: streamJobToRequestSize.get(encodeTemplateKeyFromStrings(x.streamId, x.templateId))! }))) : undefined;
+
         this.jobData = new Map<string, JobResponseData>(streamJobResults.flat().map((job: GetJobResponse): [string, JobResponseData] => {
             return [
                 job.id,
                 {
                     ...job,
-                    summarizedOutcome: summarizeJob(job),
-                    isManualJobInvocation : job.startedByUserInfo !== null
+                    summarizedOutcome: this.activeFilter.debugMode ? JobStepOutcome.Unspecified : summarizeJob(job),
+                    isManualJobInvocation: job.startedByUserInfo !== null
                 }
             ]
         }
         ));
 
-        if (jobHistoryCount.size > 0) {
+        if (streamToChangeRequestList.size > 0) {
             // Parallelize backend calls for change summaries
             const changeResults = await Promise.all(
-                changeQueries.map(({ streamId, query }) => backend.getBatchChangeSummaries(streamId, [...new Set(jobHistoryCount.get(streamId))], query.filter))
+                changeQueries.map(({ streamId, query }) => backend.getBatchChangeSummaries(streamId, [...new Set(streamToChangeRequestList.get(streamId))], query.filter))
             );
 
             let flattenedResults = changeResults.flat();
diff --git a/Engine/Source/Programs/Horde/HordeDashboard/src/components/buildhealth/stepoutcome/StepOutcomeUtilities.ts b/Engine/Source/Programs/Horde/HordeDashboard/src/components/buildhealth/stepoutcome/StepOutcomeUtilities.ts
--- a/Engine/Source/Programs/Horde/HordeDashboard/src/components/buildhealth/stepoutcome/StepOutcomeUtilities.ts
+++ b/Engine/Source/Programs/Horde/HordeDashboard/src/components/buildhealth/stepoutcome/StepOutcomeUtilities.ts
@@ -1,6 +1,6 @@
 // Copyright Epic Games, Inc. All Rights Reserved.
 
-import { GetJobResponse, JobStepBatchError, JobStepOutcome, JobStepState } from "horde/backend/Api";
+import { GetJobResponse, GetJobStepRefResponse, JobStepBatchError, JobStepOutcome, JobStepState } from "horde/backend/Api";
 import { OutcomeTableEntry, ReadyStr, RunningStr, WaitingStr } from "./StepOutcomeDataTypes";
 
 export const DURATION_NOT_SET: number = -1;
@@ -53,6 +53,43 @@ export function summarizeJob(job: GetJobResponse): JobStepOutcome {
     return accumulatorOutcome;
 }
 
+/**
+  * This utility method will summarize the provided job, with respect to the most severe step outcome encountered thus far.
+  * @param jobStepRefs The set of job step refs to use in summarization.  
+  * @param job The job to summarize.
+  * @returns The most severe job step outcome. This is ordered as follows: @see JobStepOutcome.Unspecified , @see JobStepOutcome.Success , @see JobStepOutcome.Warnings , @see JobStepOutcome.Failre .
+  */
+export function summarizeJobFromJobStepRef(jobStepRefs: GetJobStepRefResponse[], job: GetJobResponse): JobStepOutcome {
+    let accumulatorOutcome: JobStepOutcome = JobStepOutcome.Unspecified;
+    accumulatorOutcome = JobStepOutcome.Success;
+
+    for (let response of jobStepRefs) {
+        if (response.state !== JobStepState.Aborted && response.outcome === JobStepOutcome.Failure) {
+            accumulatorOutcome = JobStepOutcome.Failure;
+
+            // We cannot become less severe than Failure; exit.
+            return accumulatorOutcome;
+        } else if (response.outcome === JobStepOutcome.Warnings) {
+            accumulatorOutcome = JobStepOutcome.Warnings;
+        }
+    }
+
+    // Check to see if an entire batch has failed, in spite of the steps declaring a success (due to the absence of steps).
+    // This is a common case when a more terminal failure occurs in batch resolution.
+    if (accumulatorOutcome === JobStepOutcome.Success && job.batches) {
+        // do a sanity check for batchEdge case
+        for (let response of job.batches) {
+
+            if (response.error !== JobStepBatchError.None) {
+                accumulatorOutcome = JobStepOutcome.Failure;
+                return accumulatorOutcome;
+            }
+        }
+    }
+
+    return accumulatorOutcome;
+}
+
 /**
  * Produces a ISO compliant string, in the provided locale.
  * @param date The date to convert.
