[Horde][BuildHealth][Dashboard] Add RCA owner input to dashbaord
#Synced-CL Unknown

diff --git a/Engine/Source/Programs/Horde/HordeDashboard/src/backend/Api.ts b/Engine/Source/Programs/Horde/HordeDashboard/src/backend/Api.ts
--- a/Engine/Source/Programs/Horde/HordeDashboard/src/backend/Api.ts
+++ b/Engine/Source/Programs/Horde/HordeDashboard/src/backend/Api.ts
@@ -3235,6 +3235,9 @@ export type GetIssueResponse = {
 	/**Root cause commit id */
 	rootCommitId?: CommitId;
 
+	/**Root cause owner id*/
+	rootCauseOwnerId?: string;
+
 	/**Root cause category*/
 	rootCauseCategory?: string;
 
@@ -3307,6 +3310,9 @@ export type UpdateIssueRequest = {
 	/**Root cause commit id */
 	rootCommitId?: CommitId | null;
 
+	/**Root cause owner id*/
+	rootCauseOwnerId?: string | null;
+
 	/**Root cause category*/
 	rootCauseCategory?: string;
 
@@ -4575,6 +4581,9 @@ export type FindIssueResponse = {
 	/**Root cause commit id */
 	rootCommitId?: CommitId;
 
+	/**Root cause owner id*/
+	rootCauseOwner?: GetThinUserInfoResponse;
+
 	/**Root cause category*/
 	rootCauseCategory?: string;
 
diff --git a/Engine/Source/Programs/Horde/HordeDashboard/src/components/IssueViewV2.tsx b/Engine/Source/Programs/Horde/HordeDashboard/src/components/IssueViewV2.tsx
--- a/Engine/Source/Programs/Horde/HordeDashboard/src/components/IssueViewV2.tsx
+++ b/Engine/Source/Programs/Horde/HordeDashboard/src/components/IssueViewV2.tsx
@@ -1358,6 +1358,12 @@ const IssueSummaryPanel: React.FC = () => {
          text: `${issue.rootCommitId}`,
          link: dashboard.swarmUrl ? `${dashboard.swarmUrl}/changes/${issue.rootCommitId}` : undefined
       })
+   }
+     if (issue.rootCauseOwnerId) {
+      items.push({
+         title: "Root Cause Owner",
+         text: issue.rootCauseOwnerId
+      })
    }
    if (issue.rootCauseCategory) {
       items.push({
@@ -1826,7 +1832,7 @@ const IssueCommandBar: React.FC = () => {
          text: `Modify Issue`,
          style: { color: "unset" },
          onClick: () => { setMarkFixedShown({shown: true}); }
-      },      
+      },
       {
          key: 'more_history_link',
          iconProps: { iconName: "SearchTemplate" },
@@ -2342,11 +2348,12 @@ const MarkFixedModal: React.FC<{ onClose: () => void, autoFix?:boolean } > = ({
 
    const issue = details.issue!;
 
-   const [state, setState] = useState<{ error?: string, fixCL?: number, rootCommitId?: CommitId, rootCauseSummary?: string, rootCauseCategory?: string, submitting?: boolean, dupeCheck?: boolean, dupeId?: string }>({
+   const [state, setState] = useState<{ error?: string, fixCL?: number, rootCommitId?: CommitId, rootCauseOwnerId?: string, rootCauseSummary?: string, rootCauseCategory?: string, submitting?: boolean, dupeCheck?: boolean, dupeId?: string }>({
       fixCL: issue.fixChange ?? (autoFix ? 0 : undefined),
       rootCommitId: issue.rootCommitId,
       dupeId: issue.duplicateIssueId ? (issue.duplicateIssueId.toString() !== "0" ? issue.duplicateIssueId.toString() : undefined) : undefined,
       dupeCheck: issue.duplicateIssueId === 0 || !!issue.duplicateIssueId,
+      rootCauseOwnerId: issue.rootCauseOwnerId,
       rootCauseCategory: issue.rootCauseCategory,
       rootCauseSummary: issue.rootCauseSummary
    });
@@ -2383,8 +2390,9 @@ const MarkFixedModal: React.FC<{ onClose: () => void, autoFix?:boolean } > = ({
          fixCL = state.fixCL;
       }
 
-      // Root Cause Analysis      
+      // Root Cause Analysis
       let rootCommitId = state.rootCommitId?.length ? state.rootCommitId : null;
+      const rootCauseOwnerId = state.rootCauseOwnerId?.length ? state.rootCauseOwnerId : issue.rootCauseOwnerId?.length ? "" : null;
       const rootCauseCategory = state.rootCauseCategory?.length ? state.rootCauseCategory : "";
       const rootCauseSummary = state?.rootCauseSummary?.trim().length ? state.rootCauseSummary.trim() : "";
 
@@ -2410,6 +2418,7 @@ const MarkFixedModal: React.FC<{ onClose: () => void, autoFix?:boolean } > = ({
          fixChange: fixCL,
          summary: summary,
          rootCommitId: rootCommitId ?? "",
+         rootCauseOwnerId: rootCauseOwnerId,
          rootCauseCategory: rootCauseCategory,
          rootCauseSummary: rootCauseSummary,
          //ownerId: !issue.ownerInfo ? dashboard.userId : issue.ownerInfo.id,
@@ -2436,7 +2445,7 @@ const MarkFixedModal: React.FC<{ onClose: () => void, autoFix?:boolean } > = ({
 
    };
 
-   let height = state.error ? 250 : 210;
+   let height = state.error ? 300 : 260;
 
    height += 384;
 
@@ -2444,6 +2453,10 @@ const MarkFixedModal: React.FC<{ onClose: () => void, autoFix?:boolean } > = ({
 
    const disableRoot = !!state.dupeCheck;
 
+   // Unique set of users that are most likely to be the RCA owner
+   const rcaOwnerCandidates = [... new Set([issue.rootCauseOwnerId ?? null, issue.ownerInfo?.id ?? null, issue.resolvedByInfo?.id ?? null])];
+
+
    return <Modal className={hordeClasses.modal} isOpen={true} styles={{ main: { padding: 8, width: 540, height: height, minHeight: height } }}>
       <Stack horizontal styles={{ root: { padding: 8 } }}>
          <Stack.Item grow={2}>
@@ -2472,7 +2485,13 @@ const MarkFixedModal: React.FC<{ onClose: () => void, autoFix?:boolean } > = ({
             <Stack style={{ paddingLeft: 8, width: 500, paddingBottom: 12 }}>
                <Separator />
             </Stack>
-            <Stack tokens={{ childrenGap: 18 }}>
+            <Stack tokens={{ childrenGap: 16 }}>
+               <Stack>
+                  <Label>Root Cause Owner</Label>
+                  <UserSelect defaultId={state.rootCauseOwnerId} handleSelection={(id => {
+                     setState({ ...state, rootCauseOwnerId: id })
+                  })} userIdHints={rcaOwnerCandidates.filter(c => c !== null)} noResultsFoundText="No users found" />
+               </Stack>
                <Stack>
                   <CommitIdField disabled={disableRoot} label="Root Cause Change" defaultCommitId={state.rootCommitId} placeholder="Change which caused this issue, please only enter if certain on root cause" onChange={(commitId => { setState({ ...state, rootCommitId: commitId }) })} />
                </Stack>
