Horde: Not not include duplicate issues in root cause summary escalation notifications, and do not require Root Commit Id for a notification to be closed

#Synced-CL Unknown

diff --git a/Engine/Source/Programs/Horde/Plugins/Build/HordeServer.Build/Notifications/Sinks/SlackNotificationSink.cs b/Engine/Source/Programs/Horde/Plugins/Build/HordeServer.Build/Notifications/Sinks/SlackNotificationSink.cs
--- a/Engine/Source/Programs/Horde/Plugins/Build/HordeServer.Build/Notifications/Sinks/SlackNotificationSink.cs
+++ b/Engine/Source/Programs/Horde/Plugins/Build/HordeServer.Build/Notifications/Sinks/SlackNotificationSink.cs
@@ -1318,7 +1318,7 @@ namespace HordeServer.Notifications.Sinks
 					}
 				}
 
-				if (issue.ResolvedAt != null && workflow.EscalateRootCauseAnalysisTimes.Count > 0)
+				if ((issue.ResolvedAt != null && issue.DuplicateIssueId == null) && workflow.EscalateRootCauseAnalysisTimes.Count > 0)
 				{
 					DateTime escalateTime = issue.ResolvedAt.Value.AddMinutes(workflow.EscalateRootCauseAnalysisTimes[0]);
 					if (await _redisService.GetDatabase().SortedSetAddAsync(s_escalateRootCauseAnalysis, issue.Id, (escalateTime - DateTime.UnixEpoch).TotalSeconds, SortedSetWhen.NotExists))
@@ -3089,15 +3089,20 @@ namespace HordeServer.Notifications.Sinks
 				return null;
 			}
 
+			if (issue.DuplicateIssueId != null)
+			{
+				return null;
+			}
+
 			UserId? ownerId = issue.RootCauseOwnerId ?? issue.OwnerId ?? issue.ResolvedById;
 			if (ownerId == null || ownerId == UserId.Empty)
 			{
 				return null;
 			}
 
-			// RootCommitId, RootCauseCategory and RootCauseSummary are optional inputs in the UI because the issue closer may not have all the
+			// RootCauseCategory and RootCauseSummary are optional inputs in the UI because the issue closer may not have all the
 			// information when marking as fixed; they are all required to have a value for Root Cause Analysis to be considered closed.
-			if (issue.RootCommitId != null && issue.RootCommitId != CommitId.Empty && !String.IsNullOrEmpty(issue.RootCauseCategory) && !String.IsNullOrEmpty(issue.RootCauseSummary))
+			if (!String.IsNullOrEmpty(issue.RootCauseCategory) && !String.IsNullOrEmpty(issue.RootCauseSummary))
 			{
 				// RCA is completed
 				return null;
