Horde: Add test for sparse streams in PerforceMaterializer

diff --git a/Engine/Source/Programs/Horde/Drivers/JobDriver.Tests/PerforceMaterializerTests.cs b/Engine/Source/Programs/Horde/Drivers/JobDriver.Tests/PerforceMaterializerTests.cs
--- a/Engine/Source/Programs/Horde/Drivers/JobDriver.Tests/PerforceMaterializerTests.cs
+++ b/Engine/Source/Programs/Horde/Drivers/JobDriver.Tests/PerforceMaterializerTests.cs
@@ -52,6 +52,35 @@ public class PerforceMaterializerTests : BasePerforceFixtureTest
 		Assert.AreEqual(7, state!.Changelist);
 	}
 	
+	[TestMethod]
+	public async Task Sync_LatestManuallySpecified_SparseDev_Async()
+	{
+		StreamFixture stream = Fixture.StreamFooSparseDev;
+		using PerforceMaterializer pm = CreatePm(stream);
+		
+		await pm.SyncAsync(stream.Head.Number, -1, new SyncOptions(), CancellationToken.None);
+		stream.GetChangelist(stream.Head.Number).AssertDepotFiles(pm.SyncDir.FullName);
+		await pm.FinalizeAsync(CancellationToken.None);
+
+		PerforceMaterializer.State? state = await pm.LoadStateForTestAsync(CancellationToken.None);
+		Assert.AreEqual(stream.Head.Number, state!.Changelist);
+	}
+	
+	[TestMethod]
+	[Ignore] // Latest change number cannot be fetched for empty sparse streams
+	public async Task Sync_Latest_SparseDev_Async()
+	{
+		StreamFixture stream = Fixture.StreamFooSparseDev;
+		using PerforceMaterializer pm = CreatePm(stream);
+		
+		await pm.SyncAsync(IWorkspaceMaterializer.LatestChangeNumber, -1, new SyncOptions(), CancellationToken.None);
+		stream.GetChangelist(stream.Head.Number).AssertDepotFiles(pm.SyncDir.FullName);
+		await pm.FinalizeAsync(CancellationToken.None);
+
+		PerforceMaterializer.State? state = await pm.LoadStateForTestAsync(CancellationToken.None);
+		Assert.AreEqual(stream.Head.Number, state!.Changelist);
+	}
+	
 	[TestMethod]
 	public async Task Sync_SameChangelistTwice_Async()
 	{
@@ -63,7 +92,7 @@ public class PerforceMaterializerTests : BasePerforceFixtureTest
 		await pm.FinalizeAsync(CancellationToken.None);
 
 		PerforceMaterializer.State? state = await pm.LoadStateForTestAsync(CancellationToken.None);
-		Assert.AreEqual(IWorkspaceMaterializer.LatestChangeNumber, state!.Changelist);
+		Assert.AreEqual(Fixture.StreamFooMain.Head.Number, state!.Changelist);
 	}
 	
 	[TestMethod]
