Horde: Determine latest change number via p4 changes in PerforceMaterializer

diff --git a/Engine/Source/Programs/Horde/Drivers/JobDriver/Execution/PerforceMaterializer.cs b/Engine/Source/Programs/Horde/Drivers/JobDriver/Execution/PerforceMaterializer.cs
--- a/Engine/Source/Programs/Horde/Drivers/JobDriver/Execution/PerforceMaterializer.cs
+++ b/Engine/Source/Programs/Horde/Drivers/JobDriver/Execution/PerforceMaterializer.cs
@@ -275,16 +275,27 @@ public sealed class PerforceMaterializer : IWorkspaceMaterializer
 			await _perforceWithClient.SyncQuietAsync(EpicGames.Perforce.SyncOptions.KeepWorkspaceFiles, -1, $"@{state.Changelist}", cancellationToken);
 			status.Progress = $"({flushTimer.Elapsed.TotalSeconds:0.0}s)";
 		}
-		
+
+		// If syncing to latest we need find the actual change number
+        if (changeNum == IWorkspaceMaterializer.LatestChangeNumber)
+        {
+			List<ChangesRecord> changesRecords = await _perforceWithClient.GetChangesAsync(ChangesOptions.None,  1, ChangeStatus.Submitted, new List<string> {_options.AgentWorkspace.Stream + "/..."}, cancellationToken);
+			if (changesRecords.Count == 0)
+			{
+				throw new Exception("No change records returned when looking up latest change number");
+			}
+			
+			changeNum = changesRecords[0].Number;
+		}
+
 		await SaveStateAsync(TransactionStatus.Dirty, client.Name, changeNum, shelveChangeNum, cancellationToken);
-		int reportedSyncChangeNum;
 		
 		{
 			using TelemetrySpan syncSpan = _tracer.StartActiveSpan($"{nameof(PerforceMaterializer)}.{nameof(SyncAsync)}.Sync");
 			using ILoggerProgress status = _logger.BeginProgressScope("Syncing files...");
 			Stopwatch syncTimer = Stopwatch.StartNew();
 			
-			string changeNumStr = changeNum == IWorkspaceMaterializer.LatestChangeNumber ? "#head" : "@" + changeNum;
+			string changeNumStr = "@" + changeNum;
 			string shelveStr = shelveChangeNum > 0 ? $" with shelved CL {shelveChangeNum}" : "";
 			_logger.LogInformation("Syncing CL {ChangeNum}{Shelve} ... (Threads: {Threads} Batch: {Batch} BatchSize: {BatchSize})",
 				changeNumStr, shelveStr, _options.SyncThreads, _options.SyncBatch, _options.SyncBatchSize);
@@ -293,16 +304,6 @@ public sealed class PerforceMaterializer : IWorkspaceMaterializer
 				EpicGames.Perforce.SyncOptions.None, -1, _options.SyncThreads, _options.SyncBatch, _options.SyncBatchSize,
 				-1, -1, changeNumStr, cancellationToken);
 
-			reportedSyncChangeNum = changeNum;
-			
-			// The sync summary will only report the correct CL if #head was used, and we only need it in that case.
-			// Otherwise, just use the explicit change number passed as the parameter.
-			if (changeNum == IWorkspaceMaterializer.LatestChangeNumber)
-			{
-				// If no files are synced, and we specified #head, we never get back a valid CL in the sync summary.
-				reportedSyncChangeNum = syncSummaries.Count > 0 ? syncSummaries.Max(x => x.Change) : IWorkspaceMaterializer.LatestChangeNumber;
-			}
-
 			long totalFileCount = syncSummaries.Sum(x => x.TotalFileCount);
 			long totalFileSize = syncSummaries.Sum(x => x.TotalFileSize);
 			double totalFileSizeMb = Math.Round(totalFileSize / (1024.0 * 1024.0), 1);
@@ -328,7 +329,7 @@ public sealed class PerforceMaterializer : IWorkspaceMaterializer
 			status.Progress = $"({unshelveTimer.Elapsed.TotalSeconds:0.0}s)";
 		}
 
-		await SaveStateAsync(TransactionStatus.Clean, client.Name, reportedSyncChangeNum, shelveChangeNum, cancellationToken);
+		await SaveStateAsync(TransactionStatus.Clean, client.Name, changeNum, shelveChangeNum, cancellationToken);
 	}
 	
 	/// <inheritdoc/>
@@ -432,7 +433,7 @@ public sealed class PerforceMaterializer : IWorkspaceMaterializer
 		return await PerforceConnection.CreateAsync(perforceSettings, logger);
 	}
 
-	private async Task<IPerforceConnection> GetPerforceWithClientAsync()
+	internal async Task<IPerforceConnection> GetPerforceWithClientAsync()
 	{
 		if (_perforceWithClient != null)
 		{
