fix tool tip text

#rnx

#Synced-CL 48217809

diff --git a/Engine/Source/Programs/Horde/HordeDashboard/plugins/test-automation/testAutomationHealth.tsx b/Engine/Source/Programs/Horde/HordeDashboard/plugins/test-automation/testAutomationHealth.tsx
--- a/Engine/Source/Programs/Horde/HordeDashboard/plugins/test-automation/testAutomationHealth.tsx
+++ b/Engine/Source/Programs/Horde/HordeDashboard/plugins/test-automation/testAutomationHealth.tsx
@@ -776,7 +776,7 @@ const TestHealthDetailsModal: React.FC<{ test: TestHealthItem, handler: TestData
     return <Modal isOpen={true} isBlocking={true} topOffsetFixed={true} styles={{ main: { padding: 8, width: 1134, hasBeenOpened: false, top: "80px", position: "absolute" } }} onDismiss={() => onDismiss()} >
             <Stack grow>
                 <Stack horizontal verticalAlign="center">
-                    <Stack><Text style={{ paddingLeft: 8, fontSize: 16, fontWeight: 600 }}>Health Details - {`${test.name}`}</Text></Stack>
+                    <Stack><Text style={{ paddingLeft: 8, fontSize: 16, fontWeight: 600 }}>Health Details - {test.name}</Text></Stack>
                     <Stack grow />
                     {!isCommitting && test.isTransient &&
                         <Stack style={{ paddingRight: 6 }} title="Revert changes">
@@ -788,11 +788,11 @@ const TestHealthDetailsModal: React.FC<{ test: TestHealthItem, handler: TestData
                         </Stack>
                     }
                     {test.isTransient &&
-                        <Stack style={{ paddingRight: 6 }} title="Commit changes and close">
+                        <Stack style={{ paddingRight: 6 }} title="Commit changes">
                             {!isCommitting &&
                                 <IconButton
                                     iconProps={{ iconName: 'Save', style: { fontSize: 16 } }}
-                                    ariaLabel="Commit and close"
+                                    ariaLabel="Commit changes"
                                     onClick={() => { setCommitting(true); test.commitTransaction().then(() => { setCommitting(false) }); }}
                                 />
                             }
