Horde: Use a real stream client in PerforceMaterializer

Previously, a classic workspace client was created with views copied from the stream. This is/was to allow Horde-based views to be added in addition.

diff --git a/Engine/Source/Programs/Horde/Drivers/JobDriver.Tests/PerforceMaterializerTests.cs b/Engine/Source/Programs/Horde/Drivers/JobDriver.Tests/PerforceMaterializerTests.cs
--- a/Engine/Source/Programs/Horde/Drivers/JobDriver.Tests/PerforceMaterializerTests.cs
+++ b/Engine/Source/Programs/Horde/Drivers/JobDriver.Tests/PerforceMaterializerTests.cs
@@ -6,6 +6,7 @@ using System.IO;
 using System.Linq;
 using System.Threading;
 using System.Threading.Tasks;
+using EpicGames.Perforce;
 using EpicGames.Perforce.Fixture;
 using HordeCommon.Rpc.Messages;
 using JobDriver.Execution;
@@ -14,6 +15,7 @@ using Microsoft.Extensions.Logging;
 using Microsoft.VisualStudio.TestTools.UnitTesting;
 using OpenTelemetry.Trace;
 using FileSet = System.Collections.Generic.HashSet<(string clientFile, long size, string digest)>;
+using SyncOptions = JobDriver.Execution.SyncOptions;
 
 namespace JobDriver.Tests;
 
@@ -39,6 +41,33 @@ public class PerforceMaterializerTests : BasePerforceFixtureTest
 		}
 	}
 
+	[TestMethod]
+	[Ignore] // Disabled for now as it mutates the stream. Need to obliterate after submit, or similar kind of restore
+	public async Task Submit_Async()
+	{
+		using PerforceMaterializer pm = CreatePm(Fixture.StreamFooMain);
+		
+		await pm.SyncAsync(IWorkspaceMaterializer.LatestChangeNumber, -1, new SyncOptions(), CancellationToken.None);
+		Fixture.StreamFooMain.GetChangelist(7).AssertDepotFiles(pm.SyncDir.FullName);
+		await SubmitFileAsync(pm, "Adding new file", "hello.txt", "Hello world!");
+		await pm.FinalizeAsync(CancellationToken.None);
+
+		Assert.IsTrue(true, "Submit went through without exception being thrown");
+	}
+
+	private static async Task SubmitFileAsync(PerforceMaterializer pm, string clDescription, string filename, string fileContent)
+	{
+		string tempFile = Path.Join(pm.SyncDir.FullName, filename);
+		await File.WriteAllTextAsync(tempFile, fileContent);
+
+		IPerforceConnection perforce = await pm.GetPerforceWithClientAsync();
+		List<AddRecord> addRecords = await perforce.AddAsync(-1, tempFile);
+		ChangeRecord change = new () { Description = clDescription };
+		change.Files.Add(addRecords[0].DepotFile);
+		change = await perforce.CreateChangeAsync(change);
+		await perforce.SubmitAsync(change.Number, SubmitOptions.None);
+	}
+
 	[TestMethod]
 	public async Task Sync_Latest_Async()
 	{
diff --git a/Engine/Source/Programs/Horde/Drivers/JobDriver/Execution/PerforceMaterializer.cs b/Engine/Source/Programs/Horde/Drivers/JobDriver/Execution/PerforceMaterializer.cs
--- a/Engine/Source/Programs/Horde/Drivers/JobDriver/Execution/PerforceMaterializer.cs
+++ b/Engine/Source/Programs/Horde/Drivers/JobDriver/Execution/PerforceMaterializer.cs
@@ -471,26 +471,14 @@ public sealed class PerforceMaterializer : IWorkspaceMaterializer
 		string clientName = _perforceSettings.ClientName!;
 		string streamName = _options.AgentWorkspace.Stream;
 		
-		// Use a placeholder name to avoid conflicting with any existing client name
-		const string PlaceholderName = "horde-placeholder-name-pm";
-		
-		// Assume latest stream version is wanted. Should possibly be tied to current CL being synced
-		ClientRecord streamClient = await _perforceWithoutClient.GetStreamViewAsync(PlaceholderName, streamName, -1, cancellationToken);
-		
-		// Horde-based views are unused for now until a filtering syntax compatible with ManagedWorkspace is figured out
-		List<string> hordeView = _options.AgentWorkspace.View.Select(x => ConvertHordeViewToClientView(x, streamName, clientName)).ToList();
-		List<string> view = [..streamClient.View.Select(x => x.Replace(PlaceholderName, clientName, StringComparison.Ordinal))];
-		
-		// Create a classic client while still applying the view from stream together with extra Horde-defined view entries.
-		// A stream-based client won't allow changing the view
 		ClientRecord newClient = new (clientName, _perforceSettings.UserName, _workspaceDir.FullName)
 		{
 			Host = _options.HostName ?? Environment.MachineName,
 			Type = "partitioned",
+			Stream = streamName,
 			Root = _workspaceDir.FullName,
 			Options = _options.UseModTime ? ClientOptions.Clobber | ClientOptions.ModTime : ClientOptions.Clobber,
 		};
-		newClient.View.AddRange(view);
 		Directory.CreateDirectory(newClient.Root);
 		
 		using ILoggerProgress status = _logger.BeginProgressScope(forceCreate ? "Force creating client ..." : "Updating client ...");
@@ -511,6 +499,7 @@ public sealed class PerforceMaterializer : IWorkspaceMaterializer
 			bool clientExists = existingClient.Update != DateTime.MinValue;
 			if (clientExists &&
 				existingClient.Owner == newClient.Owner &&
+				existingClient.Stream == newClient.Stream &&
 				existingClient.Root == newClient.Root &&
 				existingClient.Host == newClient.Host &&
 				existingClient.Type == newClient.Type)
