[Horde][BuildHealth][Dashboard] Display user name for rootCauseOwnerId in issue modal

#Synced-CL Unknown

diff --git a/Engine/Source/Programs/Horde/HordeDashboard/src/backend/Api.ts b/Engine/Source/Programs/Horde/HordeDashboard/src/backend/Api.ts
--- a/Engine/Source/Programs/Horde/HordeDashboard/src/backend/Api.ts
+++ b/Engine/Source/Programs/Horde/HordeDashboard/src/backend/Api.ts
@@ -3238,6 +3238,9 @@ export type GetIssueResponse = {
 	/**Root cause owner id*/
 	rootCauseOwnerId?: string;
 
+	/**Root cause owner info*/
+	rootCauseOwnerInfo?: GetThinUserInfoResponse;
+
 	/**Root cause category*/
 	rootCauseCategory?: string;
 
diff --git a/Engine/Source/Programs/Horde/HordeDashboard/src/components/IssueViewV2.tsx b/Engine/Source/Programs/Horde/HordeDashboard/src/components/IssueViewV2.tsx
--- a/Engine/Source/Programs/Horde/HordeDashboard/src/components/IssueViewV2.tsx
+++ b/Engine/Source/Programs/Horde/HordeDashboard/src/components/IssueViewV2.tsx
@@ -1362,7 +1362,7 @@ const IssueSummaryPanel: React.FC = () => {
      if (issue.rootCauseOwnerId) {
       items.push({
          title: "Root Cause Owner",
-         text: issue.rootCauseOwnerId
+         text: issue.rootCauseOwnerInfo?.name ?? "User not found"
       })
    }
    if (issue.rootCauseCategory) {
