Horde: Update release notes

#rnx

diff --git a/Engine/Source/Programs/Horde/Docs/ReleaseNotes.md b/Engine/Source/Programs/Horde/Docs/ReleaseNotes.md
--- a/Engine/Source/Programs/Horde/Docs/ReleaseNotes.md
+++ b/Engine/Source/Programs/Horde/Docs/ReleaseNotes.md
@@ -1,5 +1,39 @@
 # Release Notes
 
+## 2025-11-13
+
+* Use a real stream client in PerforceMaterializer
+  Previously, a classic workspace client was created with views copied from the stream. This is/was to allow Horde-based views to be added in addition. (48226301)
+* Determine latest change number via p4 changes in PerforceMaterializer (48222459)
+* Add test for sparse streams in PerforceMaterializer (48216462)
+* Do not include duplicate issues in root cause summary escalation notifications, and do not require Root Commit Id for a notification to be closed (48216270)
+* EpicGames.Perforce: Add sparse streams to fixture (48215972)
+* EpicGames.Perforce: Add support for streamlog command (48215109)
+* Add RCA owner input to dashbaord (48202681)
+* Always set ResolvedById in Slack workflow for resolved issues (non-null fixCommitId) (48200054)
+* JobStepRef parity changes. (48192973)
+* Root Cause Summary Escalation Notifications to Slack (48167554)
+* IssueTelemetry event name fix (48161594)
+* Expose lease outcome reasons in REST API (48153944)
+* Fix issue with case-sensitive materializer name comparison for conforms (48140462)
+* Add lastSuccessTime and firstFailureTime to FindSuspectsForSpanAsync and use them for time delta gate.
+  Fetching these times via GetAsync was causing exceptions to be raised when commits were out of stream and thus exiting suspect finding altogether. (48100206)
+* Fix Horde server build by adding and removing paths to reflect the migration of log matchers to their own project. (48079442)
+* Skipped steps as a summarized metric & manual invocation metadata for CSV. (48062467)
+* Add options for controlling sync threads and batches in PerforceMaterializer
+  Also add support for ModTime as a flag. (48029020)
+* Updated RetrieveBlockAsync to no longer enable verbose logging on retries of extracting the directory node. We suspect that verbose logging was causing a large amount of overhead resulting in steps timing out and cancelling. We still enable hash verification of files though. (48024878)
+* Remove dependency lib inclusion (48009192)
+* Including ODBC reference (47977362)
+* Updating dev environment dependencies due to transient dependecy postinstall changes. (47974458)
+* Fix issue with logging of config exception stacktraces (47961496)
+* Make awaiting on background notification tasks more robust
+  - Catches any exception directly in the task
+  - Simplify by using a list and only waiting in StopAsync (47960832)
+* Fix config object append bug leading to infinite loop (47953303)
+* Add reason to lease outcomes (47921082)
+* Fixing an issue where we were disposing of the HttpResponse content before reading it. (47903704)
+
 ## 2025-11-06
 
 * Fixing an issue where we were disposing of the HttpResponse content before reading it. (47903704)
