Support markdown in audit notes

#rnx

#Synced-CL 48217809

diff --git a/Engine/Source/Programs/Horde/HordeDashboard/plugins/test-automation/testAutomationCommon.tsx b/Engine/Source/Programs/Horde/HordeDashboard/plugins/test-automation/testAutomationCommon.tsx
--- a/Engine/Source/Programs/Horde/HordeDashboard/plugins/test-automation/testAutomationCommon.tsx
+++ b/Engine/Source/Programs/Horde/HordeDashboard/plugins/test-automation/testAutomationCommon.tsx
@@ -11,6 +11,7 @@ import { getHordeStyling } from "horde/styles/Styles";
 
 import { GetUserResponse } from "horde/backend/Api";
 import { UserSelect } from "horde/components/UserSelect";
+import { Markdown } from "horde/base/components/Markdown";
 
 import * as d3 from "d3";
 
@@ -66,6 +67,7 @@ const getStyles = () => {
 
     const border = `1px solid ${dashboard.darktheme ? "#2D2B29" : "#EDEBE9"}`;
     const highlight = dashboard.darktheme ? "brightness(1.2)" : "brightness(0.9)";
+    const textEmphasisColor = dashboard.darktheme? "#ffffff57" : "#00000057";
 
     const styles = _styles ?? mergeStyleSets({
         list: {
@@ -105,6 +107,9 @@ const getStyles = () => {
                 }
             }
         },
+        placeholder: {
+            color: textEmphasisColor
+        },
    });
 
    _styles = styles;
@@ -912,22 +917,33 @@ export const EditableTextField: React.FC<{ label: string, onChange: (value: stri
 
     const styles = getStyles();
 
-    return <Stack className={styles.editable} horizontal verticalAlign="center" tokens={{childrenGap: 6}} style={style} grow>
-                <Stack horizontal tokens={{childrenGap: 6}} grow>
+    return <Stack className={styles.editable} horizontal tokens={{childrenGap: 6}} style={style} grow>
+                <Stack horizontal tokens={{childrenGap: 6}} grow verticalAlign="start">
                     <Label>{label}: </Label>
-                    <TextField
-                        key={`text-${loaded}`}
-                        defaultValue={text}
-                        onNotifyValidationResult={(_, value) => { if (value !== text) { onChange(value); } setEnableEditing(false); }}
-                        validateOnFocusOut
-                        validateOnLoad={false}
-                        multiline={multiline}
-                        styles={{root: {width: '100%'}}}
-                        placeholder={placeholder}
-                        disabled={!enableEditing}
-                    />
+                    {!enableEditing &&
+                        <Stack style={{cursor: 'pointer', padding: 5}} grow verticalFill
+                            onClick={() => {setEnableEditing(true)}}
+                        >
+                            {!!text && <Markdown styles={{ subComponentStyles: { paragraph: { root: { marginTop: 0 } } } }}>{text}</Markdown>}
+                            {!text && <Text className={styles.placeholder}>{placeholder}</Text>}
+                        </Stack>}
+                    {enableEditing &&
+                        <TextField
+                            key={`text-${loaded}`}
+                            defaultValue={text}
+                            onNotifyValidationResult={(_, value) => { if (value !== text) { onChange(value); } setEnableEditing(false); }}
+                            validateOnFocusOut
+                            validateOnLoad={false}
+                            multiline={multiline}
+                            borderless
+                            autoAdjustHeight
+                            styles={{root: {width: '100%'}}}
+                            placeholder={placeholder}
+                            autoFocus
+                        />
+                    }
                 </Stack>
-                <Stack className="Edit" style={{height: multiline? 50 : 28, padding: 2}} verticalFill verticalAlign="center" onClick={() => {setEnableEditing(!enableEditing);}}>
+                <Stack className="Edit" style={{padding: 2, minHeight: multiline? 50 : undefined}} verticalFill verticalAlign="center" onClick={() => {setEnableEditing(!enableEditing);}}>
                     <FontIcon style={{ fontSize: 12 }} iconName="Edit" />
                 </Stack>
             </Stack>    
