Horde: Fix issue determining latest change number in PerforceMaterializer

Use workspace root instead of stream.

diff --git a/Engine/Source/Programs/Horde/Drivers/JobDriver.Tests/PerforceMaterializerTests.cs b/Engine/Source/Programs/Horde/Drivers/JobDriver.Tests/PerforceMaterializerTests.cs
--- a/Engine/Source/Programs/Horde/Drivers/JobDriver.Tests/PerforceMaterializerTests.cs
+++ b/Engine/Source/Programs/Horde/Drivers/JobDriver.Tests/PerforceMaterializerTests.cs
@@ -96,7 +96,6 @@ public class PerforceMaterializerTests : BasePerforceFixtureTest
 	}
 	
 	[TestMethod]
-	[Ignore] // Latest change number cannot be fetched for empty sparse streams
 	public async Task Sync_Latest_SparseDev_Async()
 	{
 		StreamFixture stream = Fixture.StreamFooSparseDev;
@@ -107,7 +106,7 @@ public class PerforceMaterializerTests : BasePerforceFixtureTest
 		await pm.FinalizeAsync(CancellationToken.None);
 
 		PerforceMaterializer.State? state = await pm.LoadStateForTestAsync(CancellationToken.None);
-		Assert.AreEqual(stream.Head.Number, state!.Changelist);
+		Assert.AreEqual(7, state!.Changelist);
 	}
 	
 	[TestMethod]
diff --git a/Engine/Source/Programs/Horde/Drivers/JobDriver/Execution/PerforceMaterializer.cs b/Engine/Source/Programs/Horde/Drivers/JobDriver/Execution/PerforceMaterializer.cs
--- a/Engine/Source/Programs/Horde/Drivers/JobDriver/Execution/PerforceMaterializer.cs
+++ b/Engine/Source/Programs/Horde/Drivers/JobDriver/Execution/PerforceMaterializer.cs
@@ -279,7 +279,7 @@ public sealed class PerforceMaterializer : IWorkspaceMaterializer
 		// If syncing to latest we need find the actual change number
         if (changeNum == IWorkspaceMaterializer.LatestChangeNumber)
         {
-			List<ChangesRecord> changesRecords = await _perforceWithClient.GetChangesAsync(ChangesOptions.None,  1, ChangeStatus.Submitted, new List<string> {_options.AgentWorkspace.Stream + "/..."}, cancellationToken);
+			List<ChangesRecord> changesRecords = await _perforceWithClient.GetChangesAsync(ChangesOptions.None,  1, ChangeStatus.Submitted, new List<string> {$"//{client.Name}/..."}, cancellationToken);
 			if (changesRecords.Count == 0)
 			{
 				throw new Exception("No change records returned when looking up latest change number");
