[Horde][BuildHealth][Server] Skip RCA escalation for issues resolved by Horde
#Synced-CL Unknown

diff --git a/Engine/Source/Programs/Horde/Plugins/Build/HordeServer.Build/Notifications/Sinks/SlackNotificationSink.cs b/Engine/Source/Programs/Horde/Plugins/Build/HordeServer.Build/Notifications/Sinks/SlackNotificationSink.cs
--- a/Engine/Source/Programs/Horde/Plugins/Build/HordeServer.Build/Notifications/Sinks/SlackNotificationSink.cs
+++ b/Engine/Source/Programs/Horde/Plugins/Build/HordeServer.Build/Notifications/Sinks/SlackNotificationSink.cs
@@ -1318,7 +1318,7 @@ namespace HordeServer.Notifications.Sinks
 					}
 				}
 
-				if ((issue.ResolvedAt != null && issue.DuplicateIssueId == null) && workflow.EscalateRootCauseAnalysisTimes.Count > 0)
+				if ((issue.ResolvedAt != null && issue.ResolvedById != null && issue.DuplicateIssueId == null) && workflow.EscalateRootCauseAnalysisTimes.Count > 0)
 				{
 					DateTime escalateTime = issue.ResolvedAt.Value.AddMinutes(workflow.EscalateRootCauseAnalysisTimes[0]);
 					if (await _redisService.GetDatabase().SortedSetAddAsync(s_escalateRootCauseAnalysis, issue.Id, (escalateTime - DateTime.UnixEpoch).TotalSeconds, SortedSetWhen.NotExists))
@@ -3089,6 +3089,11 @@ namespace HordeServer.Notifications.Sinks
 				return null;
 			}
 
+			if (issue.ResolvedById == null)
+			{
+				return null;
+			}
+
 			if (issue.DuplicateIssueId != null)
 			{
 				return null;
