[Horde][BuildHealth][Server] Avoid multiple occurences of the same user in Slack workflow root cause owner options
#Synced-CL Unknown

diff --git a/Engine/Source/Programs/Horde/Plugins/Build/HordeServer.Build/Notifications/Sinks/SlackNotificationSink.cs b/Engine/Source/Programs/Horde/Plugins/Build/HordeServer.Build/Notifications/Sinks/SlackNotificationSink.cs
--- a/Engine/Source/Programs/Horde/Plugins/Build/HordeServer.Build/Notifications/Sinks/SlackNotificationSink.cs
+++ b/Engine/Source/Programs/Horde/Plugins/Build/HordeServer.Build/Notifications/Sinks/SlackNotificationSink.cs
@@ -3718,7 +3718,7 @@ namespace HordeServer.Notifications.Sinks
 							rootCauseOwnerOptions.Add(new SlackOption($"Current root cause owner ({rootCauseOwner.Name})", rootCauseOwner.Id.ToString()));
 						}
 					}
-					if (issue.OwnerId != null && issue.OwnerId != user.Id)
+					if (issue.OwnerId != null && issue.OwnerId != user.Id && issue.OwnerId != issue.RootCauseOwnerId)
 					{
 						IUser? owner = await _userCollection.GetCachedUserAsync(issue.OwnerId.Value, cancellationToken);
 						if (owner != null)
@@ -3726,7 +3726,7 @@ namespace HordeServer.Notifications.Sinks
 							rootCauseOwnerOptions.Add(new SlackOption($"Current owner ({owner.Name})", owner.Id.ToString()));
 						}
 					}
-					if (issue.ResolvedById != null && issue.ResolvedById != user.Id)
+					if (issue.ResolvedById != null && issue.ResolvedById != user.Id && issue.ResolvedById != issue.OwnerId && issue.ResolvedById != issue.RootCauseOwnerId)
 					{
 						IUser? resolvedBy = await _userCollection.GetCachedUserAsync(issue.ResolvedById.Value, cancellationToken);
 						if (resolvedBy != null)
