[Horde][BuildHealth][Server] Move Me option for Root Cause Owner to top of selection to match existing UI layout
#Synced-CL Unknown

diff --git a/Engine/Source/Programs/Horde/Plugins/Build/HordeServer.Build/Notifications/Sinks/SlackNotificationSink.cs b/Engine/Source/Programs/Horde/Plugins/Build/HordeServer.Build/Notifications/Sinks/SlackNotificationSink.cs
--- a/Engine/Source/Programs/Horde/Plugins/Build/HordeServer.Build/Notifications/Sinks/SlackNotificationSink.cs
+++ b/Engine/Source/Programs/Horde/Plugins/Build/HordeServer.Build/Notifications/Sinks/SlackNotificationSink.cs
@@ -3709,7 +3709,10 @@ namespace HordeServer.Notifications.Sinks
 
 				if (issue.DuplicateIssueId == null)
 				{
-					List<SlackOption> rootCauseOwnerOptions = [];
+					List<SlackOption> rootCauseOwnerOptions =
+					[
+						new SlackOption("Me", user.Id.ToString())
+					];
 					if (issue.RootCauseOwnerId != null && issue.RootCauseOwnerId != user.Id)
 					{
 						IUser? rootCauseOwner = await _userCollection.GetCachedUserAsync(issue.RootCauseOwnerId.Value, cancellationToken);
@@ -3734,8 +3737,6 @@ namespace HordeServer.Notifications.Sinks
 							rootCauseOwnerOptions.Add(new SlackOption($"Fixed by ({resolvedBy.Name})", resolvedBy.Id.ToString()));
 						}
 					}
-					rootCauseOwnerOptions.Add(new SlackOption("Me", user.Id.ToString()));
-					//rootCauseOwnerOptions.Add(new SlackOption("Unassigned", UserId.Empty.Id.ToString()));
 
 					RadioButtonGroupElement rootCauseOwnership = new("rootcause_owner_action", rootCauseOwnerOptions);
 					//rootCauseOwnership.InitialOption = rootCauseOwnerOptions[0];
