[Horde][BuildHealth][Dashboard] Add current user to dropdown hints for root cause owner

#Synced-CL Unknown

diff --git a/Engine/Source/Programs/Horde/HordeDashboard/src/components/IssueViewV2.tsx b/Engine/Source/Programs/Horde/HordeDashboard/src/components/IssueViewV2.tsx
--- a/Engine/Source/Programs/Horde/HordeDashboard/src/components/IssueViewV2.tsx
+++ b/Engine/Source/Programs/Horde/HordeDashboard/src/components/IssueViewV2.tsx
@@ -2454,7 +2454,7 @@ const MarkFixedModal: React.FC<{ onClose: () => void, autoFix?:boolean } > = ({
    const disableRoot = !!state.dupeCheck;
 
    // Unique set of users that are most likely to be the RCA owner
-   const rcaOwnerCandidates = [... new Set([issue.rootCauseOwnerId ?? null, issue.ownerInfo?.id ?? null, issue.resolvedByInfo?.id ?? null])];
+   const rcaOwnerCandidates = [... new Set([ dashboard.userId, issue.rootCauseOwnerId ?? null, issue.ownerInfo?.id ?? null, issue.resolvedByInfo?.id ?? null])];
 
 
    return <Modal className={hordeClasses.modal} isOpen={true} styles={{ main: { padding: 8, width: 540, height: height, minHeight: height } }}>
