Revert change to skip RCA escalation for issues closed by Horde and add additonal logging for debugging
#rnx

diff --git a/Engine/Source/Programs/Horde/Plugins/Build/HordeServer.Build/Notifications/Sinks/SlackNotificationSink.cs b/Engine/Source/Programs/Horde/Plugins/Build/HordeServer.Build/Notifications/Sinks/SlackNotificationSink.cs
--- a/Engine/Source/Programs/Horde/Plugins/Build/HordeServer.Build/Notifications/Sinks/SlackNotificationSink.cs
+++ b/Engine/Source/Programs/Horde/Plugins/Build/HordeServer.Build/Notifications/Sinks/SlackNotificationSink.cs
@@ -1318,12 +1318,15 @@ namespace HordeServer.Notifications.Sinks
 					}
 				}
 
-				if ((issue.ResolvedAt != null && issue.ResolvedById != null && issue.DuplicateIssueId == null) && workflow.EscalateRootCauseAnalysisTimes.Count > 0)
+				if ((issue.ResolvedAt != null && issue.DuplicateIssueId == null) && workflow.EscalateRootCauseAnalysisTimes.Count > 0)
 				{
 					DateTime escalateTime = issue.ResolvedAt.Value.AddMinutes(workflow.EscalateRootCauseAnalysisTimes[0]);
 					if (await _redisService.GetDatabase().SortedSetAddAsync(s_escalateRootCauseAnalysis, issue.Id, (escalateTime - DateTime.UnixEpoch).TotalSeconds, SortedSetWhen.NotExists))
 					{
 						_logger.LogInformation("First escalation time for root cause analysis {IssueId} is {Time}", issue.Id, escalateTime);
+
+						// TODO: Remove debug logging
+						_logger.LogInformation("issue.ResolvedById set to {ResolvedById}", issue.ResolvedById);
 					}
 				}
 
@@ -3089,10 +3092,13 @@ namespace HordeServer.Notifications.Sinks
 				return null;
 			}
 
-			if (issue.ResolvedById == null)
-			{
-				return null;
-			}
+			//if (issue.ResolvedById == null)
+			//{
+			//	return null;
+			//}
+
+			// TODO: Remove debug logging
+			_logger.LogInformation("issue.ResolvedById set to {ResolvedById}", issue.ResolvedById);
 
 			if (issue.DuplicateIssueId != null)
 			{
