Remove debug logging
#rnx

diff --git a/Engine/Source/Programs/Horde/Plugins/Build/HordeServer.Build/Notifications/Sinks/SlackNotificationSink.cs b/Engine/Source/Programs/Horde/Plugins/Build/HordeServer.Build/Notifications/Sinks/SlackNotificationSink.cs
--- a/Engine/Source/Programs/Horde/Plugins/Build/HordeServer.Build/Notifications/Sinks/SlackNotificationSink.cs
+++ b/Engine/Source/Programs/Horde/Plugins/Build/HordeServer.Build/Notifications/Sinks/SlackNotificationSink.cs
@@ -1324,9 +1324,6 @@ namespace HordeServer.Notifications.Sinks
 					if (await _redisService.GetDatabase().SortedSetAddAsync(s_escalateRootCauseAnalysis, issue.Id, (escalateTime - DateTime.UnixEpoch).TotalSeconds, SortedSetWhen.NotExists))
 					{
 						_logger.LogInformation("First escalation time for root cause analysis {IssueId} is {Time}", issue.Id, escalateTime);
-
-						// TODO: Remove debug logging
-						_logger.LogInformation("issue.ResolvedById set to {ResolvedById}", issue.ResolvedById);
 					}
 				}
 
@@ -3092,14 +3089,6 @@ namespace HordeServer.Notifications.Sinks
 				return null;
 			}
 
-			//if (issue.ResolvedById == null)
-			//{
-			//	return null;
-			//}
-
-			// TODO: Remove debug logging
-			_logger.LogInformation("issue.ResolvedById set to {ResolvedById}", issue.ResolvedById);
-
 			if (issue.DuplicateIssueId != null)
 			{
 				return null;
