Horde: Make ModTime flag disabled by default for PerforceMaterializer

#skipci

diff --git a/Engine/Source/Programs/Horde/Drivers/JobDriver/Execution/PerforceMaterializer.cs b/Engine/Source/Programs/Horde/Drivers/JobDriver/Execution/PerforceMaterializer.cs
--- a/Engine/Source/Programs/Horde/Drivers/JobDriver/Execution/PerforceMaterializer.cs
+++ b/Engine/Source/Programs/Horde/Drivers/JobDriver/Execution/PerforceMaterializer.cs
@@ -33,7 +33,7 @@ public record PerforceMaterializerOptions(
 	int SyncThreads = -1,
 	int SyncBatch = -1,
 	int SyncBatchSize = -1,
-	bool UseModTime = true)
+	bool UseModTime = false)
 {
 	/// <summary>
 	/// Parse options from <see cref="RpcAgentWorkspace.Method" /> as HTTP query string
