Horde: Add log message why a RCA notification was cancelled, and where it was sent to

#rnx
#Synced-CL Unknown

diff --git a/Engine/Source/Programs/Horde/Plugins/Build/HordeServer.Build/Notifications/Sinks/SlackNotificationSink.cs b/Engine/Source/Programs/Horde/Plugins/Build/HordeServer.Build/Notifications/Sinks/SlackNotificationSink.cs
--- a/Engine/Source/Programs/Horde/Plugins/Build/HordeServer.Build/Notifications/Sinks/SlackNotificationSink.cs
+++ b/Engine/Source/Programs/Horde/Plugins/Build/HordeServer.Build/Notifications/Sinks/SlackNotificationSink.cs
@@ -3075,28 +3075,33 @@ namespace HordeServer.Notifications.Sinks
 			IIssue? issue = await _issueService.Collection.GetIssueAsync(issueId, cancellationToken);
 			if (issue == null)
 			{
+				_logger.LogInformation("Cancelling escalation for {NotificationName} {IssueId} {Reason}", notificationName, issueId, "Cannot Find Issue");
 				return null;
 			}
 
 			IReadOnlyList<IIssueSpan> spans = await _issueService.Collection.FindSpansAsync(issueId, cancellationToken);
 			if (spans.Count == 0)
 			{
+				_logger.LogInformation("Cancelling escalation for {NotificationName} {IssueId} {Reason}", notificationName, issueId, "Cannot Find Issue Spans");
 				return null;
 			}
 
 			if (issue.ResolvedAt == null)
 			{
+				_logger.LogInformation("Cancelling escalation for {NotificationName} {IssueId} {Reason}", notificationName, issueId, "Issue Not Resolved");
 				return null;
 			}
 
 			if (issue.DuplicateIssueId != null)
 			{
+				_logger.LogInformation("Cancelling escalation for {NotificationName} {IssueId} {Reason}", notificationName, issueId, "Duplicate Issue");
 				return null;
 			}
 
 			UserId? ownerId = issue.RootCauseOwnerId ?? issue.OwnerId ?? issue.ResolvedById;
 			if (ownerId == null || ownerId == UserId.Empty)
 			{
+				_logger.LogInformation("Cancelling escalation for {NotificationName} {IssueId} {Reason}", notificationName, issueId, "No Root Cause Owner");
 				return null;
 			}
 
@@ -3104,7 +3109,7 @@ namespace HordeServer.Notifications.Sinks
 			// information when marking as fixed; they are all required to have a value for Root Cause Analysis to be considered closed.
 			if (!String.IsNullOrEmpty(issue.RootCauseCategory) && !String.IsNullOrEmpty(issue.RootCauseSummary))
 			{
-				// RCA is completed
+				_logger.LogInformation("Cancelling escalation for {NotificationName} {IssueId} {Reason}", notificationName, issueId, "RCA Completed");
 				return null;
 			}
 
@@ -3113,34 +3118,43 @@ namespace HordeServer.Notifications.Sinks
 			WorkflowId? workflowId = span.LastFailure.Annotations.WorkflowId;
 			if (workflowId == null)
 			{
+				_logger.LogInformation("Cancelling escalation for {NotificationName} {IssueId} {Reason}", notificationName, issueId, "Cannot Find Issue WorkflowID");
 				return null;
 			}
 
 			if (!_buildConfig.CurrentValue.TryGetStream(span.StreamId, out StreamConfig? streamConfig))
 			{
+				_logger.LogInformation("Cancelling escalation for {NotificationName} {IssueId} {Reason}", notificationName, issueId, "Cannot Find Issue StreamConfig");
 				return null;
 			}
 			if (!streamConfig.TryGetWorkflow(workflowId.Value, out WorkflowConfig? workflow))
 			{
+				_logger.LogInformation("Cancelling escalation for {NotificationName} {IssueId} {Reason}", notificationName, issueId, "Cannot Find Issue StreamConfig Workflow");
 				return null;
 			}
 			if (String.IsNullOrEmpty(workflow.TriageChannel) || workflow.EscalateRootCauseAnalysisTimes.Count == 0)
 			{
+				_logger.LogInformation("Cancelling escalation for {NotificationName} {IssueId} {Reason}", notificationName, issueId, "No Channel or Escalation Times");
 				return null;
 			}
 
 			if (!issue.IsMarkedFixed())
 			{
+				_logger.LogInformation("Cancelling escalation for {NotificationName} {IssueId} {Reason}", notificationName, issueId, "Issue Not Fixed");
 				return null;
 			}
 
 			MessageStateDocument? state = await GetMessageStateAsync(workflow.TriageChannel, GetTriageThreadEventId(issueId), cancellationToken);
 			if (state == null)
 			{
+				_logger.LogInformation("Cancelling escalation for {NotificationName} {IssueId} {Reason}", notificationName, issueId, "Cannot Find Issue Message State");
 				return null;
 			}
 
 			string mention = await FormatMentionAsync(ownerId.Value, workflow.AllowMentions, cancellationToken);
+
+			_logger.LogInformation("Sending escalation message for {NotificationName} {IssueId} {User} {Channel}", notificationName, issueId, mention, state.MessageId);
+
 			Uri issueUrl = GetIssueUrl(issue, span.FirstFailure);
 			string message = $"{mention} this root cause analysis requires your attention - Issue <{issueUrl}|{issue.Id}> - Select ...More > Modify Issue to complete the root cause analysis.";
 			await _slackClient.PostMessageToThreadAsync(state.MessageId, message, cancellationToken);
