Horde: Use const message string for RCA notification was cancelled log messages

#rnx

#Synced-CL Unknown

diff --git a/Engine/Source/Programs/Horde/Plugins/Build/HordeServer.Build/Notifications/Sinks/SlackNotificationSink.cs b/Engine/Source/Programs/Horde/Plugins/Build/HordeServer.Build/Notifications/Sinks/SlackNotificationSink.cs
--- a/Engine/Source/Programs/Horde/Plugins/Build/HordeServer.Build/Notifications/Sinks/SlackNotificationSink.cs
+++ b/Engine/Source/Programs/Horde/Plugins/Build/HordeServer.Build/Notifications/Sinks/SlackNotificationSink.cs
@@ -3072,36 +3072,38 @@ namespace HordeServer.Notifications.Sinks
 
 		async Task<double?> EscalateSingleRootCauseAnalysisAsync(int issueId, DateTime utcNow, string notificationName, CancellationToken cancellationToken)
 		{
+			const string CanceledMessage = "Canceling escalation for {NotificationName} {IssueId} {Reason}";
+
 			IIssue? issue = await _issueService.Collection.GetIssueAsync(issueId, cancellationToken);
 			if (issue == null)
 			{
-				_logger.LogInformation("Cancelling escalation for {NotificationName} {IssueId} {Reason}", notificationName, issueId, "Cannot Find Issue");
+				_logger.LogInformation(CanceledMessage, notificationName, issueId, "Cannot Find Issue");
 				return null;
 			}
 
 			IReadOnlyList<IIssueSpan> spans = await _issueService.Collection.FindSpansAsync(issueId, cancellationToken);
 			if (spans.Count == 0)
 			{
-				_logger.LogInformation("Cancelling escalation for {NotificationName} {IssueId} {Reason}", notificationName, issueId, "Cannot Find Issue Spans");
+				_logger.LogInformation(CanceledMessage, notificationName, issueId, "Cannot Find Issue Spans");
 				return null;
 			}
 
 			if (issue.ResolvedAt == null)
 			{
-				_logger.LogInformation("Cancelling escalation for {NotificationName} {IssueId} {Reason}", notificationName, issueId, "Issue Not Resolved");
+				_logger.LogInformation(CanceledMessage, notificationName , issueId, "Issue Not Resolved");
 				return null;
 			}
 
 			if (issue.DuplicateIssueId != null)
 			{
-				_logger.LogInformation("Cancelling escalation for {NotificationName} {IssueId} {Reason}", notificationName, issueId, "Duplicate Issue");
+				_logger.LogInformation(CanceledMessage, notificationName, issueId, "Duplicate Issue");
 				return null;
 			}
 
 			UserId? ownerId = issue.RootCauseOwnerId ?? issue.OwnerId ?? issue.ResolvedById;
 			if (ownerId == null || ownerId == UserId.Empty)
 			{
-				_logger.LogInformation("Cancelling escalation for {NotificationName} {IssueId} {Reason}", notificationName, issueId, "No Root Cause Owner");
+				_logger.LogInformation(CanceledMessage, notificationName, issueId, "No Root Cause Owner");
 				return null;
 			}
 
@@ -3109,7 +3111,7 @@ namespace HordeServer.Notifications.Sinks
 			// information when marking as fixed; they are all required to have a value for Root Cause Analysis to be considered closed.
 			if (!String.IsNullOrEmpty(issue.RootCauseCategory) && !String.IsNullOrEmpty(issue.RootCauseSummary))
 			{
-				_logger.LogInformation("Cancelling escalation for {NotificationName} {IssueId} {Reason}", notificationName, issueId, "RCA Completed");
+				_logger.LogInformation(CanceledMessage, notificationName, issueId, "RCA Completed");
 				return null;
 			}
 
@@ -3118,36 +3120,36 @@ namespace HordeServer.Notifications.Sinks
 			WorkflowId? workflowId = span.LastFailure.Annotations.WorkflowId;
 			if (workflowId == null)
 			{
-				_logger.LogInformation("Cancelling escalation for {NotificationName} {IssueId} {Reason}", notificationName, issueId, "Cannot Find Issue WorkflowID");
+				_logger.LogInformation(CanceledMessage, notificationName, issueId, "Cannot Find Issue WorkflowID");
 				return null;
 			}
 
 			if (!_buildConfig.CurrentValue.TryGetStream(span.StreamId, out StreamConfig? streamConfig))
 			{
-				_logger.LogInformation("Cancelling escalation for {NotificationName} {IssueId} {Reason}", notificationName, issueId, "Cannot Find Issue StreamConfig");
+				_logger.LogInformation(CanceledMessage, notificationName, issueId, "Cannot Find Issue StreamConfig");
 				return null;
 			}
 			if (!streamConfig.TryGetWorkflow(workflowId.Value, out WorkflowConfig? workflow))
 			{
-				_logger.LogInformation("Cancelling escalation for {NotificationName} {IssueId} {Reason}", notificationName, issueId, "Cannot Find Issue StreamConfig Workflow");
+				_logger.LogInformation(CanceledMessage, notificationName, issueId, "Cannot Find Issue StreamConfig Workflow");
 				return null;
 			}
 			if (String.IsNullOrEmpty(workflow.TriageChannel) || workflow.EscalateRootCauseAnalysisTimes.Count == 0)
 			{
-				_logger.LogInformation("Cancelling escalation for {NotificationName} {IssueId} {Reason}", notificationName, issueId, "No Channel or Escalation Times");
+				_logger.LogInformation(CanceledMessage, notificationName, issueId, "No Channel or Escalation Times");
 				return null;
 			}
 
 			if (!issue.IsMarkedFixed())
 			{
-				_logger.LogInformation("Cancelling escalation for {NotificationName} {IssueId} {Reason}", notificationName, issueId, "Issue Not Fixed");
+				_logger.LogInformation(CanceledMessage, notificationName, issueId, "Issue Not Fixed");
 				return null;
 			}
 
 			MessageStateDocument? state = await GetMessageStateAsync(workflow.TriageChannel, GetTriageThreadEventId(issueId), cancellationToken);
 			if (state == null)
 			{
-				_logger.LogInformation("Cancelling escalation for {NotificationName} {IssueId} {Reason}", notificationName, issueId, "Cannot Find Issue Message State");
+				_logger.LogInformation(CanceledMessage, notificationName, issueId, "Cannot Find Issue Message State");
 				return null;
 			}
 
