[Horde][BuildHealth][Dashboard] Use rootCauseOwnerInfo exclusively
#Synced-CL Unknown

diff --git a/Engine/Source/Programs/Horde/HordeDashboard/src/components/IssueViewV2.tsx b/Engine/Source/Programs/Horde/HordeDashboard/src/components/IssueViewV2.tsx
--- a/Engine/Source/Programs/Horde/HordeDashboard/src/components/IssueViewV2.tsx
+++ b/Engine/Source/Programs/Horde/HordeDashboard/src/components/IssueViewV2.tsx
@@ -1359,10 +1359,10 @@ const IssueSummaryPanel: React.FC = () => {
          link: dashboard.swarmUrl ? `${dashboard.swarmUrl}/changes/${issue.rootCommitId}` : undefined
       })
    }
-     if (issue.rootCauseOwnerId) {
+     if (issue.rootCauseOwnerInfo?.name ) {
       items.push({
          title: "Root Cause Owner",
-         text: issue.rootCauseOwnerInfo?.name ?? "User not found"
+         text: issue.rootCauseOwnerInfo.name
       })
    }
    if (issue.rootCauseCategory) {
@@ -2353,7 +2353,7 @@ const MarkFixedModal: React.FC<{ onClose: () => void, autoFix?:boolean } > = ({
       rootCommitId: issue.rootCommitId,
       dupeId: issue.duplicateIssueId ? (issue.duplicateIssueId.toString() !== "0" ? issue.duplicateIssueId.toString() : undefined) : undefined,
       dupeCheck: issue.duplicateIssueId === 0 || !!issue.duplicateIssueId,
-      rootCauseOwnerId: issue.rootCauseOwnerId,
+      rootCauseOwnerId: issue.rootCauseOwnerInfo?.id,
       rootCauseCategory: issue.rootCauseCategory,
       rootCauseSummary: issue.rootCauseSummary
    });
@@ -2392,7 +2392,7 @@ const MarkFixedModal: React.FC<{ onClose: () => void, autoFix?:boolean } > = ({
 
       // Root Cause Analysis
       let rootCommitId = state.rootCommitId?.length ? state.rootCommitId : null;
-      const rootCauseOwnerId = state.rootCauseOwnerId?.length ? state.rootCauseOwnerId : issue.rootCauseOwnerId?.length ? "" : null;
+      const rootCauseOwnerId = state.rootCauseOwnerId?.length ? state.rootCauseOwnerId : issue.rootCauseOwnerInfo?.id?.length ? "" : null;
       const rootCauseCategory = state.rootCauseCategory?.length ? state.rootCauseCategory : "";
       const rootCauseSummary = state?.rootCauseSummary?.trim().length ? state.rootCauseSummary.trim() : "";
 
@@ -2454,7 +2454,7 @@ const MarkFixedModal: React.FC<{ onClose: () => void, autoFix?:boolean } > = ({
    const disableRoot = !!state.dupeCheck;
 
    // Unique set of users that are most likely to be the RCA owner
-   const rcaOwnerCandidates = [... new Set([ dashboard.userId, issue.rootCauseOwnerId ?? null, issue.ownerInfo?.id ?? null, issue.resolvedByInfo?.id ?? null])];
+   const rcaOwnerCandidates = [... new Set([ dashboard.userId, issue.rootCauseOwnerInfo?.id ?? null, issue.ownerInfo?.id ?? null, issue.resolvedByInfo?.id ?? null])];
 
 
    return <Modal className={hordeClasses.modal} isOpen={true} styles={{ main: { padding: 8, width: 540, height: height, minHeight: height } }}>
