Horde: Update release notes

#rnx

diff --git a/Engine/Source/Programs/Horde/Docs/ReleaseNotes.md b/Engine/Source/Programs/Horde/Docs/ReleaseNotes.md
--- a/Engine/Source/Programs/Horde/Docs/ReleaseNotes.md
+++ b/Engine/Source/Programs/Horde/Docs/ReleaseNotes.md
@@ -1,5 +1,17 @@
 # Release Notes
 
+## 2025-11-14
+
+* Fix compiler warning in IssueServiceTests (48281019)
+* Make ModTime flag disabled by default for PerforceMaterializer (48280995)
+* Add current user to dropdown hints for root cause owner (48275734)
+* Move Me option for Root Cause Owner to top of selection to match existing UI layout (48275511)
+* Avoid multiple occurences of the same user in Slack workflow root cause owner options (48275417)
+* Skip RCA escalation for issues resolved by Horde (48257790)
+* Fix issue determining latest change number in PerforceMaterializer
+  Use workspace root instead of stream. (48246872)
+* Display user name for rootCauseOwnerId in issue modal (48230199)
+
 ## 2025-11-13
 
 * Use a real stream client in PerforceMaterializer
