Force empty string input in EditableTextField to be resolved as undefined

#rnx

#Synced-CL 48288107

diff --git a/Engine/Source/Programs/Horde/HordeDashboard/plugins/test-automation/testAutomationCommon.tsx b/Engine/Source/Programs/Horde/HordeDashboard/plugins/test-automation/testAutomationCommon.tsx
--- a/Engine/Source/Programs/Horde/HordeDashboard/plugins/test-automation/testAutomationCommon.tsx
+++ b/Engine/Source/Programs/Horde/HordeDashboard/plugins/test-automation/testAutomationCommon.tsx
@@ -930,7 +930,7 @@ export const EditableTextField: React.FC<{ label: string, onChange: (value: stri
                         <TextField
                             key={`text-${loaded}`}
                             defaultValue={text}
-                            onNotifyValidationResult={(_, value) => { if (value !== text) { onChange(value); } setEnableEditing(false); }}
+                            onNotifyValidationResult={(_, value) => { if (!value) { value = undefined; } if (value !== text) { onChange(value); } setEnableEditing(false); }}
                             validateOnFocusOut
                             validateOnLoad={false}
                             multiline={multiline}
