[Horde][BuildHealth][Server] Improve issue matching across job steps in no suspects case
#Synced-CL Unknown

diff --git a/Engine/Source/Programs/Horde/Plugins/Build/HordeServer.Build.Tests/Issues/IssueServiceTests.cs b/Engine/Source/Programs/Horde/Plugins/Build/HordeServer.Build.Tests/Issues/IssueServiceTests.cs
--- a/Engine/Source/Programs/Horde/Plugins/Build/HordeServer.Build.Tests/Issues/IssueServiceTests.cs
+++ b/Engine/Source/Programs/Horde/Plugins/Build/HordeServer.Build.Tests/Issues/IssueServiceTests.cs
@@ -1956,6 +1956,31 @@ namespace HordeServer.Tests.Issues
 			Assert.AreEqual("Errors in Update Version Files and Compile UnrealHeaderTool Win64", issue.Summary);
 		}
 
+		[TestMethod]
+		public async Task HashedIssueTest4Async()
+		{
+			// Previous successful build on CL with no author, to test no suspects path
+			IJob job = CreateJob(MainStreamId, 118, "Compile Test", Graph);
+			await UpdateCompleteStepAsync(job, 0, 0, JobStepOutcome.Success);
+			await UpdateCompleteStepAsync(job, 0, 1, JobStepOutcome.Success);
+
+			// Errors in different steps of the same job, still no suspects
+			IJob job2 = CreateJob(MainStreamId, 119, "Compile Test", Graph);
+
+			await ParseEventsAsync(job2, 0, 0, new[] { "Warning: This is a warning from the editor" });
+			await UpdateCompleteStepAsync(job2, 0, 0, JobStepOutcome.Failure);
+
+			await ParseEventsAsync(job2, 0, 1, new[] { "Warning: This is a warning from the editor" });
+			await UpdateCompleteStepAsync(job2, 0, 1, JobStepOutcome.Failure);
+
+			// Issues in different steps of the same job are grouped together
+			IReadOnlyList<IIssue> issues = await IssueCollection.FindIssuesAsync();
+			Assert.AreEqual(1, issues.Count);
+
+			IIssue issue = issues[0];
+			Assert.AreEqual("Warnings in Update Version Files and Compile UnrealHeaderTool Win64", issue.Summary);
+		}
+
 		[TestMethod]
 		public async Task HashedIssueFileMetadataTestAsync()
 		{
diff --git a/Engine/Source/Programs/Horde/Plugins/Build/HordeServer.Build/Issues/IssueService.cs b/Engine/Source/Programs/Horde/Plugins/Build/HordeServer.Build/Issues/IssueService.cs
--- a/Engine/Source/Programs/Horde/Plugins/Build/HordeServer.Build/Issues/IssueService.cs
+++ b/Engine/Source/Programs/Horde/Plugins/Build/HordeServer.Build/Issues/IssueService.cs
@@ -1586,7 +1586,7 @@ namespace HordeServer.Issues
 		{
 			using TelemetrySpan tracerSpan = _tracer.StartActiveSpan($"{nameof(IssueService)}.{nameof(FindOrAddIssueForSpanAsync)}");
 			IReadOnlyList<IIssue> existingIssues;
-			if (span.LastSuccess == null)
+			if (span.LastSuccess == null) // TODO: Do we want to add 'span.Suspects.Count == 0' condition here too? 
 			{
 				existingIssues = await _issueCollection.FindIssuesAsync(streamId: span.StreamId, minCommitId: span.FirstFailure.CommitId, maxCommitId: span.FirstFailure.CommitId, resolved: false, cancellationToken: cancellationToken);
 			}
@@ -1615,7 +1615,7 @@ namespace HordeServer.Issues
 		{
 			using TelemetrySpan tracerSpan = _tracer.StartActiveSpan($"{nameof(IssueService)}.{nameof(FindOrAddIssueForSpanAsync)}.NewIssueSpan");
 			IReadOnlyList<IIssue> existingIssues;
-			if (span.LastSuccess == null)
+			if (span.LastSuccess == null || span.Suspects.Count == 0)
 			{
 				existingIssues = await _issueCollection.FindIssuesAsync(streamId: span.StreamId, minCommitId: span.FirstFailure.CommitId, maxCommitId: span.FirstFailure.CommitId, resolved: false, cancellationToken: cancellationToken);
 				LoggerUtils.MultiLog([(_logger, LogLevel.Debug), (issueLogger, LogLevel.Debug)], "Found {NumIssues} open issues at {Change} in {StreamId}", existingIssues.Count, span.FirstFailure.CommitId, span.StreamId);
