[Horde][BuildHealth][Dashboard] Toggle to filter out cancelled jobs from the toplevel history, allow for filtering of jobs based on whether it was scheduled or manually kicked, and a big fix for stale data bindings on D3.

#Synced-CL 48176906

diff --git a/Engine/Source/Programs/Horde/HordeDashboard/src/components/buildhealth/BuildHealthDataTypes.ts b/Engine/Source/Programs/Horde/HordeDashboard/src/components/buildhealth/BuildHealthDataTypes.ts
--- a/Engine/Source/Programs/Horde/HordeDashboard/src/components/buildhealth/BuildHealthDataTypes.ts
+++ b/Engine/Source/Programs/Horde/HordeDashboard/src/components/buildhealth/BuildHealthDataTypes.ts
@@ -9,9 +9,11 @@ export const TEMPLATES_URL_PARAM_NAME: string = "templates";
 export const LABELS_URL_PARAM_NAME: string = "labels";
 export const STEPS_URL_PARAM_NAME: string = "stepNames";
 export const STATES_URL_PARAM_NAME: string = "states";
+export const JOB_STARTS_URL_PARAM_NAME: string = "jobStarts";
 export const PARAMETER_KEY_PREFIX: string = "parameter_key";
 export const TIME_SPAN_URL_PARAM: string = "lastTimeRange";
 export const INCLUDE_PREFLIGHT_PARAM: string = "includePreflight";
+export const INCLUDE_CANCELLED_PARAM: string = "includeCancelled";
 export const DEBUG_MODE_PARAM: string = "debugMode";
 export const KEY_SEPARATOR: string = "::";
 
@@ -21,6 +23,7 @@ export const UI_DATE_ANCHORS_PARAM: string = "dateAnchors";
 export const LABELS_URL_PARAM_ALL_LABELS: string = "labels_AllKnownValues";
 export const STEPS_URL_PARAM_ALL_STEPS: string = "stepNames_AllKnownValues";
 export const STATES_URL_PARAM_ALL_STATES: string = "stateNames_AllKnownValues";
+export const JOB_START_URL_PARAM_ALL_NAMES: string = "jobStart_akv";
 
 /**
  * Convenience type to fold in extra metadata relevant for the Template. 
diff --git a/Engine/Source/Programs/Horde/HordeDashboard/src/components/buildhealth/BuildHealthOptions.ts b/Engine/Source/Programs/Horde/HordeDashboard/src/components/buildhealth/BuildHealthOptions.ts
--- a/Engine/Source/Programs/Horde/HordeDashboard/src/components/buildhealth/BuildHealthOptions.ts
+++ b/Engine/Source/Programs/Horde/HordeDashboard/src/components/buildhealth/BuildHealthOptions.ts
@@ -1,7 +1,7 @@
 // Copyright Epic Games, Inc. All Rights Reserved.
 
 import { makeObservable, observable, runInAction } from "mobx";
-import { DEBUG_MODE_PARAM, INCLUDE_PREFLIGHT_PARAM, LABELS_URL_PARAM_ALL_LABELS, LABELS_URL_PARAM_NAME, PROJECTS_URL_PARAM_NAME, QUERY_SCHEMA_PARAM_NAME, STATES_URL_PARAM_ALL_STATES, STATES_URL_PARAM_NAME, STEPS_URL_PARAM_ALL_STEPS, STEPS_URL_PARAM_NAME, STREAMS_URL_PARAM_NAME, TEMPLATES_URL_PARAM_NAME, TIME_SPAN_URL_PARAM, UI_DATE_ANCHORS_PARAM, UI_SUMMARY_WARNING_AS_ERROR_PARAM } from "../buildhealth/BuildHealthDataTypes";
+import { DEBUG_MODE_PARAM, INCLUDE_CANCELLED_PARAM, INCLUDE_PREFLIGHT_PARAM, JOB_START_URL_PARAM_ALL_NAMES, JOB_STARTS_URL_PARAM_NAME, LABELS_URL_PARAM_ALL_LABELS, LABELS_URL_PARAM_NAME, PROJECTS_URL_PARAM_NAME, QUERY_SCHEMA_PARAM_NAME, STATES_URL_PARAM_ALL_STATES, STATES_URL_PARAM_NAME, STEPS_URL_PARAM_ALL_STEPS, STEPS_URL_PARAM_NAME, STREAMS_URL_PARAM_NAME, TEMPLATES_URL_PARAM_NAME, TIME_SPAN_URL_PARAM, UI_DATE_ANCHORS_PARAM, UI_SUMMARY_WARNING_AS_ERROR_PARAM } from "../buildhealth/BuildHealthDataTypes";
 import { BuildHealthDataHandler } from "./BuildHealthDataHandler";
 import { decodeFullyQualifiedLabelName, decodeLabelKey, decodeStateKey, decodeStepKey, decodeTemplateKey, encodeFullyQualifiedLabelName, encodeLabelKey, encodeLabelKeyFromStrings, encodeStateKeyFromString, encodeStepKey, encodeStepNameFromStringsPreserveCase, encodeTemplateKey, encodeTemplateKeyFromStrings } from "./BuildHealthUtilities";
 import { Location } from 'react-router-dom';
@@ -36,8 +36,10 @@ type BuildHealthQueryParams = {
     labels: string[];
     stepNames: string[];
     stateNames: string[];
+    jobStartNames: string[];
     timeSpanKey: string;
     includePreflight: boolean;
+    includeCancelledJobs: boolean;
     debugMode: boolean;
     querySchemaVersion: number;
 };
@@ -127,11 +129,13 @@ export function toNavigationQuery(optionsState: BuildHealthOptionsState, uiOptio
     let labels: string = "";
     let steps: string = "";
     let states: string = "";
+    let jobStarts: string = "";
 
     projects = Object.keys(optionsState.enabledProjects).join(",");
     streams = optionsState.stepOutcomeEnabledStreamKeys.join(",");
     const lastTimeRange = optionsState.jobHistoryTimeSpan.key.toString();
     const includePreflight = optionsState.includePreflight;
+    const includeCancelledJobs = optionsState.includeCancelledJobs;
 
     // Job & Steps 
     if (querySchemaVersion === 1) {
@@ -181,14 +185,21 @@ export function toNavigationQuery(optionsState: BuildHealthOptionsState, uiOptio
             });
     }
 
+    // Job Starts
+    if (querySchemaVersion >= 5) {
+        jobStarts = optionsState.stepOutcomeEnabledJobStartMethods.length === ALL_JOB_START_METHODS_LIST.length ? JOB_START_URL_PARAM_ALL_NAMES : optionsState.stepOutcomeEnabledJobStartMethods.join(",");
+    }
+
     params.set(PROJECTS_URL_PARAM_NAME, LZString.compressToEncodedURIComponent(projects));
     params.set(STREAMS_URL_PARAM_NAME, LZString.compressToEncodedURIComponent(streams));
     params.set(TEMPLATES_URL_PARAM_NAME, LZString.compressToEncodedURIComponent(jobs));
     params.set(LABELS_URL_PARAM_NAME, LZString.compressToEncodedURIComponent(labels));
     params.set(STEPS_URL_PARAM_NAME, LZString.compressToEncodedURIComponent(steps));
     params.set(STATES_URL_PARAM_NAME, LZString.compressToEncodedURIComponent(states));
+    params.set(JOB_STARTS_URL_PARAM_NAME, LZString.compressToEncodedURIComponent(jobStarts));
     params.set(TIME_SPAN_URL_PARAM, LZString.compressToEncodedURIComponent(lastTimeRange));
     params.set(INCLUDE_PREFLIGHT_PARAM, String(includePreflight));
+    params.set(INCLUDE_CANCELLED_PARAM, String(includeCancelledJobs));
     params.set(QUERY_SCHEMA_PARAM_NAME, querySchemaVersion.toString());
 
     // Handle UI Option State
@@ -237,6 +248,7 @@ export function parseBuildHealthQueryParams(location: Location): { buildHealthQu
     let labels: string[] = [];
     let stepNames: string[] = [];
     let stateNames: string[] = [];
+    let jobStartNames: string[] = [];
 
     const getList = (key: string) => {
         const raw = searchParams.get(key);
@@ -307,6 +319,12 @@ export function parseBuildHealthQueryParams(location: Location): { buildHealthQu
         }
     }
 
+    // Job Starts
+    if (querySchemaVersion >= 5) {
+        jobStartNames = getList(JOB_STARTS_URL_PARAM_NAME);
+    }
+
+    // Time Spans
     let timeSpanKey: string = (() => {
         const raw = searchParams.get(TIME_SPAN_URL_PARAM);
         return raw ? LZString.decompressFromEncodedURIComponent(raw) ?? "" : "";
@@ -316,7 +334,9 @@ export function parseBuildHealthQueryParams(location: Location): { buildHealthQu
         timeSpanKey = DEFAULT_JOB_HISTORY_TIME_SPAN.key;
     }
 
+    // Toggle Query Params
     const includePreflight = getBoolParam(INCLUDE_PREFLIGHT_PARAM);
+    const includeCancelledJobs = getBoolParam(INCLUDE_CANCELLED_PARAM);
 
     // UI params
     const uiWarningsAsErrors = getBoolParam(UI_SUMMARY_WARNING_AS_ERROR_PARAM);
@@ -326,7 +346,7 @@ export function parseBuildHealthQueryParams(location: Location): { buildHealthQu
     const debugMode = getBoolParam(DEBUG_MODE_PARAM);
 
     return {
-        buildHealthQueryParams: { projectIds, streamIds, templateIds: templateIds, labels, stepNames, stateNames, timeSpanKey, includePreflight, debugMode, querySchemaVersion },
+        buildHealthQueryParams: { projectIds, streamIds, templateIds: templateIds, labels, stepNames, stateNames, jobStartNames, timeSpanKey, includePreflight, includeCancelledJobs, debugMode, querySchemaVersion },
         buildHealthUIParams: { showDateAnchor: uiDateAnchors, summaryWarningAsErrors: uiWarningsAsErrors }
     };
 }
@@ -372,6 +392,11 @@ export function upgradeBuildHealthQueryParams(params: BuildHealthQueryParams, cu
         clonedParams.stateNames.push(encodeStateKeyFromString(-1, STATES_URL_PARAM_ALL_STATES));
     }
 
+    // We must add in the job start names as no such filter existed before
+    if (clonedParams.querySchemaVersion < 5 && currentSchemaVersion >= 5) {
+        clonedParams.jobStartNames.push(...ALL_JOB_START_METHODS_LIST);
+    }
+
     return clonedParams;
 }
 
@@ -400,13 +425,17 @@ export async function loadBuildHealthOptionsFromParams(
         optionsWriter.setIncludePreflight(params.includePreflight);
     })
 
+    runInAction(() => {
+        optionsWriter.setIncludeCancelledJobs(params.includeCancelledJobs);
+    })
+
     if (params.timeSpanKey) {
         runInAction(() => {
             optionsWriter.setJobHistoryTimeSpan(JobHistoryTimeSpans.find(ts => ts.key === params.timeSpanKey)!);
         });
     }
 
-    // Set state
+    // Set States
     if (params.stateNames) {
         optionsWriter.clearStates();
 
@@ -420,6 +449,18 @@ export async function loadBuildHealthOptionsFromParams(
         }
     }
 
+    // Set Job Starts
+    if (params.jobStartNames) {
+        optionsWriter.clearJobStartMethods();
+        if (params.jobStartNames.length === 1 && params.jobStartNames[0].includes(JOB_START_URL_PARAM_ALL_NAMES)) {
+            optionsWriter.toggleAllJobStartMethods(ALL_JOB_START_METHODS_LIST, true);
+        } else if (params.jobStartNames.length > 0) {
+            for (let idx: number = 0; idx < params.jobStartNames.length; ++idx) {
+                optionsWriter.toggleJobStartMethod(params.jobStartNames[idx], params.jobStartNames[idx], true);
+            }
+        }
+    }
+
     // Get baseline project data first
     await handler.requestHierarchicalRefresh();
 
@@ -557,7 +598,9 @@ export class BuildHealthOptionStateReceipt {
     enabledLabels: ReadonlySet<string>;
     enabledSteps: ReadonlySet<string>;
     enabledStates: ReadonlySet<string>;
+    enabledJobStartMethods: ReadonlySet<string>;
     enabledPreflight: boolean;
+    enabledCancelledJobs: boolean;
     timespan: TimeSpan;
 
     constructor(
@@ -567,7 +610,9 @@ export class BuildHealthOptionStateReceipt {
         labels: Iterable<string> = [],
         steps: Iterable<string> = [],
         states: Iterable<string> = [],
+        jobStartMethods: Iterable<string> = [],
         enabledPreflight: boolean,
+        enabledCancelledJobs: boolean,
         timespan: TimeSpan
     ) {
         this.enabledProjects = new Set(projects);
@@ -576,7 +621,9 @@ export class BuildHealthOptionStateReceipt {
         this.enabledLabels = new Set(labels);
         this.enabledSteps = new Set(steps);
         this.enabledStates = new Set(states);
+        this.enabledJobStartMethods = new Set(jobStartMethods);
         this.enabledPreflight = enabledPreflight;
+        this.enabledCancelledJobs = enabledCancelledJobs;
         this.timespan = timespan;
     }
 
@@ -598,7 +645,9 @@ export class BuildHealthOptionStateReceipt {
             BuildHealthOptionStateReceipt.difference(next.enabledLabels, prev.enabledLabels),
             BuildHealthOptionStateReceipt.difference(next.enabledSteps, prev.enabledSteps),
             BuildHealthOptionStateReceipt.difference(next.enabledStates, prev.enabledStates),
+            BuildHealthOptionStateReceipt.difference(next.enabledJobStartMethods, prev.enabledJobStartMethods),
             next.enabledPreflight,
+            next.enabledCancelledJobs,
             next.timespan
         );
 
@@ -609,7 +658,9 @@ export class BuildHealthOptionStateReceipt {
             BuildHealthOptionStateReceipt.difference(prev.enabledLabels, next.enabledLabels),
             BuildHealthOptionStateReceipt.difference(prev.enabledSteps, next.enabledSteps),
             BuildHealthOptionStateReceipt.difference(prev.enabledStates, next.enabledStates),
+            BuildHealthOptionStateReceipt.difference(prev.enabledJobStartMethods, next.enabledJobStartMethods),
             prev.enabledPreflight,
+            prev.enabledCancelledJobs,
             prev.timespan
         );
 
@@ -635,6 +686,10 @@ export const ALL_STATES_TUPLE_LIST = Object.values(JobStepState).map(
     })
 );
 
+export const SCHEDULED_START_METHOD = "Scheduled";
+export const MANUAL_START_METHOD = "Manual";
+export const ALL_JOB_START_METHODS_LIST = [MANUAL_START_METHOD, SCHEDULED_START_METHOD];
+
 /**
  * Build Health Filter Options.
  */
@@ -647,6 +702,7 @@ export class BuildHealthOptionsState {
     @observable stepOutcomeEnabledLabelKeys: string[] = [];
     @observable stepOutcomeEnabledStepKeys: string[] = [];
     @observable stepOutcomeEnabledsStates: string[] = [];
+    @observable stepOutcomeEnabledJobStartMethods: string[] = [];
     @observable startDate!: Date;
     @observable endDate?: Date;
 
@@ -660,8 +716,10 @@ export class BuildHealthOptionsState {
     @observable enabledLabels: Record<string, string> = {};
     @observable enabledSteps: Record<string, string> = {};
     @observable enabledStates: Record<string, string> = {};
+    @observable enabledJobStartMethods: Record<string, string> = {};
     @observable jobHistoryTimeSpan: TimeSpan = DEFAULT_JOB_HISTORY_TIME_SPAN;
     @observable includePreflight: boolean = false;
+    @observable includeCancelledJobs: boolean = true;
     @observable debugMode: boolean = false;
 
     // #endregion -- User Selected Options --
@@ -693,7 +751,9 @@ export class BuildHealthOptionsState {
             Object.keys(this.enabledLabels),
             Object.keys(this.enabledSteps),
             Object.values(this.enabledStates),
+            Object.values(this.enabledJobStartMethods),
             this.includePreflight,
+            this.includeCancelledJobs,
             this.jobHistoryTimeSpan
         );
 
diff --git a/Engine/Source/Programs/Horde/HordeDashboard/src/components/buildhealth/BuildHealthOptionsController.ts b/Engine/Source/Programs/Horde/HordeDashboard/src/components/buildhealth/BuildHealthOptionsController.ts
--- a/Engine/Source/Programs/Horde/HordeDashboard/src/components/buildhealth/BuildHealthOptionsController.ts
+++ b/Engine/Source/Programs/Horde/HordeDashboard/src/components/buildhealth/BuildHealthOptionsController.ts
@@ -74,6 +74,12 @@ export interface IBuildHealthOptionWriter {
      */
     produceReceiptDiff(): BuildHealthOptionStateDiff;
 
+    /**
+     * Sets whether to filter out cancelled jobs, or not.
+     * @param includeCancelledJobs True to include cancelled jobs, false otherwise.
+     */
+    setIncludeCancelledJobs(includeCancelledJobs: boolean): void;
+
     /**
      * Sets whether to filter out preflight jobs, or not.
      * @param includePreflight True to include preflight job instances, false otherwise.
@@ -203,6 +209,7 @@ export interface IBuildHealthOptionWriter {
 
     /**
      * Toggles all provided states.
+     * @param states List of states to enable.
      * @param enabled True to enable, false otherwise.
      */
     toggleAllStates(states: string[], enabled: boolean): void;
@@ -215,6 +222,21 @@ export interface IBuildHealthOptionWriter {
      */
     toggleState(id: string, state: string, enabled: boolean): void;
 
+    /**
+     * Toggles all provided methods.
+     * @param methods List of methods to enable.
+     * @param enabled True to enable, false otherwise.
+     */
+    toggleAllJobStartMethods(methods: string[], enabled: boolean): void;
+
+    /**
+     * Toggles the state of the job start method.
+     * @param id The id of the start method to toggle.
+     * @param method The name of the start method.
+     * @param enabled Whether it's enabled or not.
+     */
+    toggleJobStartMethod(id: string, method: string, enabled: boolean): void;
+
     /**
      * Clears all currently selected options.
      */
@@ -255,6 +277,11 @@ export interface IBuildHealthOptionWriter {
      * Clear all the states tracked by the options.
      */
     clearStates(): void;
+
+    /**
+     * Clear all the job start methods tracked by the options.
+     */
+    clearJobStartMethods(): void;
 }
 
 export class BuildHealthOptionsWriter implements IBuildHealthOptionWriter {
@@ -288,6 +315,7 @@ export class BuildHealthOptionsWriter implements IBuildHealthOptionWriter {
         });
         this.clearStreams(true);
         this.clearStates();
+        this.clearJobStartMethods();
     }
 
     /**
@@ -445,6 +473,18 @@ export class BuildHealthOptionsWriter implements IBuildHealthOptionWriter {
         }
     }
 
+    /**
+     * @inheritdoc
+     */
+    @action
+    toggleJobStartMethod(id: string, method: string, enabled: boolean): void {
+        if (enabled) {
+            this.state.enabledJobStartMethods[id] = method;
+        } else {
+            delete this.state.enabledJobStartMethods[id];
+        }
+    }
+
     /**
      * @inheritdoc
      */
@@ -490,6 +530,17 @@ export class BuildHealthOptionsWriter implements IBuildHealthOptionWriter {
         }
     }
 
+    /**
+     * @inheritdoc
+     */
+    @action
+    toggleAllJobStartMethods(methods: string[], enabled: boolean) {
+        for (let idx: number = 0; idx < methods.length; ++idx) {
+            this.toggleJobStartMethod(methods[idx], methods[idx], enabled);
+        }
+    }
+
+
     /**
      * @inheritdoc
      */
@@ -498,6 +549,14 @@ export class BuildHealthOptionsWriter implements IBuildHealthOptionWriter {
         this.state.jobHistoryTimeSpan = time;
     }
 
+    /**
+     * @inheritdoc
+     */
+    @action
+    setIncludeCancelledJobs(includeCancelledJobs: boolean) {
+        this.state.includeCancelledJobs = includeCancelledJobs;
+    }
+
     /**
      * @inheritdoc
      */
@@ -590,6 +649,15 @@ export class BuildHealthOptionsWriter implements IBuildHealthOptionWriter {
         });
     }
 
+    /**
+     * @inheritdoc
+     */
+    clearJobStartMethods() {
+        Object.keys(this.state.enabledJobStartMethods).forEach(key => {
+            delete this.state.enabledJobStartMethods[key];
+        });
+    }
+
     // #region -- Private API --
 
     private toggleTemplateInternal(id: string, name: string, enabled: boolean) {
@@ -682,7 +750,7 @@ export class BuildHealthOptionsController {
 
     optionsChangeVersion = 0;
     uiOptionsChangeVersion = 0;
-    querySchemaVersion: number = 4;
+    querySchemaVersion: number = 5;
 
     // #endregion -- Public Members --
 
@@ -798,6 +866,7 @@ export class BuildHealthOptionsController {
             this.state.stepOutcomeEnabledLabelKeys = Object.keys(this.state.enabledLabels);
             this.state.stepOutcomeEnabledStepKeys = Object.keys(this.state.enabledSteps);
             this.state.stepOutcomeEnabledsStates = Object.keys(this.state.enabledStates);
+            this.state.stepOutcomeEnabledJobStartMethods = Object.keys(this.state.enabledJobStartMethods);
             this.state.startDate = new Date(new Date().valueOf() - (this.state.jobHistoryTimeSpan.minutes * 60000));
             this.state.endDate = new Date();
 
diff --git a/Engine/Source/Programs/Horde/HordeDashboard/src/components/buildhealth/BuildhealthView.tsx b/Engine/Source/Programs/Horde/HordeDashboard/src/components/buildhealth/BuildhealthView.tsx
--- a/Engine/Source/Programs/Horde/HordeDashboard/src/components/buildhealth/BuildhealthView.tsx
+++ b/Engine/Source/Programs/Horde/HordeDashboard/src/components/buildhealth/BuildhealthView.tsx
@@ -12,13 +12,13 @@ import { TopNav } from "../TopNav";
 import { StepOutcomeView } from "../buildhealth/stepoutcome/StepOutcome";
 import { BuildHealthDataHandler, DataHandlerRefreshRequest, EMPTY_REFRESH_REQUEST, expandRefreshRequest } from "./BuildHealthDataHandler";
 import { KEY_SEPARATOR, LabelRefData, PARAMETER_KEY_PREFIX, StepRefData, TemplateRefData } from "./BuildHealthDataTypes";
-import { ALL_STATES_TUPLE_LIST, BuildHealthOptionsState, BuildHealthOptionStateDiff, BuildHealthUIOptionsState, JobHistoryTimeSpans, loadBuildHealthOptionsFromParams, parseBuildHealthQueryParams, TimeSpan, toNavigationQuery, upgradeBuildHealthQueryParams } from "./BuildHealthOptions";
+import { ALL_JOB_START_METHODS_LIST, ALL_STATES_TUPLE_LIST, BuildHealthOptionsState, BuildHealthOptionStateDiff, BuildHealthUIOptionsState, JobHistoryTimeSpans, loadBuildHealthOptionsFromParams, MANUAL_START_METHOD, parseBuildHealthQueryParams, SCHEDULED_START_METHOD, TimeSpan, toNavigationQuery, upgradeBuildHealthQueryParams } from "./BuildHealthOptions";
 import { decodeStepKey, decodeTemplateKey, encodeFullyQualifiedLabelName, encodeLabelKey, encodeStepKey, encodeTemplateKey } from "./BuildHealthUtilities";
 import { stepOutcomeTableClasses } from "./stepoutcome/StepOutcomeSharedUIComponents";
 import { runInAction } from "mobx";
 import { BuildHealthOptionsController, StreamDescriptor, TemplateDescriptor } from "./BuildHealthOptionsController";
 import { StepOutcomeLegend } from "./stepoutcome/StepOutcomeLegend";
-import { StepOutcomeDataHandler } from "./stepoutcome/StepOutcomeDataHandler";
+import { StepOutcomeDataHandler, StepOutcomeFilters } from "./stepoutcome/StepOutcomeDataHandler";
 
 // #region -- Dropdown Components --
 
@@ -973,6 +973,94 @@ const StateDropdown: React.FC<{ handler: BuildHealthDataHandler, buildHealthOpti
     );
 });
 
+const StartDropdown: React.FC<{ handler: BuildHealthDataHandler, buildHealthOptions: BuildHealthOptionsController }> = observer(function ConstructProjectDropdown({ handler, buildHealthOptions }) {
+    let selectedJobStartKeys: string[] = Object.keys(buildHealthOptions.state.enabledJobStartMethods);
+
+    // #region -- Data Memos --
+
+    const items = useMemo<BuildHealthOptionData[]>(() =>
+        ALL_JOB_START_METHODS_LIST.map((entry) => ({
+            id: entry,
+            text: entry,
+            group: "Job Start Methods"
+        })),
+        []
+    );
+
+    const dropdownOptions = useMemo<IDropdownOption<BuildHealthOptionData>[]>(() =>
+        items.map(item => ({
+            key: item.id,
+            text: item.text,
+            data: item,
+            group: item.group
+        })),
+        [items]
+    );
+
+    // #endregion
+
+    // #region -- Callbacks --
+
+    const handleChange = useCallback(
+        (option?: IDropdownOption<BuildHealthOptionData>) => {
+            if (!option) return;
+
+            const invocationKey = option.key as string;
+            const invocationName = option.text;
+            const isSelected = !!option.selected;
+
+            buildHealthOptions
+                .getTransactionSession()
+                .toggleJobStartMethod(invocationKey, invocationName, isSelected);
+        },
+        [buildHealthOptions]
+    );
+
+    const handleSelectAll = useCallback(
+        (
+            option?: IDropdownOption<BuildHealthOptionData>,
+            optionSelectAllContext?: IDropdownOption<BuildHealthOptionData>[]
+        ) => {
+            if (!option) return;
+
+            const isSelected = !!option.selected;
+            const invocationkeys = optionSelectAllContext === undefined ? dropdownOptions.map(x => x.key as string) : optionSelectAllContext.map(x => x.key as string);
+
+            buildHealthOptions
+                .getTransactionSession()
+                .toggleAllJobStartMethods(invocationkeys, isSelected);
+        },
+        [buildHealthOptions, dropdownOptions]
+    );
+
+    // #endregion
+
+    // Dropdown params
+    const startDropdownParams: BuildHealthOptionList = {
+        label: "Select Job Start Methods(s)",
+        items: items,
+        toolTip: "Select an option from the dropdown."
+    };
+
+    return startDropdownParams.items.length > 0 && (
+        <Stack id={`start_dropdown`}>
+            <Label>Job Start(s)</Label>
+            <Stack horizontal tokens={{ childrenGap: 8 }} verticalAlign="end">
+                <MultiListParameter
+                    handler={handler}
+                    buildHealthOptions={buildHealthOptions}
+                    param={startDropdownParams}
+                    selectedKeys={selectedJobStartKeys}
+                    request={EMPTY_REFRESH_REQUEST}
+                    onChange={handleChange}
+                    enabledSelectAll={true}
+                    onSelectAll={handleSelectAll}
+                />
+            </Stack>
+        </Stack>
+    );
+});
+
 // #endregion -- Option Specific Dropdown Components --
 
 // #endregion -- Dropdown Components --
@@ -1014,6 +1102,7 @@ const BuildHealthUrlSync: React.FC<{ buildHealthController: BuildHealthOptionsCo
 const BuildHealthSidebar: React.FC<{ handler: BuildHealthDataHandler, buildHealthOptions: BuildHealthOptionsController }> = function ConstrucBuildHealthSidebar({ handler, buildHealthOptions }) {
     const { hordeClasses, modeColors } = getHordeStyling();
     const [timeSpanKey, setTimeSpanKey] = useState(() => buildHealthOptions.state.jobHistoryTimeSpan.key);
+    const [includeCancelledJobs, setIncludeCancelledJobs] = useState(() => buildHealthOptions.state.includeCancelledJobs)
     const [includePreflight, setIncludePreflight] = useState(() => buildHealthOptions.state.includePreflight)
     const [linkedSelectionMode, setLinkedSelectionMode] = useState(() => buildHealthOptions.uiState.linkedSelectionModeEnabled)
 
@@ -1033,6 +1122,12 @@ const BuildHealthSidebar: React.FC<{ handler: BuildHealthDataHandler, buildHealt
         buildHealthOptions.commitTransactionSession();
     };
 
+    const updateIncludeCancelledJobs = (isChecked) => {
+        setIncludeCancelledJobs(isChecked);
+        buildHealthOptions.getTransactionSession().setIncludeCancelledJobs(isChecked);
+        buildHealthOptions.commitTransactionSession();
+    }
+
     const updateIncludePreflight = (isChecked) => {
         setIncludePreflight(isChecked);
         buildHealthOptions.getTransactionSession().setIncludePreflight(isChecked);
@@ -1062,6 +1157,10 @@ const BuildHealthSidebar: React.FC<{ handler: BuildHealthDataHandler, buildHealt
         setTimeSpanKey(buildHealthOptions.state.jobHistoryTimeSpan.key);
     }, [buildHealthOptions.state.jobHistoryTimeSpan.key]);
 
+    useEffect(() => {
+        setIncludeCancelledJobs(buildHealthOptions.state.includeCancelledJobs);
+    }, [buildHealthOptions.state.includeCancelledJobs]);
+
     useEffect(() => {
         setIncludePreflight(buildHealthOptions.state.includePreflight);
     }, [buildHealthOptions.state.includePreflight]);
@@ -1105,6 +1204,7 @@ const BuildHealthSidebar: React.FC<{ handler: BuildHealthDataHandler, buildHealt
                         <LabelDropdown handler={handler} buildHealthOptions={buildHealthOptions} />
                         <StepDropdown handler={handler} buildHealthOptions={buildHealthOptions} />
                         <StateDropdown handler={handler} buildHealthOptions={buildHealthOptions} />
+                        <StartDropdown handler={handler} buildHealthOptions={buildHealthOptions} />
                         <TooltipHost content={
                             <Stack style={{ maxWidth: 300 }}>
                                 <Stack>
@@ -1222,6 +1322,23 @@ const BuildHealthSidebar: React.FC<{ handler: BuildHealthDataHandler, buildHealt
                             onChange={(_, isChecked) => updateIncludePreflight(!!isChecked)}
                         />
                     </TooltipHost>
+                    <TooltipHost content={
+                        <div style={{ maxWidth: 300 }}>
+                            <div>
+                                Include cancelled jobs in the job history results, or not.
+                            </div>
+                        </div>
+                    }
+                        calloutProps={{
+                            directionalHint: DirectionalHint.rightCenter
+                        }}
+                    >
+                        <Checkbox
+                            label="Include Cancelled Jobs"
+                            checked={includeCancelledJobs}
+                            onChange={(_, isChecked) => updateIncludeCancelledJobs(!!isChecked)}
+                        />
+                    </TooltipHost>
                 </Stack>
                 <Stack
                     styles={{
@@ -1253,22 +1370,26 @@ export const BuildHealthView: React.FC = observer(function ConstructBuildHealthV
         // eslint-disable-next-line react-hooks/exhaustive-deps -- All necessary dependencies included for MobX observables
         [buildHealthState.startDate, buildHealthState.endDate]);
 
-    const filter = useMemo(() => ({
+    const filter: StepOutcomeFilters = useMemo(() => ({
         streams: buildHealthState.stepOutcomeEnabledStreamKeys,
         jobs: buildHealthState.stepOutcomeEnabledJobKeys,
         steps: buildHealthState.stepOutcomeEnabledStepKeys,
         labels: buildHealthState.stepOutcomeEnabledLabelKeys,
         jobHistorySpan,
         includePreflights: buildHealthState.includePreflight,
-        debugMode: buildHealthState.debugMode
+        includeCancelledJobs: buildHealthState.includeCancelledJobs,
+        debugMode: buildHealthState.debugMode,
+        jobStartFilter: { supportsManualStarts: buildHealthState.stepOutcomeEnabledJobStartMethods.includes(MANUAL_START_METHOD), supportsScheduledStarts: buildHealthState.stepOutcomeEnabledJobStartMethods.includes(SCHEDULED_START_METHOD) }
     }),
         // eslint-disable-next-line react-hooks/exhaustive-deps -- All necessary dependencies included for MobX observables
         [
             buildHealthState.stepOutcomeEnabledStreamKeys,
             buildHealthState.stepOutcomeEnabledJobKeys,
             buildHealthState.stepOutcomeEnabledStepKeys,
+            buildHealthState.stepOutcomeEnabledJobStartMethods,
             jobHistorySpan,
-            buildHealthState.includePreflight
+            buildHealthState.includePreflight,
+            buildHealthState.includeCancelledJobs
         ]);
 
     const uiOptions = useMemo(() => ({
diff --git a/Engine/Source/Programs/Horde/HordeDashboard/src/components/buildhealth/stepoutcome/StepOutcomeD3Component.tsx b/Engine/Source/Programs/Horde/HordeDashboard/src/components/buildhealth/stepoutcome/StepOutcomeD3Component.tsx
--- a/Engine/Source/Programs/Horde/HordeDashboard/src/components/buildhealth/stepoutcome/StepOutcomeD3Component.tsx
+++ b/Engine/Source/Programs/Horde/HordeDashboard/src/components/buildhealth/stepoutcome/StepOutcomeD3Component.tsx
@@ -10,6 +10,7 @@ import { getCellAccents, getCellClasses, getCellStyle, getColorRecords, Incremen
 import { getHordeTheme } from "horde/styles/theme";
 import { observer } from "mobx-react-lite";
 import { isTerminalState, toLocalISOString } from "./StepOutcomeUtilities";
+import { encodeStepNameFromStrings } from "../BuildHealthUtilities";
 
 // #region -- Styles --
 
@@ -620,7 +621,6 @@ function GenerateD3Table(stepOutcomeTable: StepOutcomeTable, handler: StepOutcom
 
                         // UE-324856 - This is a temporary fix until the horde server is fixed.
                         const original = new Date(entry.changeDate!);
-                        const corrected = new Date(original.getTime() - 4 * 60 * 60 * 1000);
                         div.text(""),
                             div.style("cursor", "pointer")
                                 .style("background-color", getCellStyle(entry).bgColor)
@@ -635,8 +635,8 @@ function GenerateD3Table(stepOutcomeTable: StepOutcomeTable, handler: StepOutcom
                                     tooltipContent.appendChild(changeDiv);
 
                                     const dateDiv = document.createElement("div");
-                                    dateDiv.innerHTML = `<strong>Date:</strong><p class="${tooltipClasses.tooltipLine}">(UTC) ${corrected.toISOString() ?? "-"}</p>`;
-                                    dateDiv.innerHTML += `<p class="${tooltipClasses.tooltipLineWithBorder}">(Local) ${toLocalISOString(corrected) ?? "-"}</p>`
+                                    dateDiv.innerHTML = `<strong>Date:</strong><p class="${tooltipClasses.tooltipLine}">(UTC) ${original.toISOString() ?? "-"}</p>`;
+                                    dateDiv.innerHTML += `<p class="${tooltipClasses.tooltipLineWithBorder}">(Local) ${toLocalISOString(original) ?? "-"}</p>`
                                     tooltipContent.appendChild(dateDiv);
 
                                     if (isStepOutcome(entry)) {
@@ -737,10 +737,10 @@ function GenerateD3Table(stepOutcomeTable: StepOutcomeTable, handler: StepOutcom
             const tr = d3.select(this);
 
             const cellsRaw = stepOutcomeTable.getStepOutputTableEntries(encodeStepNameFromStepNameHeader(rowData));
-            const cells = cellsRaw ? [...cellsRaw] : [];
+            const cells: TableEntry[] = cellsRaw ? [...cellsRaw] : [];
 
             const dataCells = tr.selectAll("td.step-outcome-data")
-                .data(cells)
+                .data(cells, (d : TableEntry, i) => isSummary(d) ? `%summary-${i}`: `${d?.jobId}-${d?.id}`)
                 .join("td")
                 .attr("class", "step-outcome-data");
 
diff --git a/Engine/Source/Programs/Horde/HordeDashboard/src/components/buildhealth/stepoutcome/StepOutcomeDataHandler.ts b/Engine/Source/Programs/Horde/HordeDashboard/src/components/buildhealth/stepoutcome/StepOutcomeDataHandler.ts
--- a/Engine/Source/Programs/Horde/HordeDashboard/src/components/buildhealth/stepoutcome/StepOutcomeDataHandler.ts
+++ b/Engine/Source/Programs/Horde/HordeDashboard/src/components/buildhealth/stepoutcome/StepOutcomeDataHandler.ts
@@ -1,6 +1,6 @@
 // Copyright Epic Games, Inc. All Rights Reserved.
 
-import { GetJobResponse, GetJobStepRefResponse, IssueQuery, JobStepOutcome, JobStepState, JobStreamQuery, LabelState } from "horde/backend/Api";
+import { GetJobResponse, GetJobStepRefResponse, IssueQuery, JobState, JobStepOutcome, JobStepState, JobStreamQuery, LabelState } from "horde/backend/Api";
 import { PollBase } from "horde/backend/PollBase";
 import { GetJobStepRefResponseWrapper, GetStepResponseWrapper, JobResponseData, LabelOutcomeTableEntry, StepOutcomeTable, StepOutcomeTableEntry, UpdatedChangeSummaryData } from "./StepOutcomeDataTypes";
 import backend from "horde/backend";
@@ -87,6 +87,14 @@ export interface StepOutcomeViewOptions {
     validStates: Set<string>;
 }
 
+/**
+ * Data structure that controls the job start filter for the StepOutcome
+ */
+export interface StepOutcomeJobStartFilter {
+    supportsScheduledStarts: boolean;
+    supportsManualStarts: boolean;
+}
+
 /**
  * Data structure that controls the filtering & queries of the Data Handler.
  * @todo UE-309726 - This data structure can be simplified to reduce duplication, and reduce the need for the internal container to decode.
@@ -97,8 +105,10 @@ export interface StepOutcomeFilters {
     jobs?: string[];
     steps?: string[];
     labels?: string[];
+    jobStartFilter?: StepOutcomeJobStartFilter;
     maxJobCount?: number;
     includePreflights: boolean;
+    includeCancelledJobs: boolean;
     debugMode: boolean;
 }
 
@@ -193,6 +203,26 @@ export class StepOutcomeDataHandler extends PollBase {
                     continue;
                 }
 
+                // Filter cancelled jobs as appropriate
+                if (!this.activeFilter.includeCancelledJobs && (jobResponse.abortedByUserInfo !== null)) {
+                    continue;
+                }
+
+                // Filter for job start method
+                if (this.activeFilter.jobStartFilter) {
+                    const { supportsManualStarts, supportsScheduledStarts } = this.activeFilter.jobStartFilter;
+
+                    const isManual = jobResponse.startedByUserInfo !== null;
+                    const isScheduled = !isManual;
+
+                    if (
+                        (!supportsManualStarts && isManual) ||
+                        (!supportsScheduledStarts && isScheduled)
+                    ) {
+                        continue;
+                    }
+                }
+
                 const batchResponses = jobResponse.batches ?? [];
                 let streamStepSet = stepMap.get(jobResponse.streamId);
                 let streamLabelSet = labelMap.get(jobResponse.streamId);
