Horde: Update release notes

#rnx

diff --git a/Engine/Source/Programs/Horde/Docs/ReleaseNotes.md b/Engine/Source/Programs/Horde/Docs/ReleaseNotes.md
--- a/Engine/Source/Programs/Horde/Docs/ReleaseNotes.md
+++ b/Engine/Source/Programs/Horde/Docs/ReleaseNotes.md
@@ -11,6 +11,8 @@
 * Fix issue determining latest change number in PerforceMaterializer
   Use workspace root instead of stream. (48246872)
 * Display user name for rootCauseOwnerId in issue modal (48230199)
+* Improve issue matching across job steps in no suspects case (48293032)
+* Use rootCauseOwnerInfo exclusively (48290903)
 
 ## 2025-11-13
 
