[Horde][BuildHealth][Dashboard] Linked mode QoL features:
- Added link action button to streams
- Cascade link action button to the entire hierarchy (stream -> template -> step & label
- Edge case saitated where when a stream is added, implicitly link all the hierarhcy
-- Only on the affirmative case, we do not issue this for the removal case (otherwise we would remove all hierarchy of the newly added stream)
- Optimizations made to only issue hierarchy refresh when receipts have diffs

#Synced-CL 48176906

diff --git a/Engine/Source/Programs/Horde/HordeDashboard/src/components/buildhealth/BuildHealthOptions.ts b/Engine/Source/Programs/Horde/HordeDashboard/src/components/buildhealth/BuildHealthOptions.ts
--- a/Engine/Source/Programs/Horde/HordeDashboard/src/components/buildhealth/BuildHealthOptions.ts
+++ b/Engine/Source/Programs/Horde/HordeDashboard/src/components/buildhealth/BuildHealthOptions.ts
@@ -586,6 +586,7 @@ export interface BuildHealthOptionStateDiff {
     removed?: BuildHealthOptionStateReceipt;
     enabledPreflight?: FlagDiff<boolean>;
     timespan?: FlagDiff<TimeSpan>;
+    hasDiff?: boolean;
 }
 
 /**
@@ -666,8 +667,23 @@ export class BuildHealthOptionStateReceipt {
 
         const enabledPreflight: FlagDiff<boolean> = { old: prev.enabledPreflight, new: next.enabledPreflight };
         const timespan: FlagDiff<TimeSpan> = { old: prev.timespan, new: next.timespan };
+        const preflightChanged = prev.enabledPreflight !== next.enabledPreflight;
+        const timespanChanged = prev.timespan !== next.timespan;
+        const hasDiff = this.receiptHasChanges(added) || this.receiptHasChanges(removed) || preflightChanged || timespanChanged;
 
-        return { added, removed, enabledPreflight, timespan };
+        return { added, removed, enabledPreflight, timespan, hasDiff };
+    }
+
+    private static receiptHasChanges(r: BuildHealthOptionStateReceipt): boolean {
+        return (
+            r.enabledProjects.size > 0 ||
+            r.enabledStreams.size > 0 ||
+            r.enabledTemplates.size > 0 ||
+            r.enabledLabels.size > 0 ||
+            r.enabledSteps.size > 0 ||
+            r.enabledStates.size > 0 ||
+            r.enabledJobStartMethods.size > 0
+        );
     }
 
     private static difference<T>(a: ReadonlySet<T>, b: ReadonlySet<T>): Set<T> {
diff --git a/Engine/Source/Programs/Horde/HordeDashboard/src/components/buildhealth/BuildHealthOptionsController.ts b/Engine/Source/Programs/Horde/HordeDashboard/src/components/buildhealth/BuildHealthOptionsController.ts
--- a/Engine/Source/Programs/Horde/HordeDashboard/src/components/buildhealth/BuildHealthOptionsController.ts
+++ b/Engine/Source/Programs/Horde/HordeDashboard/src/components/buildhealth/BuildHealthOptionsController.ts
@@ -140,8 +140,10 @@ export interface IBuildHealthOptionWriter {
     /**
      * Toggles all linkable templates, for all active streams.
      * @param enabled Whether it's enabled or not.
+     * @param queryOnly Whether to query the linkability of templates solely, or not. If false, apply the toggle action.
+     * @returns Whether a linking action is required happened.
      */
-    toggleAllLinkableTemplates(enabled: boolean): void;
+    toggleAllLinkableTemplates(enabled: boolean, queryOnly?: boolean): boolean;
 
     /**
      * Toggles all provided templates for a given stream.
@@ -162,8 +164,10 @@ export interface IBuildHealthOptionWriter {
     /**
      * Toggles all linkable steps, for all active templates.
      * @param enabled True if we want to enable, false otherwise.
+     * @param queryOnly Whether to query the linkability of steps solely, or not. If false, apply the toggle action.
+     * @returns Whether a linking action is required happened.
      */
-    toggleAllLinkableSteps(enabled: boolean): void;
+    toggleAllLinkableSteps(enabled: boolean, queryOnly?: boolean): boolean;
 
     /**
      * Toggles all provided steps for a given stream & template, given they are valid.
@@ -178,8 +182,10 @@ export interface IBuildHealthOptionWriter {
     /**
      * Toggles all linkable labels, for all active templates.
      * @param enabled True if we want to enable, false otherwise.
+     * @param queryOnly Whether to query the linkability of labels solely, or not. If false, apply the toggle action.
+     * @returns Whether a linking action is required happened.
      */
-    toggleAllLinkableLabels(enabled: boolean): void;
+    toggleAllLinkableLabels(enabled: boolean, queryOnly?: boolean): boolean;
 
     /**
      * Toggles all provided labels for a given stream & template, given they are valid.
@@ -404,39 +410,49 @@ export class BuildHealthOptionsWriter implements IBuildHealthOptionWriter {
      * @inheritdoc
      */
     @action
-    toggleAllLinkableTemplates(enabled: boolean) {
+    toggleAllLinkableTemplates(enabled: boolean, queryOnly: boolean = false) {
+        let requiresLinking: boolean = false;
         if (this.uiState.linkedSelectionModeEnabled) {
             Object.keys(this.state.enabledTemplates).forEach(templateKey => {
                 let templateName = this.state.enabledTemplates[templateKey];
-                this.toggleLinkedTemplates(templateKey, templateName, enabled);
+                requiresLinking ||= this.toggleLinkedTemplates(templateKey, templateName, enabled, queryOnly);
             });
         }
+
+        return requiresLinking;
     }
 
     /**
      * @inheritdoc
      */
     @action
-    toggleAllLinkableLabels(enabled: boolean) {
+    toggleAllLinkableLabels(enabled: boolean, queryOnly: boolean = false): boolean {
+        let requiresLinking: boolean = false;
         if (this.uiState.linkedSelectionModeEnabled) {
             Object.keys(this.state.enabledLabels).forEach(labelKey => {
                 let fullyQualifiedLabelName = this.state.enabledLabels[labelKey];
-                this.toggleLinkedLabels(labelKey, fullyQualifiedLabelName, enabled);
+                requiresLinking ||= this.toggleLinkedLabels(labelKey, fullyQualifiedLabelName, enabled, queryOnly);
             });
         }
+
+        return requiresLinking;
     }
 
     /**
      * @inheritdoc
      */
     @action
-    toggleAllLinkableSteps(enabled: boolean) {
+    toggleAllLinkableSteps(enabled: boolean, queryOnly: boolean = false): boolean {
+        let requiresLinking: boolean = false;
         if (this.uiState.linkedSelectionModeEnabled) {
             Object.keys(this.state.enabledSteps).forEach(stepKey => {
                 let stepName = this.state.enabledSteps[stepKey];
-                this.toggleLinkedSteps(stepKey, stepName, enabled);
+                let result = this.toggleLinkedSteps(stepKey, stepName, enabled, queryOnly);
+                requiresLinking ||= result;
             });
         }
+
+        return requiresLinking;
     }
 
     /**
@@ -671,16 +687,25 @@ export class BuildHealthOptionsWriter implements IBuildHealthOptionWriter {
         }
     }
 
-    private toggleLinkedTemplates(sourceId: string, name: string, enabled: boolean) {
+    private toggleLinkedTemplates(sourceId: string, name: string, enabled: boolean, queryOnly: boolean = false): boolean {
         let { streamId, templateId } = decodeTemplateKey(sourceId);
+        let requiresLink: boolean = false;
         Object.keys(this.state.enabledStreams).forEach(streamKey => {
             if (streamKey !== streamId) {
                 let linkedTempalteId = encodeTemplateKeyFromStrings(streamKey, templateId);
-                if (this.hierarchyValidator?.isValidTemplate(streamKey, templateId)) { // we need to see if this is a valid pair?
-                    this.toggleTemplateInternal(linkedTempalteId, name, enabled);
+                if (this.hierarchyValidator?.isValidTemplate(streamKey, templateId)) {
+                    let areDifferentStates: boolean = this.state.enabledTemplates[linkedTempalteId] != this.state.enabledTemplates[sourceId];
+                    if (areDifferentStates) {
+                        if (!queryOnly) {
+                            this.toggleTemplateInternal(linkedTempalteId, name, enabled);
+                        } else {
+                            requiresLink ||= true;
+                        }
+                    }
                 }
             }
         });
+        return requiresLink;
     }
 
     private toggleLabelInternal(id: string, name: string, enabled: boolean) {
@@ -691,18 +716,27 @@ export class BuildHealthOptionsWriter implements IBuildHealthOptionWriter {
         }
     }
 
-    private toggleLinkedLabels(sourceId: string, name: string, enabled: boolean) {
+    private toggleLinkedLabels(sourceId: string, name: string, enabled: boolean, queryOnly: boolean = false): boolean {
         let sourceDecodedLabelKey = decodeLabelKey(sourceId);
+        let requiresLink: boolean = false;
         Object.keys(this.state.enabledTemplates).forEach(key => {
             let targetDecodedTempalteKey = decodeTemplateKey(key);
-            // Set this step on every other enabled template, except the narrow case of the stream-template-step - we've just come from there.
+            // Set this label on every other enabled template, except the narrow case of the stream-template-label - we've just come from there.
             if (targetDecodedTempalteKey.templateId !== sourceDecodedLabelKey.templateId || targetDecodedTempalteKey.streamId !== sourceDecodedLabelKey.streamId) {
-                let linkedEncodedStepId = encodeLabelKeyFromStrings(targetDecodedTempalteKey.streamId, targetDecodedTempalteKey.templateId, sourceDecodedLabelKey.dashboardCategory, sourceDecodedLabelKey.dashboardName);
+                let linkedEncodedLabelId = encodeLabelKeyFromStrings(targetDecodedTempalteKey.streamId, targetDecodedTempalteKey.templateId, sourceDecodedLabelKey.dashboardCategory, sourceDecodedLabelKey.dashboardName);
                 if (this.hierarchyValidator?.isValidLabel(targetDecodedTempalteKey.streamId, targetDecodedTempalteKey.templateId, name)) {
-                    this.toggleLabelInternal(linkedEncodedStepId, name, enabled);
+                    let areDifferentStates: boolean = this.state.enabledLabels[linkedEncodedLabelId] != this.state.enabledLabels[sourceId];
+                    if (areDifferentStates) {
+                        if (!queryOnly) {
+                            this.toggleLabelInternal(linkedEncodedLabelId, name, enabled);
+                        } else {
+                            requiresLink ||= true;
+                        }
+                    }
                 }
             }
         });
+        return requiresLink;
     }
 
     private toggleStepInternal(id: string, name: string, enabled: boolean) {
@@ -713,18 +747,27 @@ export class BuildHealthOptionsWriter implements IBuildHealthOptionWriter {
         }
     }
 
-    private toggleLinkedSteps(sourceId: string, name: string, enabled: boolean) {
+    private toggleLinkedSteps(sourceId: string, name: string, enabled: boolean, queryOnly: boolean = false): boolean {
         let sourceDecodedStepKey = decodeStepKey(sourceId);
+        let requiresLink: boolean = false;
         Object.keys(this.state.enabledTemplates).forEach(key => {
             let targetDecodedTempalteKey = decodeTemplateKey(key);
             // Set this step on every other enabled template, except the narrow case of the stream-template-step - we've just come from there.
             if (targetDecodedTempalteKey.templateId !== sourceDecodedStepKey.templateId || targetDecodedTempalteKey.streamId !== sourceDecodedStepKey.streamId) {
                 let linkedEncodedStepId = encodeStepNameFromStringsPreserveCase(targetDecodedTempalteKey.streamId, targetDecodedTempalteKey.templateId, sourceDecodedStepKey.stepName);
-                if (this.hierarchyValidator?.isValidStep(targetDecodedTempalteKey.streamId, targetDecodedTempalteKey.templateId, sourceDecodedStepKey.stepName)) { // we need to see if this is a valid pair?
-                    this.toggleStepInternal(linkedEncodedStepId, name, enabled);
+                if (this.hierarchyValidator?.isValidStep(targetDecodedTempalteKey.streamId, targetDecodedTempalteKey.templateId, sourceDecodedStepKey.stepName)) {
+                    let areDifferentStates: boolean = this.state.enabledSteps[linkedEncodedStepId] != this.state.enabledSteps[sourceId];
+                    if (areDifferentStates) {
+                        if (!queryOnly) {
+                            this.toggleStepInternal(linkedEncodedStepId, name, enabled);
+                        } else {
+                            requiresLink ||= true;
+                        }
+                    }
                 }
             }
         });
+        return requiresLink;
     }
 
     // #endregion -- Private API --
@@ -840,14 +883,16 @@ export class BuildHealthOptionsController {
             };
         }
 
-        let finalReceiptes = this.writer.produceReceiptDiff();
+        let finalReceipts = this.writer.produceReceiptDiff();
         this.writer = null;
 
-        // We can be smarter here, and only only request synchronization if we need to (by checking the option state diff).
-        this.setOptionsChanged();
-        this.synchronizeDerivedKeys();
+        // Only request a synchronization if we have produced diffs in our receipt.
+        if (finalReceipts.hasDiff) {
+            this.setOptionsChanged();
+            this.synchronizeDerivedKeys();
+        }
 
-        return finalReceiptes;
+        return finalReceipts;
     }
 
     // #endregion -- Public API --
diff --git a/Engine/Source/Programs/Horde/HordeDashboard/src/components/buildhealth/BuildhealthView.tsx b/Engine/Source/Programs/Horde/HordeDashboard/src/components/buildhealth/BuildhealthView.tsx
--- a/Engine/Source/Programs/Horde/HordeDashboard/src/components/buildhealth/BuildhealthView.tsx
+++ b/Engine/Source/Programs/Horde/HordeDashboard/src/components/buildhealth/BuildhealthView.tsx
@@ -15,7 +15,7 @@ import { KEY_SEPARATOR, LabelRefData, PARAMETER_KEY_PREFIX, StepRefData, Templat
 import { ALL_JOB_START_METHODS_LIST, ALL_STATES_TUPLE_LIST, BuildHealthOptionsState, BuildHealthOptionStateDiff, BuildHealthUIOptionsState, JobHistoryTimeSpans, loadBuildHealthOptionsFromParams, MANUAL_START_METHOD, parseBuildHealthQueryParams, SCHEDULED_START_METHOD, TimeSpan, toNavigationQuery, upgradeBuildHealthQueryParams } from "./BuildHealthOptions";
 import { decodeStepKey, decodeTemplateKey, encodeFullyQualifiedLabelName, encodeLabelKey, encodeStepKey, encodeTemplateKey } from "./BuildHealthUtilities";
 import { stepOutcomeTableClasses } from "./stepoutcome/StepOutcomeSharedUIComponents";
-import { runInAction } from "mobx";
+import { makeAutoObservable, observable, runInAction } from "mobx";
 import { BuildHealthOptionsController, StreamDescriptor, TemplateDescriptor } from "./BuildHealthOptionsController";
 import { StepOutcomeLegend } from "./stepoutcome/StepOutcomeLegend";
 import { StepOutcomeDataHandler, StepOutcomeFilters } from "./stepoutcome/StepOutcomeDataHandler";
@@ -501,7 +501,7 @@ const ProjectDropdownSingle: React.FC<{ handler: BuildHealthDataHandler, buildHe
     );
 });
 
-const StreamDropdown: React.FC<{ handler: BuildHealthDataHandler, buildHealthOptions: BuildHealthOptionsController }> = observer(function ConstructProjectDropdown({ handler, buildHealthOptions }) {
+const StreamDropdown: React.FC<{ handler: BuildHealthDataHandler, buildHealthOptions: BuildHealthOptionsController, buildHealthLinkedModeNotepad: BuildHealthLinkNotepad }> = observer(function ConstructProjectDropdown({ handler, buildHealthOptions, buildHealthLinkedModeNotepad }) {
     const streamDropdownParams: BuildHealthOptionList = {
         label: "Select Stream(s)",
         items: [
@@ -542,33 +542,52 @@ const StreamDropdown: React.FC<{ handler: BuildHealthDataHandler, buildHealthOpt
                         .map((x) => ({ id: x.data!.id, name: x.data!.text! }));
 
             // Do not close the transaction session; the user still has the modal open
-            buildHealthOptions
-                .getTransactionSession()
-                .toggleAllStreams(streamIds, isSelected);
+            buildHealthOptions.getTransactionSession().toggleAllStreams(streamIds, isSelected);
+
+            buildHealthLinkedModeNotepad.dirtyTemplate(true);
         },
         [buildHealthOptions, handler]
     );
 
     let selectedStreamKeys: string[] = Object.keys(buildHealthOptions.state.enabledStreams);
     let request: DataHandlerRefreshRequest = expandRefreshRequest({ streams: true });
+    let linkString = buildHealthLinkedModeNotepad.releaseDirty ? "Current State: Unlinked - Some templates, steps, and labels may require linking." : "Current State: Linked";
 
     return streamDropdownParams.items.length > 0 && (
         <Stack>
             <Label>Stream(s)</Label>
-            <MultiListParameter handler={handler} buildHealthOptions={buildHealthOptions} param={streamDropdownParams} selectedKeys={selectedStreamKeys} request={request}
-                enabledSelectAll={true}
-                onSelectAll={handleSelectAll}
-                onChange={(option: IDropdownOption<BuildHealthOptionData>) => {
-                    const key = option.key as string;
-                    const isSelected = !!option.selected;
-
-                    buildHealthOptions.getTransactionSession().toggleStream(key, option.text, isSelected);
-                }} />
+            <Stack horizontal tokens={{ childrenGap: 8 }} verticalAlign="end">
+                <MultiListParameter handler={handler} buildHealthOptions={buildHealthOptions} param={streamDropdownParams} selectedKeys={selectedStreamKeys} request={request}
+                    enabledSelectAll={true}
+                    onSelectAll={handleSelectAll}
+                    onChange={(option: IDropdownOption<BuildHealthOptionData>) => {
+                        const key = option.key as string;
+                        const isSelected = !!option.selected;
+
+                        buildHealthOptions.getTransactionSession().toggleStream(key, option.text, isSelected);
+
+                        // Workflow special case to cascade all links, and only in the affirmitive case otherwise we will cascade disables
+                        if (buildHealthOptions.uiState.linkedSelectionModeEnabled && isSelected) {
+                            let linkRequest: DataHandlerRefreshRequest = expandRefreshRequest({ jobs: true });
+
+                            // We do not want to re-issue select all steps for all newly linked templates, as it can be quite destructive.
+                            buildHealthOptions.getTransactionSession().toggleAllLinkableTemplates(true);
+
+                            // We commit a hierarchical refresh so we can build any newly added template's step & label list
+                            handler.requestHierarchicalRefresh(linkRequest, true, (request) => {
+                                buildHealthOptions.getTransactionSession().toggleAllLinkableSteps(true);
+                                buildHealthOptions.getTransactionSession().toggleAllLinkableLabels(true);
+
+                                buildHealthLinkedModeNotepad.dirtyTemplate(false);
+                            });
+                        }
+                    }} />
+            </Stack>
         </Stack>
     );
 });
 
-const TemplateDropdown: React.FC<{ handler: BuildHealthDataHandler, buildHealthOptions: BuildHealthOptionsController }> = observer(function ConstructProjectDropdown({ handler, buildHealthOptions }) {
+const TemplateDropdown: React.FC<{ handler: BuildHealthDataHandler, buildHealthOptions: BuildHealthOptionsController, buildHealthLinkedModeNotepad: BuildHealthLinkNotepad }> = observer(function ConstructProjectDropdown({ handler, buildHealthOptions, buildHealthLinkedModeNotepad }) {
     const templateDropdownParams: BuildHealthOptionList = {
         label: "Select Template(s)",
         items: [
@@ -621,6 +640,7 @@ const TemplateDropdown: React.FC<{ handler: BuildHealthDataHandler, buildHealthO
 
     let selectedTemplateKeys: string[] = Object.keys(buildHealthOptions.state.enabledTemplates);
     let request: DataHandlerRefreshRequest = expandRefreshRequest({ jobs: true });
+    let linkString = buildHealthLinkedModeNotepad.labelDirty ? "Current State: Unlinked - Some templates, steps, and labels may require linking." : "Current State: Linked";
 
     return templateDropdownParams.items.length > 0 && (
         <Stack id={`job_dropdown`}>
@@ -631,20 +651,36 @@ const TemplateDropdown: React.FC<{ handler: BuildHealthDataHandler, buildHealthO
                         const key = option.key as string;
                         const isSelected = !!option.selected;
                         buildHealthOptions.getTransactionSession().toggleTemplate(key, option.text, isSelected);
+                        buildHealthLinkedModeNotepad.dirtyTemplateDependents(true);
                     }} />
                 {
                     buildHealthOptions.uiState.linkedSelectionModeEnabled &&
                     (
                         <TooltipHost
-                            content={"Link templates: Apply selected templates across selected streams."}
+                            content={
+                                <span>
+                                    Link templates: Apply selected templates across selected streams, as well as steps and labels.
+                                    <br /><br />
+                                    {linkString}
+                                </span>
+                            }
                             calloutProps={{ directionalHint: DirectionalHint.rightCenter }}>
                             <IconButton
-                                iconProps={{ iconName: 'Link' }}
+                                iconProps={{ iconName: buildHealthLinkedModeNotepad.templateDirty ? 'UnGrouped' : 'Grouped' }}
                                 styles={{ root: { width: 15, minWidth: 15, padding: 0 } }} onClick={() => {
+                                    let linkRequest: DataHandlerRefreshRequest = expandRefreshRequest({ streams: true });
                                     // We do not want to re-issue select all steps for all newly linked templates, as it can be quite destructive.
                                     buildHealthOptions.getTransactionSession().toggleAllLinkableTemplates(true);
-                                    buildHealthOptions.commitTransactionSession();
-                                    handler.requestHierarchicalRefresh(request);
+
+                                    // We commit a hierarchical refresh so we can build any newly added template's step & label list
+                                    handler.requestHierarchicalRefresh(linkRequest, true, (request) => {
+                                        buildHealthOptions.getTransactionSession().toggleAllLinkableSteps(true);
+                                        buildHealthOptions.getTransactionSession().toggleAllLinkableLabels(true);
+
+                                        buildHealthLinkedModeNotepad.dirtyTemplate(false);
+
+                                        buildHealthOptions.commitTransactionSession();
+                                    });
                                 }
                                 } />
                         </TooltipHost>
@@ -655,7 +691,7 @@ const TemplateDropdown: React.FC<{ handler: BuildHealthDataHandler, buildHealthO
     );
 });
 
-const StepDropdown: React.FC<{ handler: BuildHealthDataHandler, buildHealthOptions: BuildHealthOptionsController }> = observer(function ConstructProjectDropdown({ handler, buildHealthOptions }) {
+const StepDropdown: React.FC<{ handler: BuildHealthDataHandler, buildHealthOptions: BuildHealthOptionsController, buildHealthLinkedModeNotepad: BuildHealthLinkNotepad }> = observer(function ConstructProjectDropdown({ handler, buildHealthOptions, buildHealthLinkedModeNotepad }) {
     let selectedStepKeys: string[] = Object.keys(buildHealthOptions.state.enabledSteps);
     let request: DataHandlerRefreshRequest = { steps: false };
 
@@ -719,6 +755,7 @@ const StepDropdown: React.FC<{ handler: BuildHealthDataHandler, buildHealthOptio
             buildHealthOptions
                 .getTransactionSession()
                 .toggleAllSteps(streamId, templateId, stepNames, isSelected);
+
         },
         [buildHealthOptions, handler]
     );
@@ -728,7 +765,7 @@ const StepDropdown: React.FC<{ handler: BuildHealthDataHandler, buildHealthOptio
             // Close the transaction session; the modal is not open so there will be no such transaction completion otherwise.
             buildHealthOptions.getTransactionSession().toggleAllLinkableSteps(true);
             buildHealthOptions.commitTransactionSession();
-            handler.requestHierarchicalRefresh(request);
+            buildHealthLinkedModeNotepad.stepDirty = false;
         },
         [buildHealthOptions, handler]
     );
@@ -743,6 +780,7 @@ const StepDropdown: React.FC<{ handler: BuildHealthDataHandler, buildHealthOptio
     };
 
     stepDropdownParams.items = items;
+    let linkString = buildHealthLinkedModeNotepad.stepDirty ? "Current State: Unlinked - steps may require linking." : "Current State: Linked";
 
     return stepDropdownParams.items.length > 0 && (
         <Stack id={`step_dropdown`}>
@@ -756,10 +794,16 @@ const StepDropdown: React.FC<{ handler: BuildHealthDataHandler, buildHealthOptio
                     buildHealthOptions.uiState.linkedSelectionModeEnabled &&
                     (
                         <TooltipHost
-                            content={"Link steps: Apply selected steps across selected templates."}
+                            content={
+                                <span>
+                                    Link steps: Apply selected steps across selected templates.
+                                    <br /><br />
+                                    {linkString}
+                                </span>
+                            }
                             calloutProps={{ directionalHint: DirectionalHint.rightCenter }}>
                             <IconButton
-                                iconProps={{ iconName: 'Link' }}
+                                iconProps={{ iconName: buildHealthLinkedModeNotepad.stepDirty ? 'UnGrouped' : 'Grouped' }}
                                 styles={{ root: { width: 15, minWidth: 15, padding: 0 } }}
                                 onClick={handleLinkSteps} />
                         </TooltipHost>
@@ -770,7 +814,7 @@ const StepDropdown: React.FC<{ handler: BuildHealthDataHandler, buildHealthOptio
     );
 });
 
-const LabelDropdown: React.FC<{ handler: BuildHealthDataHandler, buildHealthOptions: BuildHealthOptionsController }> = observer(function ConstructProjectDropdown({ handler, buildHealthOptions }) {
+const LabelDropdown: React.FC<{ handler: BuildHealthDataHandler, buildHealthOptions: BuildHealthOptionsController, buildHealthLinkedModeNotepad: BuildHealthLinkNotepad }> = observer(function ConstructProjectDropdown({ handler, buildHealthOptions, buildHealthLinkedModeNotepad }) {
     let selectedLabelKeys: string[] = Object.keys(buildHealthOptions.state.enabledLabels);
     let request: DataHandlerRefreshRequest = { steps: false };
 
@@ -838,12 +882,12 @@ const LabelDropdown: React.FC<{ handler: BuildHealthDataHandler, buildHealthOpti
         [buildHealthOptions, handler]
     );
 
-    const handleLinkSteps = useCallback(
+    const handleLinkLabels = useCallback(
         () => {
             // Close the transaction session; the modal is not open so there will be no such transaction completion otherwise.
             buildHealthOptions.getTransactionSession().toggleAllLinkableLabels(true);
             buildHealthOptions.commitTransactionSession();
-            handler.requestHierarchicalRefresh(request);
+            buildHealthLinkedModeNotepad.labelDirty = false;
         },
         [buildHealthOptions, handler]
     );
@@ -858,6 +902,7 @@ const LabelDropdown: React.FC<{ handler: BuildHealthDataHandler, buildHealthOpti
     };
 
     labelDropdownParams.items = items;
+    let linkString = buildHealthLinkedModeNotepad.templateDirty ? "Current State: Unlinked - labels may require linking." : "Current State: Linked";
 
     return labelDropdownParams.items.length > 0 && (
         <Stack id={`label_dropdown`}>
@@ -871,12 +916,18 @@ const LabelDropdown: React.FC<{ handler: BuildHealthDataHandler, buildHealthOpti
                     buildHealthOptions.uiState.linkedSelectionModeEnabled &&
                     (
                         <TooltipHost
-                            content={"Link labels: Apply selected labels across selected templates."}
+                            content={
+                                <span>
+                                    Link labels: Apply selected labels across selected templates.
+                                    <br /><br />
+                                    {linkString}
+                                </span>
+                            }
                             calloutProps={{ directionalHint: DirectionalHint.rightCenter }}>
                             <IconButton
-                                iconProps={{ iconName: 'Link' }}
+                                iconProps={{ iconName: buildHealthLinkedModeNotepad.labelDirty ? 'UnGrouped' : 'Grouped' }}
                                 styles={{ root: { width: 15, minWidth: 15, padding: 0 } }}
-                                onClick={handleLinkSteps} />
+                                onClick={handleLinkLabels} />
                         </TooltipHost>
                     )
                 }
@@ -1095,11 +1146,44 @@ const BuildHealthUrlSync: React.FC<{ buildHealthController: BuildHealthOptionsCo
     return null;
 };
 
+/**
+ * Data structure used to keep track of the current linked selection mode session.
+ */
+class BuildHealthLinkNotepad {
+    @observable releaseDirty: boolean = true;
+    @observable templateDirty: boolean = true;
+    @observable stepDirty: boolean = true;
+    @observable labelDirty: boolean = true;
+
+    constructor() {
+        makeAutoObservable(this);
+    }
+
+    dirtyTemplate(isDirty: boolean): void {
+        this.templateDirty = isDirty;
+        this.stepDirty = isDirty;
+        this.labelDirty = isDirty;
+    }
+
+    dirtyStep(isDirty: boolean): void {
+        this.stepDirty = isDirty;
+    }
+
+    dirtyLabel(isDirty: boolean): void {
+        this.labelDirty = isDirty;
+    }
+
+    dirtyTemplateDependents(isDirty: boolean) {
+        this.stepDirty = isDirty;
+        this.labelDirty = isDirty;
+    }
+}
+
 /**
  * React Component that represents the Build Health sidebar. All filters are present here.
  * @returns React Component.
  */
-const BuildHealthSidebar: React.FC<{ handler: BuildHealthDataHandler, buildHealthOptions: BuildHealthOptionsController }> = function ConstrucBuildHealthSidebar({ handler, buildHealthOptions }) {
+const BuildHealthSidebar: React.FC<{ handler: BuildHealthDataHandler, buildHealthOptions: BuildHealthOptionsController, buildHealthLinkedModeNotepad: BuildHealthLinkNotepad }> = function ConstrucBuildHealthSidebar({ handler, buildHealthOptions, buildHealthLinkedModeNotepad }) {
     const { hordeClasses, modeColors } = getHordeStyling();
     const [timeSpanKey, setTimeSpanKey] = useState(() => buildHealthOptions.state.jobHistoryTimeSpan.key);
     const [includeCancelledJobs, setIncludeCancelledJobs] = useState(() => buildHealthOptions.state.includeCancelledJobs)
@@ -1137,6 +1221,17 @@ const BuildHealthSidebar: React.FC<{ handler: BuildHealthDataHandler, buildHealt
     const updateLinkedSelectionMode = (isChecked) => {
         setLinkedSelectionMode(isChecked);
         buildHealthOptions.setLinkedSelectionMode(isChecked);
+
+        // Determine the linkable state of each of the primary sets.
+        let dirtyTemplates = buildHealthOptions.getTransactionSession().toggleAllLinkableTemplates(true, true);
+        let dirtyLabels = buildHealthOptions.getTransactionSession().toggleAllLinkableLabels(true, true);
+        let dirtySteps = buildHealthOptions.getTransactionSession().toggleAllLinkableSteps(true, true);
+
+        buildHealthLinkedModeNotepad.dirtyTemplate(dirtyTemplates || dirtyLabels || dirtySteps);
+        buildHealthLinkedModeNotepad.dirtyLabel(dirtyLabels);
+        buildHealthLinkedModeNotepad.dirtyStep(dirtySteps);
+
+        buildHealthOptions.commitTransactionSession();
     }
 
     const updateIncludeDateAnchors = (isChecked) => {
@@ -1199,10 +1294,10 @@ const BuildHealthSidebar: React.FC<{ handler: BuildHealthDataHandler, buildHealt
                         },
                     }}>
                         <ProjectDropdownSingle handler={handler} buildHealthOptions={buildHealthOptions} />
-                        <StreamDropdown handler={handler} buildHealthOptions={buildHealthOptions} />
-                        <TemplateDropdown handler={handler} buildHealthOptions={buildHealthOptions} />
-                        <LabelDropdown handler={handler} buildHealthOptions={buildHealthOptions} />
-                        <StepDropdown handler={handler} buildHealthOptions={buildHealthOptions} />
+                        <StreamDropdown handler={handler} buildHealthOptions={buildHealthOptions} buildHealthLinkedModeNotepad={buildHealthLinkedModeNotepad} />
+                        <TemplateDropdown handler={handler} buildHealthOptions={buildHealthOptions} buildHealthLinkedModeNotepad={buildHealthLinkedModeNotepad} />
+                        <LabelDropdown handler={handler} buildHealthOptions={buildHealthOptions} buildHealthLinkedModeNotepad={buildHealthLinkedModeNotepad} />
+                        <StepDropdown handler={handler} buildHealthOptions={buildHealthOptions} buildHealthLinkedModeNotepad={buildHealthLinkedModeNotepad} />
                         <StateDropdown handler={handler} buildHealthOptions={buildHealthOptions} />
                         <StartDropdown handler={handler} buildHealthOptions={buildHealthOptions} />
                         <TooltipHost content={
@@ -1416,7 +1511,7 @@ export const BuildHealthView: React.FC = observer(function ConstructBuildHealthV
                 <Stack horizontalAlign="center" grow styles={{ root: { width: "100%", padding: 12, backgroundColor: modeColors.background } }}>
                     <Stack styles={{ root: { width: "100%" } }}>
                         <Stack horizontal styles={{ root: { minHeight: '85vh' } }} >
-                            <BuildHealthSidebar handler={handler} buildHealthOptions={buildHealthController} />
+                            <BuildHealthSidebar handler={handler} buildHealthOptions={buildHealthController} buildHealthLinkedModeNotepad={buildHealthLinkedModeNotepad} />
                             <Stack id="parent-buildhealth-stepoutcome" className={stepOutcomeTableClasses.root} grow style={{ overflowX: "auto", overflowY: "visible", minWidth: "800px", position: "relative", border: `${modeColors.content} solid 2px` }}>
                                 <StepOutcomeView key="step-outcome-inline" filter={filter} uiOptions={uiOptions} />
                             </Stack>
@@ -1431,7 +1526,7 @@ export const BuildHealthView: React.FC = observer(function ConstructBuildHealthV
 // #endregion -- Primary View Components --
 
 // #region -- Script --
-
+const buildHealthLinkedModeNotepad = new BuildHealthLinkNotepad();
 const buildHealthState = new BuildHealthOptionsState();
 const buildHealthUIState = new BuildHealthUIOptionsState();
 const buildHealthController = new BuildHealthOptionsController(buildHealthState, buildHealthUIState);
