[Horde][BuildHealth][Dashboard] Adding additional style icons.

#Synced-CL 48176906

diff --git a/Engine/Source/Programs/Horde/HordeDashboard/src/styles/Styles.tsx b/Engine/Source/Programs/Horde/HordeDashboard/src/styles/Styles.tsx
--- a/Engine/Source/Programs/Horde/HordeDashboard/src/styles/Styles.tsx
+++ b/Engine/Source/Programs/Horde/HordeDashboard/src/styles/Styles.tsx
@@ -110,7 +110,9 @@ registerIcons({
       "Star": IconContents.STAR,
       "Clean": IconContents.CLEAN,
       "Info" : IconContents.INFO_SIGN,
-      "Form": IconContents.FORM
+      "Form": IconContents.FORM,
+      "Grouped": IconContents.GROUP_OBJECTS,
+      "UnGrouped": IconContents.UNGROUP_OBJECTS
    }
 });
 
