Horde: Fix conform issue where mixing materializers caused deleted workspace dirs

#Synced-CL Unknown

diff --git a/Engine/Source/Programs/Horde/Drivers/JobDriver.Tests/Execution/ConformExecutorTests.cs b/Engine/Source/Programs/Horde/Drivers/JobDriver.Tests/Execution/ConformExecutorTests.cs
new file mode 100644
--- /dev/null
+++ b/Engine/Source/Programs/Horde/Drivers/JobDriver.Tests/Execution/ConformExecutorTests.cs
@@ -0,0 +1,105 @@
+// Copyright Epic Games, Inc. All Rights Reserved.
+
+using System;
+using System.Collections.Generic;
+using System.Threading;
+using System.Threading.Tasks;
+using EpicGames.Horde.Agents;
+using EpicGames.Horde.Agents.Leases;
+using EpicGames.Perforce.Fixture;
+using HordeCommon.Rpc.Messages;
+using HordeCommon.Rpc.Tasks;
+using JobDriver.Execution;
+using Microsoft.Extensions.DependencyInjection;
+using Microsoft.Extensions.Logging;
+using Microsoft.VisualStudio.TestTools.UnitTesting;
+using OpenTelemetry.Trace;
+using SyncOptions = JobDriver.Execution.SyncOptions;
+
+namespace JobDriver.Tests.Execution;
+
+[TestClass]
+public class ConformTest : BasePerforceFixtureTest
+{
+	private static Tracer NoOpTracer { get; } = TracerProvider.Default.GetTracer("NoOp");
+
+	private readonly ILogger _logger;
+	private readonly IServiceProvider _serviceProvider;
+	private readonly ManagedWorkspaceMaterializerFactory _mwMatFactory;
+	private readonly PerforceMaterializerFactory _p4MatFactory;
+	private readonly List<IWorkspaceMaterializerFactory> _materializerFactories;
+
+	public ConformTest()
+	{
+		_logger = LoggerFactory.CreateLogger("ConformExecutor");
+		ServiceCollection services = [];
+		services.AddSingleton<Tracer>(NoOpTracer);
+		services.AddLogging(builder => 
+		{
+			builder.AddConsole();
+			builder.AddDebug();
+			builder.SetMinimumLevel(LogLevel.Debug);
+		});
+		_serviceProvider = services.BuildServiceProvider();
+		_mwMatFactory = new ManagedWorkspaceMaterializerFactory(_serviceProvider);
+		_p4MatFactory = new PerforceMaterializerFactory(_serviceProvider);
+		_materializerFactories = [_mwMatFactory, _p4MatFactory];
+	}
+
+	[TestMethod]
+	public async Task ConformAsync()
+	{
+		RpcAgentWorkspace rawMw = CreateAgentWorkspace("mw");
+		RpcAgentWorkspace rawPm = CreateAgentWorkspace("pm", "name=perforce");
+
+		IWorkspaceMaterializer? mwm = await _mwMatFactory.CreateMaterializerAsync(ManagedWorkspaceMaterializer.TypeName, rawMw, TempDir, false, CancellationToken.None);
+		IWorkspaceMaterializer? pm = await _p4MatFactory.CreateMaterializerAsync(PerforceMaterializer.TypeName, rawPm, TempDir, false, CancellationToken.None);
+		Assert.IsNotNull(mwm);
+		Assert.IsNotNull(pm);
+
+		// Sync and verify workspaces to a random CL
+		ChangelistFixture syncedCl = Fixture.StreamFooMain.GetChangelist(5);
+		await mwm.SyncAsync(syncedCl.Number, -1, new SyncOptions(), CancellationToken.None);
+		await pm.SyncAsync(syncedCl.Number, -1, new SyncOptions(), CancellationToken.None);
+		syncedCl.AssertDepotFiles(mwm.SyncDir.FullName);
+		syncedCl.AssertDepotFiles(pm.SyncDir.FullName);
+
+		// Conform the workspaces, which also results in them syncing to latest CL of the stream
+		List<RpcAgentWorkspace> workspaces = [rawMw, rawPm];
+		ConformExecutor executor = CreateConformExecutor(workspaces);
+		await executor.ExecuteAsync(CancellationToken.None);
+		
+		// Verify synced files (nothing missing/added, except what's at latest CL)
+		ChangelistFixture conformedCl = Fixture.StreamFooMain.Head;
+		conformedCl.AssertDepotFiles(mwm.SyncDir.FullName);
+		conformedCl.AssertDepotFiles(pm.SyncDir.FullName);
+	}
+
+	private RpcAgentWorkspace CreateAgentWorkspace(string identifer, string method = "")
+	{
+		return new RpcAgentWorkspace
+		{
+			ServerAndPort = PerforceConnection.ServerAndPort,
+			UserName = PerforceConnection.UserName,
+			Password = PerforceConnection.Settings.Password,
+			Identifier = identifer,
+			Stream = Fixture.StreamFooMain.Root,
+			MinScratchSpace = 1,
+			Method = method
+		};
+	}
+
+	private ConformExecutor CreateConformExecutor(List<RpcAgentWorkspace> workspaces)
+	{
+		ConformTask conformTask = new () { LogId = "222222222222222222222222", RemoveUntrackedFiles = false };
+		conformTask.Workspaces.AddRange(workspaces);
+
+		LeaseId leaseId = LeaseId.Parse("111111111111111111111111");
+		JobRpcClientStub jobRpcClientStub = new (_logger);
+		FakeHordeClient hordeClient = new (jobRpcClientStub);
+		return new ConformExecutor(hordeClient, _materializerFactories, TempDir, new AgentId("agent1"), leaseId, conformTask,
+			new DriverSettings(),
+			NoOpTracer, _logger);
+	}
+}
+
diff --git a/Engine/Source/Programs/Horde/Drivers/JobDriver.Tests/Execution/WorkspaceExecutorTest.cs b/Engine/Source/Programs/Horde/Drivers/JobDriver.Tests/Execution/WorkspaceExecutorTest.cs
--- a/Engine/Source/Programs/Horde/Drivers/JobDriver.Tests/Execution/WorkspaceExecutorTest.cs
+++ b/Engine/Source/Programs/Horde/Drivers/JobDriver.Tests/Execution/WorkspaceExecutorTest.cs
@@ -27,111 +27,104 @@ using EpicGames.Horde.Tools;
 
 namespace JobDriver.Tests.Execution
 {
-	[TestClass]
-	public sealed class WorkspaceExecutorTest : IDisposable
+	internal class FakeHordeClient : IHordeClient
 	{
-		class FakeHordeClient : IHordeClient
+		class FakeJobRpc : JobRpc.JobRpcClient
 		{
-			class FakeJobRpc : JobRpc.JobRpcClient
-			{
-			}
+		}
 
-			private readonly Dictionary<StreamId, RpcGetStreamResponse> _streamIdToStreamResponse = new();
-			private readonly Dictionary<JobId, RpcGetJobResponse> _jobIdToJobResponse = new();
+		private readonly JobRpc.JobRpcClient? _jobRpcClient = null;
+		private readonly Dictionary<StreamId, RpcGetStreamResponse> _streamIdToStreamResponse = new();
+		private readonly Dictionary<JobId, RpcGetJobResponse> _jobIdToJobResponse = new();
 
-			public event Action? OnAccessTokenStateChanged
-			{
-				add { }
-				remove { }
-			}
+		public event Action? OnAccessTokenStateChanged
+		{
+			add { }
+			remove { }
+		}
 
-			public void AddStream(StreamId streamId, string streamName)
+		public void AddStream(StreamId streamId, string streamName)
+		{
+			if (_streamIdToStreamResponse.ContainsKey(streamId))
 			{
-				if (_streamIdToStreamResponse.ContainsKey(streamId))
-				{
-					throw new Exception($"Stream ID {streamId} already added");
-				}
-
-				_streamIdToStreamResponse[streamId] = new RpcGetStreamResponse
-				{
-					Name = streamName,
-				};
+				throw new Exception($"Stream ID {streamId} already added");
 			}
 
-			public void AddAgentType(StreamId streamId, string agentType)
+			_streamIdToStreamResponse[streamId] = new RpcGetStreamResponse
 			{
-				if (!_streamIdToStreamResponse.TryGetValue(streamId, out RpcGetStreamResponse? streamResponse))
-				{
-					throw new Exception($"Stream ID {streamId} not found");
-				}
-
-				string tempDir = Path.Join(Path.GetTempPath(), $"horde-agent-type-{agentType}-" + Guid.NewGuid().ToString()[..8]);
-				Directory.CreateDirectory(tempDir);
-
-				streamResponse.AgentTypes[agentType] = new RpcGetAgentTypeResponse
-				{
-					TempStorageDir = tempDir
-				};
-			}
+				Name = streamName,
+			};
+		}
 
-			public void AddJob(JobId jobId, StreamId streamId, int change, int preflightChange)
+		public void AddAgentType(StreamId streamId, string agentType)
+		{
+			if (!_streamIdToStreamResponse.TryGetValue(streamId, out RpcGetStreamResponse? streamResponse))
 			{
-				if (!_streamIdToStreamResponse.ContainsKey(streamId))
-				{
-					throw new Exception($"Stream ID {streamId} not found");
-				}
-
-				_jobIdToJobResponse[jobId] = new RpcGetJobResponse
-				{
-					StreamId = streamId.ToString(),
-					Change = change,
-					PreflightChange = preflightChange
-				};
+				throw new Exception($"Stream ID {streamId} not found");
 			}
 
-			public Uri ServerUrl
-				=> new Uri("http://fake-horde-server");
-
-			public IArtifactCollection Artifacts => throw new NotImplementedException();
-			public IComputeClient Compute => throw new NotImplementedException();
-			public IProjectCollection Projects => throw new NotImplementedException();
-			public ISecretCollection Secrets => throw new NotImplementedException();
-			public IToolCollection Tools => throw new NotImplementedException();
+			string tempDir = Path.Join(Path.GetTempPath(), $"horde-agent-type-{agentType}-" + Guid.NewGuid().ToString()[..8]);
+			Directory.CreateDirectory(tempDir);
 
-			public IComputeClient CreateComputeClient()
-				=> throw new NotImplementedException();
-
-			public Task<GrpcChannel> CreateGrpcChannelAsync(CancellationToken cancellationToken = default)
-				=> throw new NotImplementedException();
+			streamResponse.AgentTypes[agentType] = new RpcGetAgentTypeResponse
+			{
+				TempStorageDir = tempDir
+			};
+		}
 
-			public Task<TClient> CreateGrpcClientAsync<TClient>(CancellationToken cancellationToken = default) where TClient : ClientBase<TClient>
+		public void AddJob(JobId jobId, StreamId streamId, int change, int preflightChange)
+		{
+			if (!_streamIdToStreamResponse.ContainsKey(streamId))
 			{
-				if (typeof(TClient) == typeof(JobRpc.JobRpcClient))
-				{
-					return Task.FromResult((TClient)(object)new FakeJobRpc());
-				}
-				throw new NotImplementedException();
+				throw new Exception($"Stream ID {streamId} not found");
 			}
 
-			public HordeHttpClient CreateHttpClient()
-				=> throw new NotImplementedException();
-
-			public IServerLogger CreateServerLogger(LogId logId, LogLevel minimumLevel = LogLevel.Information)
-				=> throw new NotImplementedException();
-
-			public IStorageNamespace GetStorageNamespace(string relativePath, string? accessToken = null)
-				=> throw new NotImplementedException();
+			_jobIdToJobResponse[jobId] = new RpcGetJobResponse
+			{
+				StreamId = streamId.ToString(),
+				Change = change,
+				PreflightChange = preflightChange
+			};
+		}
 
-			public Task<string?> GetAccessTokenAsync(bool interactive, CancellationToken cancellationToken = default)
-				=> Task.FromResult<string?>("access-token");
+		public FakeHordeClient(JobRpc.JobRpcClient? jobRpcClient = null)
+		{
+			_jobRpcClient = jobRpcClient;
+		}
 
-			public bool HasValidAccessToken()
-				=> true;
+		public Uri ServerUrl => new ("http://fake-horde-server");
+		public IArtifactCollection Artifacts => throw new NotImplementedException();
+		public IComputeClient Compute => throw new NotImplementedException();
+		public IProjectCollection Projects => throw new NotImplementedException();
+		public ISecretCollection Secrets => throw new NotImplementedException();
+		public IToolCollection Tools => throw new NotImplementedException();
 
-			public Task<bool> LoginAsync(bool allowPrompt, CancellationToken cancellationToken)
-				=> throw new NotImplementedException();
+		public IComputeClient CreateComputeClient() => throw new NotImplementedException();
+		public Task<GrpcChannel> CreateGrpcChannelAsync(CancellationToken cancellationToken = default) => throw new NotImplementedException();
+		public Task<TClient> CreateGrpcClientAsync<TClient>(CancellationToken cancellationToken = default) where TClient : ClientBase<TClient>
+		{
+			if (typeof(TClient) == typeof(JobRpc.JobRpcClient))
+			{
+				return _jobRpcClient != null
+					? Task.FromResult((TClient)(object)_jobRpcClient)
+					: Task.FromResult((TClient)(object)new FakeJobRpc());
+			}
+			throw new NotImplementedException();
 		}
 
+		public HordeHttpClient CreateHttpClient() => throw new NotImplementedException();
+		public IServerLogger CreateServerLogger(LogId logId, LogLevel minimumLevel = LogLevel.Information) => throw new NotImplementedException();
+		public IStorageNamespace GetStorageNamespace(string relativePath, string? accessToken = null) => throw new NotImplementedException();
+		public Task<string?> GetAccessTokenAsync(bool interactive, CancellationToken cancellationToken = default) => Task.FromResult<string?>("access-token");
+		public bool HasValidAccessToken() => true;
+		public Task<bool> LoginAsync(bool allowPrompt, CancellationToken cancellationToken) => throw new NotImplementedException();
+	}
+	
+	[TestClass]
+	public sealed class WorkspaceExecutorTest : IDisposable
+	{
+		
+
 		private readonly StreamId _streamId = new StreamId("foo-main");
 		private readonly JobId _jobId = JobId.Parse("65bd0655591b5d5d7d047b58");
 		private readonly JobStepBatchId _batchId = new JobStepBatchId(0x1234);
diff --git a/Engine/Source/Programs/Horde/Drivers/JobDriver.Tests/JobRpcClientStub.cs b/Engine/Source/Programs/Horde/Drivers/JobDriver.Tests/JobRpcClientStub.cs
--- a/Engine/Source/Programs/Horde/Drivers/JobDriver.Tests/JobRpcClientStub.cs
+++ b/Engine/Source/Programs/Horde/Drivers/JobDriver.Tests/JobRpcClientStub.cs
@@ -93,6 +93,12 @@ namespace JobDriver.Tests
 			return Wrap(new RpcGetStepResponse());
 		}
 
+		public override AsyncUnaryCall<RpcUpdateAgentWorkspacesResponse> UpdateAgentWorkspacesAsync(RpcUpdateAgentWorkspacesRequest request, CallOptions options)
+		{
+			RpcUpdateAgentWorkspacesResponse res = new () { Retry = false };
+			return Wrap(res);
+		}
+
 		public static AsyncUnaryCall<T> Wrap<T>(T res)
 		{
 			return new AsyncUnaryCall<T>(Task.FromResult(res), Task.FromResult(Metadata.Empty),
diff --git a/Engine/Source/Programs/Horde/Drivers/JobDriver.Tests/PerforceMaterializerTests.cs b/Engine/Source/Programs/Horde/Drivers/JobDriver.Tests/PerforceMaterializerTests.cs
--- a/Engine/Source/Programs/Horde/Drivers/JobDriver.Tests/PerforceMaterializerTests.cs
+++ b/Engine/Source/Programs/Horde/Drivers/JobDriver.Tests/PerforceMaterializerTests.cs
@@ -189,7 +189,7 @@ public class PerforceMaterializerTests : BasePerforceFixtureTest
 		IServiceProvider serviceProvider = services.BuildServiceProvider();
 		List<IWorkspaceMaterializerFactory> factories = [new PerforceMaterializerFactory(serviceProvider)];
 		
-		await PerforceExecutor.ConformMaterializersAsync(factories, TempDir, [raw], removeUntrackedFiles, NoOpTracer, logger, CancellationToken.None);
+		await PerforceExecutor.ConformMaterializersAsync(factories, TempDir, [raw], removeUntrackedFiles, [], NoOpTracer, logger, CancellationToken.None);
 	}
 	
 	private const string Stream = "//Foo/Main";
diff --git a/Engine/Source/Programs/Horde/Drivers/JobDriver/Execution/ConformExecutor.cs b/Engine/Source/Programs/Horde/Drivers/JobDriver/Execution/ConformExecutor.cs
--- a/Engine/Source/Programs/Horde/Drivers/JobDriver/Execution/ConformExecutor.cs
+++ b/Engine/Source/Programs/Horde/Drivers/JobDriver/Execution/ConformExecutor.cs
@@ -84,16 +84,36 @@ namespace JobDriver.Execution
 						managedWorkspaces, perforceWorkspaces);
 				}
 
+				foreach (RpcAgentWorkspace raw in managedWorkspaces)
+				{
+					_logger.LogInformation("ManagedWorkspace: Stream={Stream} Identifier={Identifier} Incremental={Incremental} Method={Method}",
+						raw.Stream, raw.Identifier, raw.Incremental, raw.Method);
+				}
+				
+				foreach (RpcAgentWorkspace raw in perforceWorkspaces)
+				{
+					_logger.LogInformation("PerforceMaterializer: Stream={Stream} Identifier={Identifier} Incremental={Incremental} Method={Method}",
+						raw.Stream, raw.Identifier, raw.Incremental, raw.Method);
+				}
+
+				// Get the workspace metadata dirs to protect from deletion when running the MW-based vs PerforceMaterialized-based conform
+				// This is cheating a bit since we infer the metadata dir without asking the implementations for it
+				// Conform handling should be rewritten to be strictly materializer-based instead
+				List<DirectoryReference> workspaceMetadataDirs = pendingWorkspaces
+					.Select(x => DirectoryReference.Combine(_workingDir, x.Identifier)).ToList();
+
 				// Run the conform task
 				if (isExecutorConformCompatible && _driverSettings.PerforceExecutor.RunConform)
 				{
 					if (managedWorkspaces.Count > 0)
 					{
-						await PerforceExecutor.ConformAsync(_workingDir, managedWorkspaces, removeUntrackedFiles, _tracer, _logger, cancellationToken);	
+						_logger.LogInformation("Conforming ManagedWorkspaces...");
+						await PerforceExecutor.ConformAsync(_workingDir, managedWorkspaces, removeUntrackedFiles, workspaceMetadataDirs, _tracer, _logger, cancellationToken);	
 					}
 					if (perforceWorkspaces.Count > 0)
 					{
-						await PerforceExecutor.ConformMaterializersAsync(_materializerFactories, _workingDir, perforceWorkspaces, removeUntrackedFiles, _tracer, _logger, cancellationToken);	
+						_logger.LogInformation("Conforming PerforceMaterializers...");
+						await PerforceExecutor.ConformMaterializersAsync(_materializerFactories, _workingDir, perforceWorkspaces, removeUntrackedFiles, workspaceMetadataDirs, _tracer, _logger, cancellationToken);	
 					}
 				}
 				else
diff --git a/Engine/Source/Programs/Horde/Drivers/JobDriver/Execution/PerforceExecutor.cs b/Engine/Source/Programs/Horde/Drivers/JobDriver/Execution/PerforceExecutor.cs
--- a/Engine/Source/Programs/Horde/Drivers/JobDriver/Execution/PerforceExecutor.cs
+++ b/Engine/Source/Programs/Horde/Drivers/JobDriver/Execution/PerforceExecutor.cs
@@ -247,7 +247,14 @@ namespace JobDriver.Execution
 			await base.FinalizeAsync(logger, cancellationToken);
 		}
 
-		public static async Task ConformAsync(DirectoryReference rootDir, IList<RpcAgentWorkspace> pendingWorkspaces, bool removeUntrackedFiles, Tracer tracer, ILogger logger, CancellationToken cancellationToken)
+		public static async Task ConformAsync(
+			DirectoryReference rootDir,
+			IList<RpcAgentWorkspace> pendingWorkspaces,
+			bool removeUntrackedFiles,
+			List<DirectoryReference> ignoreDirs,
+			Tracer tracer,
+			ILogger logger,
+			CancellationToken cancellationToken)
 		{
 			using TelemetrySpan span = tracer.StartActiveSpan("Conform");
 			span.SetAttribute("horde.perforce.workspaces", String.Join(',', pendingWorkspaces.Select(x => x.Identifier)));
@@ -355,6 +362,7 @@ namespace JobDriver.Execution
 						protectDirs.Add(DirectoryReference.Combine(rootDir, "Saved"));
 					}
 					protectDirs.UnionWith(workspaces.Select(x => x.MetadataDir));
+					protectDirs.UnionWith(ignoreDirs);
 
 					// Delete all the directories which aren't a workspace root
 					foreach (DirectoryReference dir in DirectoryReference.EnumerateDirectories(rootDir))
@@ -448,6 +456,7 @@ namespace JobDriver.Execution
 			DirectoryReference rootDir,
 			IList<RpcAgentWorkspace> workspaces,
 			bool removeUntrackedFiles,
+			List<DirectoryReference> ignoreDirs,
 			Tracer tracer,
 			ILogger logger,
 			CancellationToken cancellationToken)
@@ -494,8 +503,8 @@ namespace JobDriver.Execution
 					protectDirs.Add(DirectoryReference.Combine(rootDir, "Temp"));
 					protectDirs.Add(DirectoryReference.Combine(rootDir, "Saved"));
 				}
-				// protectDirs.UnionWith(workspaces.Select(x => x.MetadataDir));
 				protectDirs.UnionWith(materializers.Select(x => x.BaseDir));
+				protectDirs.UnionWith(ignoreDirs);
 
 				// Delete all the directories which aren't a workspace root
 				foreach (DirectoryReference dir in DirectoryReference.EnumerateDirectories(rootDir))
