[Horde][BuildHealth][Dashboard] Remove rootCauseOwnerId from GetIssueResponse type.
Only rootCauseOwnerInfo is passed from the server and should be all we need.
#rnx
#Synced-CL Unknown

diff --git a/Engine/Source/Programs/Horde/HordeDashboard/src/backend/Api.ts b/Engine/Source/Programs/Horde/HordeDashboard/src/backend/Api.ts
--- a/Engine/Source/Programs/Horde/HordeDashboard/src/backend/Api.ts
+++ b/Engine/Source/Programs/Horde/HordeDashboard/src/backend/Api.ts
@@ -3235,9 +3235,6 @@ export type GetIssueResponse = {
 	/**Root cause commit id */
 	rootCommitId?: CommitId;
 
-	/**Root cause owner id*/
-	rootCauseOwnerId?: string;
-
 	/**Root cause owner info*/
 	rootCauseOwnerInfo?: GetThinUserInfoResponse;
 
