[Horde][BuildHealth][Server] JobStepRef metadata to be emitted under more circumstances & ensure backwards compatibility on old result set.

#Synced-CL 48176906

diff --git a/Engine/Source/Programs/Horde/Plugins/Build/HordeServer.Build/Jobs/Bisect/BisectService.cs b/Engine/Source/Programs/Horde/Plugins/Build/HordeServer.Build/Jobs/Bisect/BisectService.cs
--- a/Engine/Source/Programs/Horde/Plugins/Build/HordeServer.Build/Jobs/Bisect/BisectService.cs
+++ b/Engine/Source/Programs/Horde/Plugins/Build/HordeServer.Build/Jobs/Bisect/BisectService.cs
@@ -178,7 +178,7 @@ namespace HordeServer.Jobs.Bisect
 			// Query jobs before the start position
 			for (int maxJobCount = 20; ; maxJobCount += 20)
 			{
-				List<IJobStepRef> jobStepRefs = await _jobStepRefCollection.GetStepsForNodeAsync(bisectTask.StreamId, bisectTask.TemplateId, bisectTask.NodeName, bisectTask.InitialCommitId, true, maxJobCount, cancellationToken);
+				List<IJobStepRef> jobStepRefs = await _jobStepRefCollection.GetStepsForNodeAsync(bisectTask.StreamId, bisectTask.TemplateId, bisectTask.NodeName, bisectTask.InitialCommitId, true, maxJobCount, false, cancellationToken);
 
 				IJobStepRef? currentJobStepRef = null;
 				foreach (IJobStepRef jobStepRef in jobStepRefs.OrderByDescending(x => x.CommitId).ThenBy(x => x.Id))
diff --git a/Engine/Source/Programs/Horde/Plugins/Build/HordeServer.Build/Jobs/IJobStepRefCollection.cs b/Engine/Source/Programs/Horde/Plugins/Build/HordeServer.Build/Jobs/IJobStepRefCollection.cs
--- a/Engine/Source/Programs/Horde/Plugins/Build/HordeServer.Build/Jobs/IJobStepRefCollection.cs
+++ b/Engine/Source/Programs/Horde/Plugins/Build/HordeServer.Build/Jobs/IJobStepRefCollection.cs
@@ -61,6 +61,21 @@ namespace HordeServer.Jobs
 		/// <returns></returns>
 		Task<IJobStepRef?> FindAsync(JobId jobId, JobStepBatchId batchId, JobStepId stepId);
 
+		/// <summary>
+		/// Gets job step references given a jobId and batchId.
+		/// </summary>
+		/// <param name="jobId">The job to obtain the steps from.</param>
+		/// <param name="batchId">The batch to obtain the steps from.</param>
+		/// <returns>The list of job step refs.</returns>
+		Task<List<IJobStepRef>> FindAsync(JobId jobId, JobStepBatchId batchId);
+
+		/// <summary>
+		/// Gets job step references given a jobId.
+		/// </summary>
+		/// <param name="jobId">The job to obtain the steps from.</param>
+		/// <returns>The list of job step refs.</returns>
+		Task<List<IJobStepRef>> FindAsync(JobId jobId);
+
 		/// <summary>
 		/// Gets job step references given an array of ids
 		/// </summary>
@@ -78,9 +93,10 @@ namespace HordeServer.Jobs
 		/// <param name="commitId">The current change</param>
 		/// <param name="includeFailed">Whether to include failed nodes</param>
 		/// <param name="maxCount">Number of results to return</param>
+		/// <param name="includePreStartAndSkippedStates">Whether to include pre-run and skipped states.</param>
 		/// <param name="cancellationToken">Cancellation token for the operation</param>
 		/// <returns>List of step references</returns>
-		Task<List<IJobStepRef>> GetStepsForNodeAsync(StreamId streamId, TemplateId templateId, string nodeName, CommitIdWithOrder? commitId, bool includeFailed, int maxCount, CancellationToken cancellationToken = default);
+		Task<List<IJobStepRef>> GetStepsForNodeAsync(StreamId streamId, TemplateId templateId, string nodeName, CommitIdWithOrder? commitId, bool includeFailed, int maxCount, bool includePreStartAndSkippedStates, CancellationToken cancellationToken = default);
 
 		/// <summary>
 		/// Gets the previous job that ran a given step
@@ -111,6 +127,16 @@ namespace HordeServer.Jobs
 
 	static class JobStepRefCollectionExtensions
 	{
+		/// <summary>
+		/// Updates a job step ref.
+		/// </summary>
+		/// <param name="jobStepRefs">The job step ref collection.</param>
+		/// <param name="job">The job which the job step belongs to..</param>
+		/// <param name="batch">The batch which the job step belongs to.</param>
+		/// <param name="step">The step.</param>
+		/// <param name="graph">The graph.</param>
+		/// <param name="logger">The logger.</param>
+		/// <returns>Task object associated with the operation.</returns>
 		public static async Task UpdateAsync(this IJobStepRefCollection jobStepRefs, IJob job, IJobStepBatch batch, IJobStep step, IGraph graph, ILogger? logger = null)
 		{
 			if (job.PreflightCommitId == null)
@@ -151,7 +177,111 @@ namespace HordeServer.Jobs
 					logger.LogInformation("Updating step reference {StepId} for job {JobId}, batch {BatchId}, with outcome {JobStepOutcome}", step.Id, job.Id, batch.Id, outcome);
 				}
 
-				await jobStepRefs.InsertOrReplaceAsync(new JobStepRefId(job.Id, batch.Id, step.Id), job.Name, nodeName, job.StreamId, job.TemplateId, job.CommitId, step.LogId, batch.PoolId, batch.AgentId, step.State, outcome, job.UpdateIssues, lastSuccess, lastWarning, waitTime, initTime, job.CreateTimeUtc, step.StartTimeUtc ?? DateTime.UtcNow, step.FinishTimeUtc, step.Metadata.ToList());
+				await jobStepRefs.InsertOrReplaceAsync(new JobStepRefId(job.Id, batch.Id, step.Id), job.Name, nodeName, job.StreamId, job.TemplateId, job.CommitId, step.LogId, batch.PoolId, batch.AgentId, step.State, outcome, job.UpdateIssues, lastSuccess, lastWarning, waitTime, initTime, job.CreateTimeUtc, step.StartTimeUtc ?? DateTime.UtcNow, step.FinishTimeUtc, step.Metadata?.ToList());
+			}
+		}
+
+		/// <summary>
+		/// Updates a job step ref.
+		/// </summary>
+		/// <param name="jobStepRefs">The job step ref collection.</param>
+		/// <param name="job">The job which the job step belongs to..</param>
+		/// <param name="batch">The batch which the job step belongs to.</param>
+		/// <param name="existingStepRef">The existing job step ref to update.</param>
+		/// <param name="logger">The logger.</param>
+		/// <returns>Task object associated with the operation.</returns>
+		public static async Task UpdateAsync(this IJobStepRefCollection jobStepRefs, IJob job, IJobStepBatch batch, IJobStepRef existingStepRef, ILogger? logger = null)
+		{
+			if (job.PreflightCommitId == null)
+			{
+				float waitTime = (float)(batch.GetWaitTime() ?? TimeSpan.Zero).TotalSeconds;
+				float initTime = (float)(batch.GetInitTime() ?? TimeSpan.Zero).TotalSeconds;
+
+				string nodeName = existingStepRef.NodeName;
+
+				CommitIdWithOrder? lastSuccess = null;
+				CommitIdWithOrder? lastWarning = null;
+				JobStepOutcome? newOutcome = existingStepRef.Outcome;
+				JobStepState? newState = existingStepRef.State;
+
+				if (newOutcome != JobStepOutcome.Success)
+				{
+					IJobStepRef? prevStep = await jobStepRefs.GetPrevStepForNodeAsync(job.StreamId, job.TemplateId, nodeName, job.CommitId);
+					if (prevStep != null)
+					{
+						lastSuccess = prevStep.LastSuccess;
+						if (newOutcome != JobStepOutcome.Warnings)
+						{
+							lastWarning = prevStep.LastWarning;
+						}
+					}
+				}
+
+				if (job.AbortedByUserId != null)
+				{
+					newOutcome = JobStepOutcome.Failure;
+					newState = JobStepState.Aborted;
+				}
+
+				if (batch.Error != JobStepBatchError.None)
+				{
+					newOutcome = JobStepOutcome.Failure;
+				}
+
+				if (logger != null)
+				{
+					logger.LogInformation("Updating step reference {StepId} for job {JobId}, batch {BatchId}, with outcome {JobStepOutcome}", existingStepRef.Id, job.Id, batch.Id, newOutcome);
+				}
+
+				await jobStepRefs.InsertOrReplaceAsync(existingStepRef.Id, job.Name, nodeName, job.StreamId, job.TemplateId, job.CommitId, existingStepRef.LogId, batch.PoolId, batch.AgentId, newState, newOutcome, job.UpdateIssues, lastSuccess, lastWarning, waitTime, initTime, job.CreateTimeUtc, existingStepRef.StartTimeUtc, existingStepRef.FinishTimeUtc, existingStepRef.Metadata?.ToList());
+			}
+		}
+
+		/// <summary>
+		/// Updates a job step ref.
+		/// </summary>
+		/// <param name="jobStepRefs">The job step ref collection.</param>
+		/// <param name="job">The job which the job step belongs to..</param>
+		/// <param name="existingStepRef">The existing job step ref to update.</param>
+		/// <param name="logger">The logger.</param>
+		/// <returns>Task object associated with the operation.</returns>
+		public static async Task UpdateAsync(this IJobStepRefCollection jobStepRefs, IJob job, IJobStepRef existingStepRef, ILogger? logger = null)
+		{
+			if (job.PreflightCommitId == null)
+			{
+
+				string nodeName = existingStepRef.NodeName;
+
+				CommitIdWithOrder? lastSuccess = null;
+				CommitIdWithOrder? lastWarning = null;
+				JobStepOutcome? newOutcome = existingStepRef.Outcome;
+				JobStepState? newState = existingStepRef.State;
+
+				if (newOutcome != JobStepOutcome.Success)
+				{
+					IJobStepRef? prevStep = await jobStepRefs.GetPrevStepForNodeAsync(job.StreamId, job.TemplateId, nodeName, job.CommitId);
+					if (prevStep != null)
+					{
+						lastSuccess = prevStep.LastSuccess;
+						if (newOutcome != JobStepOutcome.Warnings)
+						{
+							lastWarning = prevStep.LastWarning;
+						}
+					}
+				}
+
+				if (job.AbortedByUserId != null)
+				{
+					newOutcome = JobStepOutcome.Failure;
+					newState = JobStepState.Aborted;
+				}
+
+				if (logger != null)
+				{
+					//logger.LogInformation("Updating step reference {StepId} for job {JobId}, batch {BatchId}, with outcome {JobStepOutcome}", , job.Id, batch.Id, newOutcome);
+				}
+
+				await jobStepRefs.InsertOrReplaceAsync(existingStepRef.Id, job.Name, nodeName, job.StreamId, job.TemplateId, job.CommitId, existingStepRef.LogId, existingStepRef.PoolId, existingStepRef.AgentId, newState, newOutcome, job.UpdateIssues, lastSuccess, lastWarning, existingStepRef.BatchWaitTime, existingStepRef.BatchInitTime, job.CreateTimeUtc, existingStepRef.StartTimeUtc, existingStepRef.FinishTimeUtc, existingStepRef.Metadata?.ToList());
 			}
 		}
 	}
diff --git a/Engine/Source/Programs/Horde/Plugins/Build/HordeServer.Build/Jobs/JobCollection.cs b/Engine/Source/Programs/Horde/Plugins/Build/HordeServer.Build/Jobs/JobCollection.cs
--- a/Engine/Source/Programs/Horde/Plugins/Build/HordeServer.Build/Jobs/JobCollection.cs
+++ b/Engine/Source/Programs/Horde/Plugins/Build/HordeServer.Build/Jobs/JobCollection.cs
@@ -432,7 +432,7 @@ namespace HordeServer.Jobs
 			[BsonIgnore]
 			public CommitIdWithOrder CommitId
 			{
-				get => (CommitName != null)? new CommitIdWithOrder(CommitName, CommitOrder) : CommitIdWithOrder.FromPerforceChange(CommitOrder);
+				get => (CommitName != null) ? new CommitIdWithOrder(CommitName, CommitOrder) : CommitIdWithOrder.FromPerforceChange(CommitOrder);
 				set => (CommitName, CommitOrder) = (value.Name, value.Order);
 			}
 
@@ -788,6 +788,7 @@ namespace HordeServer.Jobs
 		readonly ITelemetryWriter _telemetryWriter;
 		readonly IClock _clock;
 		readonly IGraphCollection _graphCollection;
+		readonly IJobStepRefCollection _jobStepRefs;
 		readonly ILogCollection _logCollection;
 		readonly ICommitService _commitService;
 		readonly Tracer _tracer;
@@ -798,10 +799,11 @@ namespace HordeServer.Jobs
 		/// <summary>
 		/// Constructor
 		/// </summary>
-		public JobCollection(IMongoService mongoService, IClock clock, IGraphCollection graphCollection, ILogCollection logCollection, ICommitService commitService, ITelemetryWriter telemetryWriter, IOptionsMonitor<BuildConfig> buildConfig, Tracer tracer, ILogger<JobCollection> logger, IAuditLog<JobId> auditLog)
+		public JobCollection(IMongoService mongoService, IClock clock, IGraphCollection graphCollection, ILogCollection logCollection, ICommitService commitService, ITelemetryWriter telemetryWriter, IOptionsMonitor<BuildConfig> buildConfig, Tracer tracer, ILogger<JobCollection> logger, IAuditLog<JobId> auditLog, IJobStepRefCollection jobStepRefCollection)
 		{
 			_clock = clock;
 			_graphCollection = graphCollection;
+			_jobStepRefs = jobStepRefCollection;
 			_logCollection = logCollection;
 			_commitService = commitService;
 			_telemetryWriter = telemetryWriter;
@@ -822,7 +824,7 @@ namespace HordeServer.Jobs
 			indexes.Add(keys => keys.Descending(x => x.SchedulePriority));
 			indexes.Add(_startedByBisectTaskIdIndex = MongoIndex.Create<JobDocument>(keys => keys.Descending(x => x.StartedByBisectTaskId), sparse: true));
 			_jobs = mongoService.GetCollection<JobDocument>("Jobs", indexes);
-		}		
+		}
 
 		/// <summary>
 		/// Get Job Logger
@@ -912,7 +914,7 @@ namespace HordeServer.Jobs
 
 			JobAudit(jobId, "Job {JobId} added, stream {StreamId}, templateRefId {TemplateId}, name {Name}, commitId {CommitId}", jobId, streamId.Id, templateRefId, name, commitId);
 
-			CreateBatches(newJob, graph, _logger, GetJobLogger(jobId));
+			await CreateBatchesInternalAsync(newJob, graph, _logger, GetJobLogger(jobId));
 
 			await _jobs.InsertOneAsync(newJob, null, cancellationToken);
 
@@ -956,7 +958,7 @@ namespace HordeServer.Jobs
 
 		/// <inheritdoc/>
 		async Task<bool> TryDeleteAsync(JobDocument job, CancellationToken cancellationToken)
-		{			
+		{
 			DeleteResult result = await _jobs.DeleteOneAsync(x => x.Id == job.Id && x.UpdateIndex == job.UpdateIndex, null, cancellationToken);
 			foreach (JobStepDocument step in job.Batches.SelectMany(x => x.Steps))
 			{
@@ -1243,14 +1245,14 @@ namespace HordeServer.Jobs
 			// This find call uses an index hint. Modifying the parameter passed to FindAsync can affect execution time a lot as the query planner is forced to use the specified index.
 			return await FindInternalAsync(
 				new FindJobOptions(
-					StreamId: streamId, 
-					Templates: templates, 
-					PreflightStartedByUser: preflightStartedByUser, 
-					ModifiedAfter: modifiedAfter, 
+					StreamId: streamId,
+					Templates: templates,
+					PreflightStartedByUser: preflightStartedByUser,
+					ModifiedAfter: modifiedAfter,
 					MaxCreateTime: maxCreateTime),
 				index: index,
 				count: count,
-				consistentRead: consistentRead, 
+				consistentRead: consistentRead,
 				cancellationToken: cancellationToken);
 		}
 
@@ -1268,7 +1270,7 @@ namespace HordeServer.Jobs
 
 			// Build the update list
 			jobDocument = Clone(jobDocument);
-			
+
 			if (name != null)
 			{
 				jobDocument.Name = name;
@@ -1327,7 +1329,7 @@ namespace HordeServer.Jobs
 					throw new ArgumentException("Cannot update label trigger that already exists");
 				}
 				jobDocument._labelNotifications.Add(new LabelNotificationDocument { _labelIdx = labelIdxToTriggerId.Value.Key, _triggerId = labelIdxToTriggerId.Value.Value });
-				updates.Add(updateBuilder.Set(x => x._labelNotifications, jobDocument._labelNotifications));				
+				updates.Add(updateBuilder.Set(x => x._labelNotifications, jobDocument._labelNotifications));
 			}
 			if (jobTrigger != null)
 			{
@@ -1345,7 +1347,7 @@ namespace HordeServer.Jobs
 						}
 						catch (Exception)
 						{
-						}						
+						}
 					}
 				}
 			}
@@ -1380,7 +1382,7 @@ namespace HordeServer.Jobs
 			// Update the batches
 			if (updateBatches)
 			{
-				UpdateBatches(jobDocument, graph, updates, _logger, GetJobLogger(jobId));
+				await UpdateBatchesInternalAsync(jobDocument, graph, updates, _logger, GetJobLogger(jobId));
 			}
 
 			// Update the new list of job steps
@@ -1466,7 +1468,7 @@ namespace HordeServer.Jobs
 
 				// Force an update of all batches in the job. This will fail or reschedule any nodes that can no longer be executed in this batch.
 				updates.Clear();
-				UpdateBatches(jobDocument, graph, updates, _logger, GetJobLogger(jobDocument.Id));
+				await UpdateBatchesInternalAsync(jobDocument, graph, updates, _logger, GetJobLogger(jobDocument.Id));
 
 				if (retriedNodes.Count > 0)
 				{
@@ -1717,13 +1719,17 @@ namespace HordeServer.Jobs
 			if (refreshBatches)
 			{
 				updates.Clear(); // UpdateBatches will update the entire batches list. We need to remove all the individual step updates to avoid an exception.
-				UpdateBatches(jobDocument, graph, updates, _logger, GetJobLogger(jobDocument.Id));
+				await UpdateBatchesInternalAsync(jobDocument, graph, updates, _logger, GetJobLogger(jobDocument.Id));
 			}
 
 			// Update the state of dependent jobsteps
 			if (refreshDependentJobSteps)
 			{
 				RefreshDependentJobSteps(jobDocument, graph, updates, _logger, GetJobLogger(jobDocument.Id));
+
+				IJob jobObj = CreateJobObject(jobDocument, graph);
+				await UpdateRefStateAsync(jobObj, graph, _jobStepRefs);
+
 				RefreshJobPriority(jobDocument, updates, _logger, GetJobLogger(jobDocument.Id));
 			}
 
@@ -1860,7 +1866,7 @@ namespace HordeServer.Jobs
 			jobDocument.GraphHash = newGraph.Id;
 			updates.Add(updateBuilder.Set(x => x.GraphHash, jobDocument.GraphHash));
 
-			UpdateBatches(jobDocument, newGraph, updates, _logger, GetJobLogger(jobDocument.Id));
+			await UpdateBatchesInternalAsync(jobDocument, newGraph, updates, _logger, GetJobLogger(jobDocument.Id));
 
 			// Update the new list of job steps
 			return await TryUpdateAsync(jobDocument, updates, cancellationToken);
@@ -1940,7 +1946,7 @@ namespace HordeServer.Jobs
 			}
 
 			List<UpdateDefinition<JobDocument>> updates = new List<UpdateDefinition<JobDocument>>();
-			UpdateBatches(jobDocument, graph, updates, _logger, GetJobLogger(jobDocument.Id));
+			await UpdateBatchesInternalAsync(jobDocument, graph, updates, _logger, GetJobLogger(jobDocument.Id));
 
 			return await TryUpdateAsync(jobDocument, updates, cancellationToken);
 		}
@@ -1971,7 +1977,7 @@ namespace HordeServer.Jobs
 			}
 
 			List<UpdateDefinition<JobDocument>> updates = new List<UpdateDefinition<JobDocument>>();
-			UpdateBatches(jobDocument, graph, updates, _logger, GetJobLogger(jobDocument.Id));
+			await UpdateBatchesInternalAsync(jobDocument, graph, updates, _logger, GetJobLogger(jobDocument.Id));
 
 			return await TryUpdateAsync(jobDocument, updates, cancellationToken);
 		}
@@ -2082,6 +2088,64 @@ namespace HordeServer.Jobs
 			return await TryUpdateAsync(jobDocument, updateBuilder.Combine(updates), cancellationToken);
 		}
 
+		private async Task UpdateRefStateAsync(IJob job, IGraph graph, IJobStepRefCollection jobStepRefCollection)
+		{
+			ICollection<JobStepRefId> jobStepRefSet;
+			ICollection<string> jobBatchSet = new HashSet<string>(job.Batches.Count);
+
+			foreach (IJobStepBatch batch in job.Batches)
+			{
+				jobBatchSet.Add($"{job.Id}:{batch.Id}");
+				jobStepRefSet = new HashSet<JobStepRefId>(batch.Steps.Count);
+
+				foreach (IJobStep step in batch.Steps)
+				{
+					await jobStepRefCollection.UpdateAsync(job, batch, step, graph);
+					JobStepRefId refId = new JobStepRefId(job.Id, batch.Id, step.Id);
+					jobStepRefSet.Add(refId);
+				}
+
+				// verify against known steps
+				List<IJobStepRef> jobBatchStepRefItems = await _jobStepRefs.FindAsync(job.Id, batch.Id);
+
+				// Some early cancel steps may have been removed; mark closed on those not contained within the set.
+				if (jobBatchStepRefItems.Count != batch.Steps.Count)
+				{
+					foreach (IJobStepRef item in jobBatchStepRefItems)
+					{
+						if (!jobStepRefSet.Contains(item.Id))
+						{
+							await _jobStepRefs.UpdateAsync(job, batch, item);
+						}
+					}
+				}
+			}
+
+			// Som batches may have been entirely culled, attempt to update at the job level.
+			List<IJobStepRef> jobStepRefItems = await _jobStepRefs.FindAsync(job.Id);
+			foreach (IJobStepRef item in jobStepRefItems)
+			{
+				if (!jobBatchSet.Contains($"{item.Id.JobId}:{item.Id.BatchId}"))
+				{
+					await _jobStepRefs.UpdateAsync(job, item);
+				}
+			}
+		}
+
+		private async Task UpdateBatchesInternalAsync(JobDocument job, IGraph graph, List<UpdateDefinition<JobDocument>> updates, ILogger logger, IAuditLogChannel<JobId> jobAudit)
+		{
+			UpdateBatches(job, graph, updates, logger, jobAudit);
+			IJob jobObj = CreateJobObject(job, graph);
+			await UpdateRefStateAsync(jobObj, graph, _jobStepRefs);
+		}
+
+		private async Task CreateBatchesInternalAsync(JobDocument job, IGraph graph, ILogger logger, IAuditLogChannel<JobId> jobAudit)
+		{
+			CreateBatches(job, graph, logger, jobAudit);
+			IJob jobObj = CreateJobObject(job, graph);
+			await UpdateRefStateAsync(jobObj, graph, _jobStepRefs);
+		}
+
 		/// <summary>
 		/// Populate the list of batches for a job
 		/// </summary>
@@ -2192,7 +2256,7 @@ namespace HordeServer.Jobs
 					{
 						if (JobStepExtensions.IsPendingState(step.State))
 						{
-							step.State = JobStepState.Skipped;							
+							step.State = JobStepState.Skipped;
 							JobAudit(job.Id, logger, jobAudit, "Setting batch {BatchId}, step {StepId} to Skipped", batch.Id, step.Id);
 						}
 					}
@@ -2270,7 +2334,7 @@ namespace HordeServer.Jobs
 				{
 					newNodesToExecute.UnionWith(aggregate.Nodes.Select(x => graph.GetNode(x)));
 				}
-			}			
+			}
 
 			// Add any individual nodes
 			foreach (INode node in graph.Groups.SelectMany(x => x.Nodes))
@@ -2378,7 +2442,7 @@ namespace HordeServer.Jobs
 			if (newNodesToExecute.Count > 0)
 			{
 				JobAudit(job.Id, logger, jobAudit, "New Nodes to execute {Nodes}", String.Join(',', newNodesToExecute.Select(x => x.Name)));
-			}				
+			}
 
 			// Figure out the existing batch for each group
 			JobStepBatchDocument?[] appendToBatches = new JobStepBatchDocument?[graph.Groups.Count];
@@ -2436,7 +2500,7 @@ namespace HordeServer.Jobs
 							job.Batches.Add(batch);
 							appendToBatches[groupIdx] = batch;
 
-							JobAudit(job.Id, logger, jobAudit, "Created new batch {BatchId} for {NodeName}", batchId, node.Name);							
+							JobAudit(job.Id, logger, jobAudit, "Created new batch {BatchId} for {NodeName}", batchId, node.Name);
 						}
 
 						// Don't re-add nodes that have already executed in this batch. If we were missing a dependency, we would have already nulled out the entry in appendToBatches above; anything else
@@ -2745,14 +2809,14 @@ namespace HordeServer.Jobs
 				if (stepMetadata != null)
 				{
 					foreach (KeyValuePair<JobStepId, List<string>> entry in stepMetadata)
-					{						
+					{
 						int batchIdx = job.Batches.FindIndex(x => x.Steps.FirstOrDefault(y => y.Id == entry.Key) != null);
 						if (batchIdx >= 0)
 						{
 							int stepIdx = job.Batches[batchIdx].Steps.FindIndex(x => x.Id == entry.Key);
 							updates.Add(updateBuilder.AddToSetEach(x => x.Batches[batchIdx].Steps[stepIdx].Metadata, entry.Value));
 						}
-					}						
+					}
 				}
 
 				job = await TryUpdateAsync(job, updates, cancellationToken);
@@ -2763,7 +2827,7 @@ namespace HordeServer.Jobs
 			{
 				_logger.LogError(ex, "Failed to update job {JobId} meta data", jobId);
 				return null;
-			}			
+			}
 		}
 
 		/// <inheritdoc/>
diff --git a/Engine/Source/Programs/Horde/Plugins/Build/HordeServer.Build/Jobs/JobService.cs b/Engine/Source/Programs/Horde/Plugins/Build/HordeServer.Build/Jobs/JobService.cs
--- a/Engine/Source/Programs/Horde/Plugins/Build/HordeServer.Build/Jobs/JobService.cs
+++ b/Engine/Source/Programs/Horde/Plugins/Build/HordeServer.Build/Jobs/JobService.cs
@@ -683,7 +683,7 @@ namespace HordeServer.Jobs
 			span.SetAttribute("Change", commitId?.ToString());
 
 			// Find all the steps matching the given criteria
-			List<IJobStepRef> steps = await _jobStepRefs.GetStepsForNodeAsync(streamId, templateId, nodeName, commitId, false, 10, cancellationToken);
+			List<IJobStepRef> steps = await _jobStepRefs.GetStepsForNodeAsync(streamId, templateId, nodeName, commitId, false, 10, false, cancellationToken);
 
 			// Sum up all the durations and wait times
 			int count = 0;
diff --git a/Engine/Source/Programs/Horde/Plugins/Build/HordeServer.Build/Jobs/JobStepRefCollection.cs b/Engine/Source/Programs/Horde/Plugins/Build/HordeServer.Build/Jobs/JobStepRefCollection.cs
--- a/Engine/Source/Programs/Horde/Plugins/Build/HordeServer.Build/Jobs/JobStepRefCollection.cs
+++ b/Engine/Source/Programs/Horde/Plugins/Build/HordeServer.Build/Jobs/JobStepRefCollection.cs
@@ -12,6 +12,7 @@ using HordeServer.Streams;
 using HordeServer.Telemetry;
 using HordeServer.Utilities;
 using Microsoft.Extensions.Options;
+using MongoDB.Bson;
 using MongoDB.Bson.Serialization.Attributes;
 using MongoDB.Driver;
 
@@ -211,6 +212,26 @@ namespace HordeServer.Jobs
 			return await _jobStepRefs.FindOneAndUpdateAsync(x => x.Id.Equals(id), updateBuilder.Combine(updates));
 		}
 
+		/// <inheritdoc/>
+		public async Task<List<IJobStepRef>> FindAsync(JobId jobId, JobStepBatchId batchId)
+		{
+			string prefix = $"^{jobId}:{batchId}:";
+			FilterDefinition<JobStepRef> filter = Builders<JobStepRef>.Filter.Regex(x => x.Id, new BsonRegularExpression(prefix));
+
+			List<JobStepRef> results = await _jobStepRefs.Find(filter).ToListAsync();
+			return results.Cast<IJobStepRef>().ToList();
+		}
+
+		/// <inheritdoc/>
+		public async Task<List<IJobStepRef>> FindAsync(JobId jobId)
+		{
+			string prefix = $"^{jobId}:";
+			FilterDefinition<JobStepRef> filter = Builders<JobStepRef>.Filter.Regex(x => x.Id, new BsonRegularExpression(prefix));
+
+			List<JobStepRef> results = await _jobStepRefs.Find(filter).ToListAsync();
+			return results.Cast<IJobStepRef>().ToList();
+		}
+
 		/// <inheritdoc/>
 		public async Task<IJobStepRef?> FindAsync(JobId jobId, JobStepBatchId batchId, JobStepId stepId)
 		{
@@ -230,7 +251,7 @@ namespace HordeServer.Jobs
 		}
 
 		/// <inheritdoc/>
-		public async Task<List<IJobStepRef>> GetStepsForNodeAsync(StreamId streamId, TemplateId templateId, string nodeName, CommitIdWithOrder? commitId, bool includeFailed, int maxCount, CancellationToken cancellationToken)
+		public async Task<List<IJobStepRef>> GetStepsForNodeAsync(StreamId streamId, TemplateId templateId, string nodeName, CommitIdWithOrder? commitId, bool includeFailed, int maxCount, bool includePreStartAndSkippedStates, CancellationToken cancellationToken)
 		{
 			// Find all the steps matching the given criteria
 			FilterDefinitionBuilder<JobStepRef> filterBuilder = Builders<JobStepRef>.Filter;
@@ -247,6 +268,12 @@ namespace HordeServer.Jobs
 			{
 				filter &= filterBuilder.Ne(x => x.Outcome, JobStepOutcome.Failure);
 			}
+			if(!includePreStartAndSkippedStates)
+			{
+				filter &= filterBuilder.Ne(x => x.State, JobStepState.Skipped);
+				filter &= filterBuilder.Ne(x => x.State, JobStepState.Ready);
+				filter &= filterBuilder.Ne(x => x.State, JobStepState.Waiting);
+			}
 
 			List<JobStepRef> steps = await _jobStepRefs.Find(filter).SortByDescending(x => x.CommitOrder).ThenByDescending(x => x.StartTimeUtc).Limit(maxCount).ToListAsync(cancellationToken);
 			return steps.ConvertAll<IJobStepRef>(x => x);
diff --git a/Engine/Source/Programs/Horde/Plugins/Build/HordeServer.Build/Streams/StreamsController.cs b/Engine/Source/Programs/Horde/Plugins/Build/HordeServer.Build/Streams/StreamsController.cs
--- a/Engine/Source/Programs/Horde/Plugins/Build/HordeServer.Build/Streams/StreamsController.cs
+++ b/Engine/Source/Programs/Horde/Plugins/Build/HordeServer.Build/Streams/StreamsController.cs
@@ -409,6 +409,7 @@ namespace HordeServer.Streams
 		/// <param name="count">Number of results to return</param>
 		/// <param name="filter">The filter to apply to the results</param>
 		/// <returns>Http result code</returns>
+		/// <remarks>This will exclude <see cref="JobStepState.Ready"/>, <see cref="JobStepState.Waiting"/>, <see cref="JobStepState.Skipped"/> states.</remarks>
 		[HttpGet]
 		[Route("/api/v1/streams/{streamId}/history")]
 		[ProducesResponseType(typeof(List<GetJobStepRefResponse>), 200)]
@@ -432,7 +433,7 @@ namespace HordeServer.Streams
 				commitIdWithOrder = await _commitService.GetOrderedAsync(streamId, change, HttpContext.RequestAborted);
 			}
 
-			List<IJobStepRef> steps = await _jobStepRefCollection.GetStepsForNodeAsync(streamId, templateIdValue, step, commitIdWithOrder, true, count, HttpContext.RequestAborted);
+			List<IJobStepRef> steps = await _jobStepRefCollection.GetStepsForNodeAsync(streamId, templateIdValue, step, commitIdWithOrder, true, count, false, HttpContext.RequestAborted);
 			return steps.ConvertAll(x => PropertyFilter.Apply(CreateGetJobStepRefResponse(x), filter));
 		}
 
@@ -449,6 +450,7 @@ namespace HordeServer.Streams
 		/// <param name="filter">The filter to apply to the results</param>
 		/// <param name="cancellationToken">The cancellation token to use for the job.</param>
 		/// <returns>Http result code along with the corresponding filtered list of <see cref="GetJobStepRefResponse"/>.</returns>
+		/// <remarks>This will return all states.</remarks>
 		[HttpPost]
 		[Route("/api/v1/streams/{streamId}/history")]
 		[ProducesResponseType(typeof(List<GetJobStepRefResponse>), 200)]
@@ -485,7 +487,7 @@ namespace HordeServer.Streams
 					await maxStepHistoryActiveRequests.WaitAsync(cancellationToken);
 					try
 					{
-						return await _jobStepRefCollection.GetStepsForNodeAsync(streamId, templateIdValue, step, commitIdWithOrder, true, count, cancellationToken);
+						return await _jobStepRefCollection.GetStepsForNodeAsync(streamId, templateIdValue, step, commitIdWithOrder, true, count, true, cancellationToken);
 					}
 					catch (OperationCanceledException)
 					{
