Fix modal interaction with div padding

#rnx

#Synced-CL 48352216

diff --git a/Engine/Source/Programs/Horde/HordeDashboard/plugins/test-automation/testAutomationSummary.tsx b/Engine/Source/Programs/Horde/HordeDashboard/plugins/test-automation/testAutomationSummary.tsx
--- a/Engine/Source/Programs/Horde/HordeDashboard/plugins/test-automation/testAutomationSummary.tsx
+++ b/Engine/Source/Programs/Horde/HordeDashboard/plugins/test-automation/testAutomationSummary.tsx
@@ -457,7 +457,7 @@ export const TestSummary: React.FC<{ handler: TestDataHandler }> = ({ handler })
 
     const { hordeClasses } = getHordeStyling();
 
-    return <Stack tokens={{childrenGap: 12}}>
+    return <Stack>
             {!!historyState.test &&
                 <TestHistoryModal
                     test={historyState.test}
@@ -469,33 +469,35 @@ export const TestSummary: React.FC<{ handler: TestDataHandler }> = ({ handler })
                     onDetailsClick={historyState.onClick!}
                 />
             }
-            {failingTests.length > 0 &&
-                <Stack className={hordeClasses.raised}>
-                    <Stack style={{ paddingBottom: 12 }}>
-                        <Stack>
-                            <Text style={{ fontFamily: "Horde Open Sans Semibold" }} variant='medium'>Latest Failing Test Results</Text>
+            <Stack tokens={{childrenGap: 12}}>
+                {failingTests.length > 0 &&
+                    <Stack className={hordeClasses.raised}>
+                        <Stack style={{ paddingBottom: 12 }}>
+                            <Stack>
+                                <Text style={{ fontFamily: "Horde Open Sans Semibold" }} variant='medium'>Latest Failing Test Results</Text>
+                            </Stack>
                         </Stack>
-                    </Stack>
-                    <Stack style={{ paddingLeft: 12 }} tokens={{ childrenGap: 4 }}>
-                        <Stack wrap horizontal style={{ width: "100%" }} tokens={{ childrenGap: 12 }}>
-                            {failingTests}
+                        <Stack style={{ paddingLeft: 12 }} tokens={{ childrenGap: 4 }}>
+                            <Stack wrap horizontal style={{ width: "100%" }} tokens={{ childrenGap: 12 }}>
+                                {failingTests}
+                            </Stack>
                         </Stack>
                     </Stack>
-                </Stack>
-            }
-            {passingTests.length > 0 &&
-                <Stack className={hordeClasses.raised}>
-                    <Stack style={{ paddingBottom: 12 }}>
-                        <Stack>
-                            <Text style={{ fontFamily: "Horde Open Sans Semibold" }} variant='medium'>Latest Passing Test Results</Text>
+                }
+                {passingTests.length > 0 &&
+                    <Stack className={hordeClasses.raised}>
+                        <Stack style={{ paddingBottom: 12 }}>
+                            <Stack>
+                                <Text style={{ fontFamily: "Horde Open Sans Semibold" }} variant='medium'>Latest Passing Test Results</Text>
+                            </Stack>
                         </Stack>
-                    </Stack>
-                    <Stack style={{ paddingLeft: 12 }} tokens={{ childrenGap: 4 }}>
-                        <Stack wrap horizontal style={{ width: "100%" }} tokens={{ childrenGap: 12 }}>
-                            {passingTests}
+                        <Stack style={{ paddingLeft: 12 }} tokens={{ childrenGap: 4 }}>
+                            <Stack wrap horizontal style={{ width: "100%" }} tokens={{ childrenGap: 12 }}>
+                                {passingTests}
+                            </Stack>
                         </Stack>
                     </Stack>
-                </Stack>
-            }
+                }
+            </Stack>
         </Stack>
 }
\ No newline at end of file
