Horde: Updating schema documentation

#rnx

diff --git a/Engine/Source/Programs/Horde/Docs/Config/Schema/AclActions.md b/Engine/Source/Programs/Horde/Docs/Config/Schema/AclActions.md
--- a/Engine/Source/Programs/Horde/Docs/Config/Schema/AclActions.md
+++ b/Engine/Source/Programs/Horde/Docs/Config/Schema/AclActions.md
@@ -588,6 +588,12 @@
 | `UpdateReplicator` |  |
 | `ViewReplicator` |  |
 
+## Secrets
+
+| Name | Description |
+| ---- | ----------- |
+| `ViewSecret` | View a credential |
+
 ## ServiceAccounts
 
 | Name | Description |
@@ -750,7 +756,7 @@
 
 | Name | Description |
 | ---- | ----------- |
-| `ViewSecret` | View a credential |
+| `ViewSecret` |  |
 
 ## ServiceAccounts
 
@@ -940,6 +946,19 @@
 | `DownloadSoftware` |  |
 | `DeleteSoftware` |  |
 
+## Storage
+
+| Name | Description |
+| ---- | ----------- |
+| `ReadBlobs` | Ability to read blobs from the storage service |
+| `WriteBlobs` | Ability to write blobs to the storage service |
+| `ReadAliases` | Ability to read aliases from the storage service |
+| `WriteAliases` | Ability to write aliases to the storage service |
+| `DeleteAliases` | Ability to write aliases to the storage service |
+| `ReadRefs` | Ability to read refs from the storage service |
+| `WriteRefs` | Ability to write refs to the storage service |
+| `DeleteRefs` | Ability to delete refs |
+
 ## Streams
 
 | Name | Description |
@@ -951,12 +970,6 @@
 | `ViewChanges` |  |
 | `ViewTemplate` |  |
 
-## Symbols
-
-| Name | Description |
-| ---- | ----------- |
-| `ReadSymbols` | Ability to download symbols |
-
 ## Accounts
 
 | Name | Description |
@@ -1114,14 +1127,14 @@
 
 | Name | Description |
 | ---- | ----------- |
-| `ReadBlobs` | Ability to read blobs from the storage service |
-| `WriteBlobs` | Ability to write blobs to the storage service |
-| `ReadAliases` | Ability to read aliases from the storage service |
-| `WriteAliases` | Ability to write aliases to the storage service |
-| `DeleteAliases` | Ability to write aliases to the storage service |
-| `ReadRefs` | Ability to read refs from the storage service |
-| `WriteRefs` | Ability to write refs to the storage service |
-| `DeleteRefs` | Ability to delete refs |
+| `ReadBlobs` |  |
+| `WriteBlobs` |  |
+| `ReadAliases` |  |
+| `WriteAliases` |  |
+| `DeleteAliases` |  |
+| `ReadRefs` |  |
+| `WriteRefs` |  |
+| `DeleteRefs` |  |
 
 ## Streams
 
@@ -1138,7 +1151,7 @@
 
 | Name | Description |
 | ---- | ----------- |
-| `ReadSymbols` |  |
+| `ReadSymbols` | Ability to download symbols |
 
 ## Accounts
 
diff --git a/Engine/Source/Programs/Horde/Docs/Config/Schema/Globals.md b/Engine/Source/Programs/Horde/Docs/Config/Schema/Globals.md
--- a/Engine/Source/Programs/Horde/Docs/Config/Schema/Globals.md
+++ b/Engine/Source/Programs/Horde/Docs/Config/Schema/Globals.md
@@ -63,10 +63,10 @@ Name | Description
 `build` | [BuildConfig](#buildconfig)<br>Configuration for the build plugin
 `compute` | [ComputeConfig](#computeconfig)<br>Configuration for the compute plugin
 `exampleplugin` | [ExamplePluginConfig](#examplepluginconfig)<br>Configuration for the exampleplugin plugin
-`experimental` | [ExperimentalConfig](#experimentalconfig)<br>Configuration for the experimental plugin
 `secrets` | [SecretsConfig](#secretsconfig)<br>Configuration for the secrets plugin
-`symbols` | [SymbolsConfig](#symbolsconfig)<br>Configuration for the symbols plugin
+`experimental` | [ExperimentalConfig](#experimentalconfig)<br>Configuration for the experimental plugin
 `storage` | [StorageConfig](#storageconfig)<br>Configuration for the storage plugin
+`symbols` | [SymbolsConfig](#symbolsconfig)<br>Configuration for the symbols plugin
 `tools` | [ToolsConfig](#toolsconfig)<br>Configuration for the tools plugin
 
 ## AnalyticsConfig
@@ -442,14 +442,6 @@ Name | Description
 ---- | -----------
 `greetingMessage` | `string`<br>The Greeting provided by the example plugin
 
-## ExperimentalConfig
-
-Configuration for the experimental plugin
-
-Name | Description
----- | -----------
-`notifications` | [NotificationConfig](Notification.md)`[]`<br>Notifications to be sent on the Horde server
-
 ## SecretsConfig
 
 Configuration for the secrets system
@@ -532,24 +524,13 @@ Name | Description
 ---- | -----------
 `vaultUri` | `string`<br>The Uri to the Azure Key Vault.
 
-## SymbolsConfig
-
-Configuration for the tools system
-
-Name | Description
----- | -----------
-`stores` | [SymbolStoreConfig](#symbolstoreconfig)`[]`<br>List of symbol stores
-
-## SymbolStoreConfig
+## ExperimentalConfig
 
-Configuration for a symbol store
+Configuration for the experimental plugin
 
 Name | Description
 ---- | -----------
-`id` | `string`<br>Identifier for this store
-`namespaceId` | `string`<br>Configuration for the symbol store backend
-`public` | `boolean`<br>Whether to make this store available without auth
-`acl` | [AclConfig](#aclconfig)<br>Access to the symbol store
+`notifications` | [NotificationConfig](Notification.md)`[]`<br>Notifications to be sent on the Horde server
 
 ## StorageConfig
 
@@ -623,6 +604,25 @@ Name | Description
 `enableAliases` | `boolean`<br>Support querying exports by their aliases
 `acl` | [AclConfig](#aclconfig)<br>Access list for this namespace
 
+## SymbolsConfig
+
+Configuration for the tools system
+
+Name | Description
+---- | -----------
+`stores` | [SymbolStoreConfig](#symbolstoreconfig)`[]`<br>List of symbol stores
+
+## SymbolStoreConfig
+
+Configuration for a symbol store
+
+Name | Description
+---- | -----------
+`id` | `string`<br>Identifier for this store
+`namespaceId` | `string`<br>Configuration for the symbol store backend
+`public` | `boolean`<br>Whether to make this store available without auth
+`acl` | [AclConfig](#aclconfig)<br>Access to the symbol store
+
 ## ToolsConfig
 
 Configuration for the tools system
diff --git a/Engine/Source/Programs/Horde/Docs/Deployment/ServerSettings.md b/Engine/Source/Programs/Horde/Docs/Deployment/ServerSettings.md
--- a/Engine/Source/Programs/Horde/Docs/Deployment/ServerSettings.md
+++ b/Engine/Source/Programs/Horde/Docs/Deployment/ServerSettings.md
@@ -131,10 +131,10 @@ Name | Description
 `build` | [BuildServerConfig](#buildserverconfig)<br>Configuration for the build plugin
 `compute` | [ComputeServerConfig](#computeserverconfig)<br>Configuration for the compute plugin
 `exampleplugin` | [ExamplePluginServerConfig](#examplepluginserverconfig)<br>Configuration for the exampleplugin plugin
-`experimental` | [ExperimentalServerConfig](#experimentalserverconfig)<br>Configuration for the experimental plugin
 `secrets` | [SecretsServerConfig](#secretsserverconfig)<br>Configuration for the secrets plugin
-`symbols` | [PluginServerConfig](#pluginserverconfig)<br>Configuration for the symbols plugin
+`experimental` | [ExperimentalServerConfig](#experimentalserverconfig)<br>Configuration for the experimental plugin
 `storage` | [StorageServerConfig](#storageserverconfig)<br>Configuration for the storage plugin
+`symbols` | [PluginServerConfig](#pluginserverconfig)<br>Configuration for the symbols plugin
 `tools` | [ToolsServerConfig](#toolsserverconfig)<br>Configuration for the tools plugin
 
 ## AnalyticsServerConfig
@@ -369,14 +369,6 @@ Name | Description
 ---- | -----------
 `enabled` | `boolean`<br>Whether the plugin should be enabled or not
 
-## ExperimentalServerConfig
-
-Server configuration for the experimental plugin
-
-Name | Description
----- | -----------
-`enabled` | `boolean`<br>Whether the plugin should be enabled or not
-
 ## SecretsServerConfig
 
 Static configuration for the secrets plugin
@@ -386,9 +378,9 @@ Name | Description
 `withAws` | `boolean`<br>Whether to enable Amazon Web Services (AWS) specific features
 `enabled` | `boolean`<br>Whether the plugin should be enabled or not
 
-## PluginServerConfig
+## ExperimentalServerConfig
 
-Base class for plugin server config objects
+Server configuration for the experimental plugin
 
 Name | Description
 ---- | -----------
@@ -453,6 +445,14 @@ Name | Description
 `AssumeRole` | Assume a particular role. Should specify ARN in
 `AssumeRoleWebIdentity` | Assume a particular role using the current environment variables.
 
+## PluginServerConfig
+
+Base class for plugin server config objects
+
+Name | Description
+---- | -----------
+`enabled` | `boolean`<br>Whether the plugin should be enabled or not
+
 ## ToolsServerConfig
 
 Server configuration for bundled tools
