Horde: Add Root Cause Analysis changes to Issue Audit

#rnx
#Synced-CL Unknown

diff --git a/Engine/Source/Programs/Horde/Plugins/Build/HordeServer.Build/Issues/IssueCollection.cs b/Engine/Source/Programs/Horde/Plugins/Build/HordeServer.Build/Issues/IssueCollection.cs
--- a/Engine/Source/Programs/Horde/Plugins/Build/HordeServer.Build/Issues/IssueCollection.cs
+++ b/Engine/Source/Programs/Horde/Plugins/Build/HordeServer.Build/Issues/IssueCollection.cs
@@ -939,6 +939,7 @@ namespace HordeServer.Issues
 					issueLogger.LogInformation("Issue was marked as not fixed as a systemic issue");
 				}
 			}
+
 			if (newIssue.RootCommitId != oldIssue.RootCommitId)
 			{
 				if (newIssue.RootCommitId == null)
@@ -951,6 +952,28 @@ namespace HordeServer.Issues
 				}
 			}
 
+			if (newIssue.RootCauseOwnerId != oldIssue.RootCauseOwnerId)
+			{
+				if (newIssue.RootCauseOwnerId == null)
+				{
+					issueLogger.LogInformation("Root cause owner was cleared");
+				}
+				else
+				{
+					issueLogger.LogInformation("Root cause owner is {UserName} ({UserId})", await GetUserNameAsync(newIssue.RootCauseOwnerId, cancellationToken), newIssue.RootCauseOwnerId);
+				}
+			}
+
+			if (newIssue.RootCauseCategory != oldIssue.RootCauseCategory)
+			{
+				issueLogger.LogInformation("Changed root cause category to \"{Category}\"", newIssue.RootCauseCategory);
+			}
+
+			if (newIssue.RootCauseSummary != oldIssue.RootCauseSummary)
+			{
+				issueLogger.LogInformation("Changed root cause summary to \"{Summary}\"", newIssue.RootCauseSummary);
+			}
+
 			if (newIssue.DuplicateIssueId != oldIssue.DuplicateIssueId)
 			{
 				if (newIssue.DuplicateIssueId == null)
