Horde: Remove RootCauseOwner and RootCauseOwnerId from IssueMessage. RootCauseOwnerInfo should only be used

#rnx
#Synced-CL Unknown

diff --git a/Engine/Source/Programs/Horde/Plugins/Build/HordeServer.Build/Issues/IssuesController.cs b/Engine/Source/Programs/Horde/Plugins/Build/HordeServer.Build/Issues/IssuesController.cs
--- a/Engine/Source/Programs/Horde/Plugins/Build/HordeServer.Build/Issues/IssuesController.cs
+++ b/Engine/Source/Programs/Horde/Plugins/Build/HordeServer.Build/Issues/IssuesController.cs
@@ -518,8 +518,6 @@ namespace HordeServer.Issues
 			response.FixCommitId = issue.FixCommitId;
 			response.FixSystemic = issue.FixSystemic;
 			response.RootCommitId = issue.RootCommitId;
-			response.RootCauseOwner = details.RootCauseOwner?.Login;
-			response.RootCauseOwnerId = details.RootCauseOwner?.Id.ToString();
 			response.RootCauseOwnerInfo = details.RootCauseOwner?.ToThinApiResponse();
 			response.RootCauseCategory = issue.RootCauseCategory;
 			response.RootCauseSummary = issue.RootCauseSummary;
diff --git a/Engine/Source/Programs/Shared/EpicGames.Horde/Issues/IssueMessages.cs b/Engine/Source/Programs/Shared/EpicGames.Horde/Issues/IssueMessages.cs
--- a/Engine/Source/Programs/Shared/EpicGames.Horde/Issues/IssueMessages.cs
+++ b/Engine/Source/Programs/Shared/EpicGames.Horde/Issues/IssueMessages.cs
@@ -408,16 +408,6 @@ namespace EpicGames.Horde.Issues
 		}
 		CommitId? _rootCommitId;
 
-		/// <summary>
-		/// User that provided the root cause analysis
-		/// </summary>
-		public string? RootCauseOwnerId { get; set; }
-
-		/// <summary>
-		/// User that provided the root cause analysis
-		/// </summary>
-		public string? RootCauseOwner { get; set; }
-
 		/// <summary>
 		/// User that provided the root cause analysis
 		/// </summary>
