Fix missing history of phase session results from compared stream

#rnx

#Synced-CL 48393387

diff --git a/Engine/Source/Programs/Horde/HordeDashboard/plugins/test-automation/testAutomationPhaseView.tsx b/Engine/Source/Programs/Horde/HordeDashboard/plugins/test-automation/testAutomationPhaseView.tsx
--- a/Engine/Source/Programs/Horde/HordeDashboard/plugins/test-automation/testAutomationPhaseView.tsx
+++ b/Engine/Source/Programs/Horde/HordeDashboard/plugins/test-automation/testAutomationPhaseView.tsx
@@ -57,7 +57,7 @@ const TestPhaseExpandPane: React.FC<{ phase: TestPhaseStatus, onEventsLoaded?: (
     useEffect(() => {
         // query phase data
         setSessionId(undefined);
-        handler.queryPhase(phase, handler.stream!).then(() => {
+        handler.queryPhase(phase, phase.testSession?.streamId ?? handler.stream!).then(() => {
             setComparison(comparisonContext.isEnable);
             setSessionId(phase.session?.id ?? "");
             if (!phase.session) {
