Change Incremental for Incrementally

#rnx

#Synced-CL 48393387

diff --git a/Engine/Source/Programs/Horde/HordeDashboard/plugins/test-automation/testHealthData.ts b/Engine/Source/Programs/Horde/HordeDashboard/plugins/test-automation/testHealthData.ts
--- a/Engine/Source/Programs/Horde/HordeDashboard/plugins/test-automation/testHealthData.ts
+++ b/Engine/Source/Programs/Horde/HordeDashboard/plugins/test-automation/testHealthData.ts
@@ -33,7 +33,7 @@ export const healthColors: Map<TestHealth, string> = new Map([
 
 export enum TestCadence {
     Nightly = "Nightly",
-    Incremental = "Incremental",
+    Incremental = "Incrementally",
     Daily = "Daily",
     Limited = "Limited",
     Weekly = "Weekly",
