[Test Health Details] Update user picker handle tooltips
+ change text size for selected stream

#rnx

#Synced-CL 48393387

diff --git a/Engine/Source/Programs/Horde/HordeDashboard/plugins/test-automation/testAutomationHealth.tsx b/Engine/Source/Programs/Horde/HordeDashboard/plugins/test-automation/testAutomationHealth.tsx
--- a/Engine/Source/Programs/Horde/HordeDashboard/plugins/test-automation/testAutomationHealth.tsx
+++ b/Engine/Source/Programs/Horde/HordeDashboard/plugins/test-automation/testAutomationHealth.tsx
@@ -621,7 +621,7 @@ export const TestHealthView: React.FC<{ test: TestHealthItem, handler: TestDataH
                             </Stack>
                         </Stack>
 
-                        <Stack horizontal tokens={{childrenGap: 12}} style={{paddingLeft: 12}} title="edit disabled for now">
+                        <Stack horizontal tokens={{childrenGap: 12}} style={{paddingLeft: 12}}>
                             <Stack style={{borderRadius: 4}} verticalAlign="center">
                                 <Stack  horizontal tokens={{childrenGap: 6}} verticalAlign="center">
                                     <Label>Owner: </Label>
@@ -629,14 +629,14 @@ export const TestHealthView: React.FC<{ test: TestHealthItem, handler: TestDataH
                                 <Stack style={{backgroundColor: booleanColors.get(false), padding: '2px 9px', borderRadius: 11}} horizontalAlign="center" verticalAlign="center">
                                     {!!userCache.current?.get(test.owner) &&
                                         <Stack horizontal verticalAlign="center">
-                                            <Text style={{cursor: 'pointer', color:textEmphasisColor()}}
+                                            <Text style={{cursor: 'pointer', color:textEmphasisColor()}} title="change owner"
                                                 onClick={() => setUserPicker({target: UserType.Owner, defaultUser: userCache.current?.get(test.owner)})}
                                             >{userCache.current?.get(test.owner)?.name ?? <Spinner size={SpinnerSize.xSmall} />}</Text>
-                                            <FontIcon iconName="Cancel" style={{cursor: 'pointer', fontSize: 11, paddingLeft: 6}} onClick={() => {test.owner = undefined}}/>
+                                            <FontIcon iconName="Cancel" style={{cursor: 'pointer', fontSize: 11, paddingLeft: 6}} title="remove" onClick={() => {test.owner = undefined}}/>
                                         </Stack>
                                     }
                                     {!userCache.current?.get(test.owner) &&
-                                        <FontIcon iconName="AddTo" style={{cursor: 'pointer', fontSize: 15, padding: 2}}
+                                        <FontIcon iconName="AddTo" style={{cursor: 'pointer', fontSize: 15, padding: 2}} title="add an owner"
                                             onClick={() => setUserPicker({target: UserType.Owner})}
                                         />
                                     }
@@ -651,12 +651,12 @@ export const TestHealthView: React.FC<{ test: TestHealthItem, handler: TestDataH
                                         test.customers.map(user => {
                                             return <Stack style={{backgroundColor: booleanColors.get(false), padding: '2px 9px', borderRadius: 11}} verticalAlign="center" horizontal>
                                                         <Text style={{color:textEmphasisColor()}}>{userCache.current?.get(user)?.name ?? <Spinner size={SpinnerSize.xSmall} />}</Text>
-                                                        <FontIcon iconName="Cancel" style={{cursor: 'pointer', fontSize: 11, paddingLeft: 6}} onClick={() => {test.customers = test.customers?.filter(c => c !== user)}}/>
+                                                        <FontIcon iconName="Cancel" style={{cursor: 'pointer', fontSize: 11, paddingLeft: 6}} title="remove" onClick={() => {test.customers = test.customers?.filter(c => c !== user)}}/>
                                                     </Stack>
                                         })
                                     }
                                     <Stack style={{backgroundColor: booleanColors.get(false), padding: '2px 9px', borderRadius: 11, width: 46}} verticalAlign="center" horizontalAlign="center">
-                                        <FontIcon iconName="AddTo" style={{cursor: 'pointer', fontSize: 15, padding: 2}}
+                                        <FontIcon iconName="AddTo" style={{cursor: 'pointer', fontSize: 15, padding: 2}} title="add a customer"
                                             onClick={() => setUserPicker({target: UserType.Customers})}
                                         />
                                     </Stack>
@@ -700,8 +700,8 @@ export const TestHealthView: React.FC<{ test: TestHealthItem, handler: TestDataH
                     <Stack style={{paddingLeft: 16}} tokens={{childrenGap: 16}}>
                         <Stack style={{paddingLeft: 12}} tokens={{childrenGap: 16}}>
                             <Stack horizontal tokens={{childrenGap: 12}} verticalAlign="center">
-                                <Stack style={{width: 200, backgroundColor: "#8e8e8e48", padding:  '3px 12px', borderRadius: 3}}>
-                                    <Text>{projectStore.streamById(streams.current)?.fullname ?? streams.current}</Text>
+                                <Stack style={{width: 200, backgroundColor: "#8e8e8e48", padding: '3px 12px', borderRadius: 3}} horizontalAlign="center">
+                                    <Text variant="small">{projectStore.streamById(streams.current)?.fullname ?? streams.current}</Text>
                                 </Stack>
                                 <Stack style={{cursor: 'default'}}>{renderHealthBadge(test, 75)}</Stack>
                                 <Stack style={{backgroundColor: cadenceColors.get(test.cadence), padding:  '3px 12px', borderRadius: 4, cursor: 'default'}} verticalAlign="center">
