Fix a bug when no test session metadata exists for the given set of filters

#rnx

#Synced-CL 48393387

diff --git a/Engine/Source/Programs/Horde/HordeDashboard/plugins/test-automation/testAutomationHealth.tsx b/Engine/Source/Programs/Horde/HordeDashboard/plugins/test-automation/testAutomationHealth.tsx
--- a/Engine/Source/Programs/Horde/HordeDashboard/plugins/test-automation/testAutomationHealth.tsx
+++ b/Engine/Source/Programs/Horde/HordeDashboard/plugins/test-automation/testAutomationHealth.tsx
@@ -348,7 +348,7 @@ const getTestHealthItems = (handler: TestDataHandler, stream?: string, test?: Te
                 const pCount = !!pmeta? status.sessions.get(pmeta)?.history.length ?? 0 : 0;
                 const cCount = status.sessions.get(meta)?.history.length ?? 0;
                 return cCount >= pCount ? meta : pmeta;
-            }, undefined)!)!.history.length;
+            }, undefined)!)?.history.length ?? 0;
 
             if (highestRunCount > 0) {
                 const rangeDayCount = handler.filterState.weeks * 7;
