Add missing UnrealEngine.csproj.props to dotnet projects

#rnx
#Synced-CL 47605596

diff --git a/Engine/Source/Programs/Horde/Drivers/JobDriver.Tests/JobDriver.Tests.csproj b/Engine/Source/Programs/Horde/Drivers/JobDriver.Tests/JobDriver.Tests.csproj
--- a/Engine/Source/Programs/Horde/Drivers/JobDriver.Tests/JobDriver.Tests.csproj
+++ b/Engine/Source/Programs/Horde/Drivers/JobDriver.Tests/JobDriver.Tests.csproj
@@ -1,4 +1,5 @@
 <Project Sdk="Microsoft.NET.Sdk">
+  <Import Project="..\..\..\Shared\UnrealEngine.csproj.props" />
 
   <PropertyGroup>
     <TargetFramework>net8.0</TargetFramework>
diff --git a/Engine/Source/Programs/Horde/Drivers/JobDriver/JobDriver.csproj b/Engine/Source/Programs/Horde/Drivers/JobDriver/JobDriver.csproj
--- a/Engine/Source/Programs/Horde/Drivers/JobDriver/JobDriver.csproj
+++ b/Engine/Source/Programs/Horde/Drivers/JobDriver/JobDriver.csproj
@@ -1,4 +1,5 @@
 ﻿<Project Sdk="Microsoft.NET.Sdk">
+  <Import Project="..\..\..\Shared\UnrealEngine.csproj.props" />
 
   <PropertyGroup>
     <OutputType>Exe</OutputType>
diff --git a/Engine/Source/Programs/Horde/HordeAgent.Tests/HordeAgent.Tests.csproj b/Engine/Source/Programs/Horde/HordeAgent.Tests/HordeAgent.Tests.csproj
--- a/Engine/Source/Programs/Horde/HordeAgent.Tests/HordeAgent.Tests.csproj
+++ b/Engine/Source/Programs/Horde/HordeAgent.Tests/HordeAgent.Tests.csproj
@@ -1,4 +1,5 @@
 <Project Sdk="Microsoft.NET.Sdk">
+  <Import Project="..\..\Shared\UnrealEngine.csproj.props" />
 
   <PropertyGroup>
     <TargetFramework>net8.0</TargetFramework>
diff --git a/Engine/Source/Programs/Horde/HordeAgent/HordeAgent.csproj b/Engine/Source/Programs/Horde/HordeAgent/HordeAgent.csproj
--- a/Engine/Source/Programs/Horde/HordeAgent/HordeAgent.csproj
+++ b/Engine/Source/Programs/Horde/HordeAgent/HordeAgent.csproj
@@ -1,4 +1,5 @@
 ﻿<Project Sdk="Microsoft.NET.Sdk">
+  <Import Project="..\..\Shared\UnrealEngine.csproj.props" />
 
   <PropertyGroup>
     <OutputType>Exe</OutputType>
diff --git a/Engine/Source/Programs/Horde/HordeCmd/HordeCmd.csproj b/Engine/Source/Programs/Horde/HordeCmd/HordeCmd.csproj
--- a/Engine/Source/Programs/Horde/HordeCmd/HordeCmd.csproj
+++ b/Engine/Source/Programs/Horde/HordeCmd/HordeCmd.csproj
@@ -1,4 +1,5 @@
 <Project Sdk="Microsoft.NET.Sdk">
+  <Import Project="..\..\Shared\UnrealEngine.csproj.props" />
 
   <PropertyGroup>
     <OutputType>Exe</OutputType>
diff --git a/Engine/Source/Programs/Horde/HordeServer.Shared/HordeServer.Shared.csproj b/Engine/Source/Programs/Horde/HordeServer.Shared/HordeServer.Shared.csproj
--- a/Engine/Source/Programs/Horde/HordeServer.Shared/HordeServer.Shared.csproj
+++ b/Engine/Source/Programs/Horde/HordeServer.Shared/HordeServer.Shared.csproj
@@ -1,4 +1,5 @@
 ﻿<Project Sdk="Microsoft.NET.Sdk">
+  <Import Project="..\..\Shared\UnrealEngine.csproj.props" />
 
   <PropertyGroup>
     <TargetFramework>net8.0</TargetFramework>
diff --git a/Engine/Source/Programs/Horde/HordeServer.Tests.Shared/HordeServer.Tests.Shared.csproj b/Engine/Source/Programs/Horde/HordeServer.Tests.Shared/HordeServer.Tests.Shared.csproj
--- a/Engine/Source/Programs/Horde/HordeServer.Tests.Shared/HordeServer.Tests.Shared.csproj
+++ b/Engine/Source/Programs/Horde/HordeServer.Tests.Shared/HordeServer.Tests.Shared.csproj
@@ -1,4 +1,5 @@
 ﻿<Project Sdk="Microsoft.NET.Sdk">
+  <Import Project="..\..\Shared\UnrealEngine.csproj.props" />
 
   <PropertyGroup>
     <TargetFramework>net8.0</TargetFramework>
diff --git a/Engine/Source/Programs/Horde/HordeServer.Tests/HordeServer.Tests.csproj b/Engine/Source/Programs/Horde/HordeServer.Tests/HordeServer.Tests.csproj
--- a/Engine/Source/Programs/Horde/HordeServer.Tests/HordeServer.Tests.csproj
+++ b/Engine/Source/Programs/Horde/HordeServer.Tests/HordeServer.Tests.csproj
@@ -1,4 +1,5 @@
 <Project Sdk="Microsoft.NET.Sdk">
+  <Import Project="..\..\Shared\UnrealEngine.csproj.props" />
 
   <PropertyGroup>
     <TargetFramework>net8.0</TargetFramework>
diff --git a/Engine/Source/Programs/Horde/HordeServer/HordeServer.csproj b/Engine/Source/Programs/Horde/HordeServer/HordeServer.csproj
--- a/Engine/Source/Programs/Horde/HordeServer/HordeServer.csproj
+++ b/Engine/Source/Programs/Horde/HordeServer/HordeServer.csproj
@@ -1,4 +1,5 @@
 ﻿<Project Sdk="Microsoft.NET.Sdk.Web">
+  <Import Project="..\..\Shared\UnrealEngine.csproj.props" />
 
   <PropertyGroup>
     <TargetFramework>net8.0</TargetFramework>
diff --git a/Engine/Source/Programs/Horde/Plugins/Analytics/HordeServer.Analytics.Integration.Tests/HordeServer.Analytics.Integration.Tests.csproj b/Engine/Source/Programs/Horde/Plugins/Analytics/HordeServer.Analytics.Integration.Tests/HordeServer.Analytics.Integration.Tests.csproj
--- a/Engine/Source/Programs/Horde/Plugins/Analytics/HordeServer.Analytics.Integration.Tests/HordeServer.Analytics.Integration.Tests.csproj
+++ b/Engine/Source/Programs/Horde/Plugins/Analytics/HordeServer.Analytics.Integration.Tests/HordeServer.Analytics.Integration.Tests.csproj
@@ -1,4 +1,5 @@
 <Project Sdk="Microsoft.NET.Sdk">
+  <Import Project="..\..\..\..\Shared\UnrealEngine.csproj.props" />
 
   <PropertyGroup>
     <TargetFramework>net8.0</TargetFramework>
diff --git a/Engine/Source/Programs/Horde/Plugins/Analytics/HordeServer.Analytics.Tests/HordeServer.Analytics.Tests.csproj b/Engine/Source/Programs/Horde/Plugins/Analytics/HordeServer.Analytics.Tests/HordeServer.Analytics.Tests.csproj
--- a/Engine/Source/Programs/Horde/Plugins/Analytics/HordeServer.Analytics.Tests/HordeServer.Analytics.Tests.csproj
+++ b/Engine/Source/Programs/Horde/Plugins/Analytics/HordeServer.Analytics.Tests/HordeServer.Analytics.Tests.csproj
@@ -1,4 +1,5 @@
 <Project Sdk="Microsoft.NET.Sdk">
+  <Import Project="..\..\..\..\Shared\UnrealEngine.csproj.props" />
 
   <PropertyGroup>
     <TargetFramework>net8.0</TargetFramework>
diff --git a/Engine/Source/Programs/Horde/Plugins/Analytics/HordeServer.Analytics/HordeServer.Analytics.csproj b/Engine/Source/Programs/Horde/Plugins/Analytics/HordeServer.Analytics/HordeServer.Analytics.csproj
--- a/Engine/Source/Programs/Horde/Plugins/Analytics/HordeServer.Analytics/HordeServer.Analytics.csproj
+++ b/Engine/Source/Programs/Horde/Plugins/Analytics/HordeServer.Analytics/HordeServer.Analytics.csproj
@@ -1,4 +1,5 @@
 ﻿<Project Sdk="Microsoft.NET.Sdk">
+  <Import Project="..\..\..\..\Shared\UnrealEngine.csproj.props" />
 
   <PropertyGroup>
     <TargetFramework>net8.0</TargetFramework>
diff --git a/Engine/Source/Programs/Horde/Plugins/Build/HordeServer.Build.Tests/HordeServer.Build.Tests.csproj b/Engine/Source/Programs/Horde/Plugins/Build/HordeServer.Build.Tests/HordeServer.Build.Tests.csproj
--- a/Engine/Source/Programs/Horde/Plugins/Build/HordeServer.Build.Tests/HordeServer.Build.Tests.csproj
+++ b/Engine/Source/Programs/Horde/Plugins/Build/HordeServer.Build.Tests/HordeServer.Build.Tests.csproj
@@ -1,4 +1,5 @@
 <Project Sdk="Microsoft.NET.Sdk">
+  <Import Project="..\..\..\..\Shared\UnrealEngine.csproj.props" />
 
   <PropertyGroup>
     <TargetFramework>net8.0</TargetFramework>
diff --git a/Engine/Source/Programs/Horde/Plugins/Build/HordeServer.Build/HordeServer.Build.csproj b/Engine/Source/Programs/Horde/Plugins/Build/HordeServer.Build/HordeServer.Build.csproj
--- a/Engine/Source/Programs/Horde/Plugins/Build/HordeServer.Build/HordeServer.Build.csproj
+++ b/Engine/Source/Programs/Horde/Plugins/Build/HordeServer.Build/HordeServer.Build.csproj
@@ -1,4 +1,5 @@
 ﻿<Project Sdk="Microsoft.NET.Sdk">
+  <Import Project="..\..\..\..\Shared\UnrealEngine.csproj.props" />
 
   <PropertyGroup>
     <TargetFramework>net8.0</TargetFramework>
diff --git a/Engine/Source/Programs/Horde/Plugins/Compute/HordeServer.Compute.Tests/HordeServer.Compute.Tests.csproj b/Engine/Source/Programs/Horde/Plugins/Compute/HordeServer.Compute.Tests/HordeServer.Compute.Tests.csproj
--- a/Engine/Source/Programs/Horde/Plugins/Compute/HordeServer.Compute.Tests/HordeServer.Compute.Tests.csproj
+++ b/Engine/Source/Programs/Horde/Plugins/Compute/HordeServer.Compute.Tests/HordeServer.Compute.Tests.csproj
@@ -1,4 +1,5 @@
 <Project Sdk="Microsoft.NET.Sdk">
+  <Import Project="..\..\..\..\Shared\UnrealEngine.csproj.props" />
 
   <PropertyGroup>
     <TargetFramework>net8.0</TargetFramework>
diff --git a/Engine/Source/Programs/Horde/Plugins/Compute/HordeServer.Compute/HordeServer.Compute.csproj b/Engine/Source/Programs/Horde/Plugins/Compute/HordeServer.Compute/HordeServer.Compute.csproj
--- a/Engine/Source/Programs/Horde/Plugins/Compute/HordeServer.Compute/HordeServer.Compute.csproj
+++ b/Engine/Source/Programs/Horde/Plugins/Compute/HordeServer.Compute/HordeServer.Compute.csproj
@@ -1,4 +1,5 @@
 ﻿<Project Sdk="Microsoft.NET.Sdk">
+  <Import Project="..\..\..\..\Shared\UnrealEngine.csproj.props" />
 
   <PropertyGroup>
     <TargetFramework>net8.0</TargetFramework>
diff --git a/Engine/Source/Programs/Horde/Plugins/ExamplePlugin/HordeServer.ExamplePlugin.Tests/HordeServer.ExamplePlugin.Tests.csproj b/Engine/Source/Programs/Horde/Plugins/ExamplePlugin/HordeServer.ExamplePlugin.Tests/HordeServer.ExamplePlugin.Tests.csproj
--- a/Engine/Source/Programs/Horde/Plugins/ExamplePlugin/HordeServer.ExamplePlugin.Tests/HordeServer.ExamplePlugin.Tests.csproj
+++ b/Engine/Source/Programs/Horde/Plugins/ExamplePlugin/HordeServer.ExamplePlugin.Tests/HordeServer.ExamplePlugin.Tests.csproj
@@ -1,4 +1,5 @@
 ﻿<Project Sdk="Microsoft.NET.Sdk">
+  <Import Project="..\..\..\..\Shared\UnrealEngine.csproj.props" />
 
   <PropertyGroup>
     <TargetFramework>net8.0</TargetFramework>
diff --git a/Engine/Source/Programs/Horde/Plugins/ExamplePlugin/HordeServer.ExamplePlugin/HordeServer.ExamplePlugin.csproj b/Engine/Source/Programs/Horde/Plugins/ExamplePlugin/HordeServer.ExamplePlugin/HordeServer.ExamplePlugin.csproj
--- a/Engine/Source/Programs/Horde/Plugins/ExamplePlugin/HordeServer.ExamplePlugin/HordeServer.ExamplePlugin.csproj
+++ b/Engine/Source/Programs/Horde/Plugins/ExamplePlugin/HordeServer.ExamplePlugin/HordeServer.ExamplePlugin.csproj
@@ -1,4 +1,5 @@
 ﻿<Project Sdk="Microsoft.NET.Sdk">
+  <Import Project="..\..\..\..\Shared\UnrealEngine.csproj.props" />
 
   <PropertyGroup>
     <TargetFramework>net8.0</TargetFramework>
diff --git a/Engine/Source/Programs/Horde/Plugins/Experimental/HordeServer.Experimental.Tests/HordeServer.Experimental.Tests.csproj b/Engine/Source/Programs/Horde/Plugins/Experimental/HordeServer.Experimental.Tests/HordeServer.Experimental.Tests.csproj
--- a/Engine/Source/Programs/Horde/Plugins/Experimental/HordeServer.Experimental.Tests/HordeServer.Experimental.Tests.csproj
+++ b/Engine/Source/Programs/Horde/Plugins/Experimental/HordeServer.Experimental.Tests/HordeServer.Experimental.Tests.csproj
@@ -1,4 +1,5 @@
 <Project Sdk="Microsoft.NET.Sdk">
+  <Import Project="..\..\..\..\Shared\UnrealEngine.csproj.props" />
 
   <PropertyGroup>
     <TargetFramework>net8.0</TargetFramework>
diff --git a/Engine/Source/Programs/Horde/Plugins/Experimental/HordeServer.Experimental/HordeServer.Experimental.csproj b/Engine/Source/Programs/Horde/Plugins/Experimental/HordeServer.Experimental/HordeServer.Experimental.csproj
--- a/Engine/Source/Programs/Horde/Plugins/Experimental/HordeServer.Experimental/HordeServer.Experimental.csproj
+++ b/Engine/Source/Programs/Horde/Plugins/Experimental/HordeServer.Experimental/HordeServer.Experimental.csproj
@@ -1,4 +1,5 @@
 ﻿<Project Sdk="Microsoft.NET.Sdk">
+  <Import Project="..\..\..\..\Shared\UnrealEngine.csproj.props" />
 
   <PropertyGroup>
     <TargetFramework>net8.0</TargetFramework>
diff --git a/Engine/Source/Programs/Horde/Plugins/Secrets/HordeServer.Secrets.Tests/HordeServer.Secrets.Tests.csproj b/Engine/Source/Programs/Horde/Plugins/Secrets/HordeServer.Secrets.Tests/HordeServer.Secrets.Tests.csproj
--- a/Engine/Source/Programs/Horde/Plugins/Secrets/HordeServer.Secrets.Tests/HordeServer.Secrets.Tests.csproj
+++ b/Engine/Source/Programs/Horde/Plugins/Secrets/HordeServer.Secrets.Tests/HordeServer.Secrets.Tests.csproj
@@ -1,4 +1,5 @@
 <Project Sdk="Microsoft.NET.Sdk">
+  <Import Project="..\..\..\..\Shared\UnrealEngine.csproj.props" />
 
   <PropertyGroup>
     <TargetFramework>net8.0</TargetFramework>
diff --git a/Engine/Source/Programs/Horde/Plugins/Secrets/HordeServer.Secrets/HordeServer.Secrets.csproj b/Engine/Source/Programs/Horde/Plugins/Secrets/HordeServer.Secrets/HordeServer.Secrets.csproj
--- a/Engine/Source/Programs/Horde/Plugins/Secrets/HordeServer.Secrets/HordeServer.Secrets.csproj
+++ b/Engine/Source/Programs/Horde/Plugins/Secrets/HordeServer.Secrets/HordeServer.Secrets.csproj
@@ -1,4 +1,5 @@
 <Project Sdk="Microsoft.NET.Sdk">
+  <Import Project="..\..\..\..\Shared\UnrealEngine.csproj.props" />
 
   <PropertyGroup>
     <TargetFramework>net8.0</TargetFramework>
diff --git a/Engine/Source/Programs/Horde/Plugins/Storage/HordeServer.Storage.Tests/HordeServer.Storage.Tests.csproj b/Engine/Source/Programs/Horde/Plugins/Storage/HordeServer.Storage.Tests/HordeServer.Storage.Tests.csproj
--- a/Engine/Source/Programs/Horde/Plugins/Storage/HordeServer.Storage.Tests/HordeServer.Storage.Tests.csproj
+++ b/Engine/Source/Programs/Horde/Plugins/Storage/HordeServer.Storage.Tests/HordeServer.Storage.Tests.csproj
@@ -1,4 +1,5 @@
 <Project Sdk="Microsoft.NET.Sdk">
+  <Import Project="..\..\..\..\Shared\UnrealEngine.csproj.props" />
 
   <PropertyGroup>
     <TargetFramework>net8.0</TargetFramework>
diff --git a/Engine/Source/Programs/Horde/Plugins/Storage/HordeServer.Storage/HordeServer.Storage.csproj b/Engine/Source/Programs/Horde/Plugins/Storage/HordeServer.Storage/HordeServer.Storage.csproj
--- a/Engine/Source/Programs/Horde/Plugins/Storage/HordeServer.Storage/HordeServer.Storage.csproj
+++ b/Engine/Source/Programs/Horde/Plugins/Storage/HordeServer.Storage/HordeServer.Storage.csproj
@@ -1,4 +1,5 @@
 ﻿<Project Sdk="Microsoft.NET.Sdk">
+  <Import Project="..\..\..\..\Shared\UnrealEngine.csproj.props" />
 
   <PropertyGroup>
     <TargetFramework>net8.0</TargetFramework>
diff --git a/Engine/Source/Programs/Horde/Plugins/Symbols/HordeServer.Symbols/HordeServer.Symbols.csproj b/Engine/Source/Programs/Horde/Plugins/Symbols/HordeServer.Symbols/HordeServer.Symbols.csproj
--- a/Engine/Source/Programs/Horde/Plugins/Symbols/HordeServer.Symbols/HordeServer.Symbols.csproj
+++ b/Engine/Source/Programs/Horde/Plugins/Symbols/HordeServer.Symbols/HordeServer.Symbols.csproj
@@ -1,4 +1,5 @@
 <Project Sdk="Microsoft.NET.Sdk">
+  <Import Project="..\..\..\..\Shared\UnrealEngine.csproj.props" />
 
   <PropertyGroup>
     <TargetFramework>net8.0</TargetFramework>
diff --git a/Engine/Source/Programs/Horde/Plugins/Tools/HordeServer.Tools.Tests/HordeServer.Tools.Tests.csproj b/Engine/Source/Programs/Horde/Plugins/Tools/HordeServer.Tools.Tests/HordeServer.Tools.Tests.csproj
--- a/Engine/Source/Programs/Horde/Plugins/Tools/HordeServer.Tools.Tests/HordeServer.Tools.Tests.csproj
+++ b/Engine/Source/Programs/Horde/Plugins/Tools/HordeServer.Tools.Tests/HordeServer.Tools.Tests.csproj
@@ -1,4 +1,5 @@
 <Project Sdk="Microsoft.NET.Sdk">
+  <Import Project="..\..\..\..\Shared\UnrealEngine.csproj.props" />
 
   <PropertyGroup>
     <TargetFramework>net8.0</TargetFramework>
diff --git a/Engine/Source/Programs/Horde/Plugins/Tools/HordeServer.Tools/HordeServer.Tools.csproj b/Engine/Source/Programs/Horde/Plugins/Tools/HordeServer.Tools/HordeServer.Tools.csproj
--- a/Engine/Source/Programs/Horde/Plugins/Tools/HordeServer.Tools/HordeServer.Tools.csproj
+++ b/Engine/Source/Programs/Horde/Plugins/Tools/HordeServer.Tools/HordeServer.Tools.csproj
@@ -1,4 +1,5 @@
 ﻿<Project Sdk="Microsoft.NET.Sdk">
+  <Import Project="..\..\..\..\Shared\UnrealEngine.csproj.props" />
 
   <PropertyGroup>
     <TargetFramework>net8.0</TargetFramework>
diff --git a/Engine/Source/Programs/Horde/Samples/RemoteClient/RemoteClient.csproj b/Engine/Source/Programs/Horde/Samples/RemoteClient/RemoteClient.csproj
--- a/Engine/Source/Programs/Horde/Samples/RemoteClient/RemoteClient.csproj
+++ b/Engine/Source/Programs/Horde/Samples/RemoteClient/RemoteClient.csproj
@@ -1,4 +1,5 @@
 <Project Sdk="Microsoft.NET.Sdk">
+  <Import Project="..\..\..\Shared\UnrealEngine.csproj.props" />
 
   <PropertyGroup>
     <OutputType>Exe</OutputType>
diff --git a/Engine/Source/Programs/Horde/Samples/RemoteWorker/RemoteWorker.csproj b/Engine/Source/Programs/Horde/Samples/RemoteWorker/RemoteWorker.csproj
--- a/Engine/Source/Programs/Horde/Samples/RemoteWorker/RemoteWorker.csproj
+++ b/Engine/Source/Programs/Horde/Samples/RemoteWorker/RemoteWorker.csproj
@@ -1,4 +1,5 @@
 <Project Sdk="Microsoft.NET.Sdk">
+  <Import Project="..\..\..\Shared\UnrealEngine.csproj.props" />
 
   <PropertyGroup>
     <OutputType>Exe</OutputType>
diff --git a/Engine/Source/Programs/Horde/Tools/ReleaseNotes/ReleaseNotes.csproj b/Engine/Source/Programs/Horde/Tools/ReleaseNotes/ReleaseNotes.csproj
--- a/Engine/Source/Programs/Horde/Tools/ReleaseNotes/ReleaseNotes.csproj
+++ b/Engine/Source/Programs/Horde/Tools/ReleaseNotes/ReleaseNotes.csproj
@@ -1,4 +1,5 @@
 <Project Sdk="Microsoft.NET.Sdk">
+  <Import Project="..\..\..\Shared\UnrealEngine.csproj.props" />
 
   <PropertyGroup>
     <OutputType>Exe</OutputType>
diff --git a/Engine/Source/Programs/OidcToken/OidcToken.csproj b/Engine/Source/Programs/OidcToken/OidcToken.csproj
--- a/Engine/Source/Programs/OidcToken/OidcToken.csproj
+++ b/Engine/Source/Programs/OidcToken/OidcToken.csproj
@@ -1,4 +1,5 @@
 ﻿<Project Sdk="Microsoft.NET.Sdk">
+  <Import Project="..\Shared\UnrealEngine.csproj.props" />
 
   <PropertyGroup>
     <OutputType>Exe</OutputType>
diff --git a/Engine/Source/Programs/Shared/EpicGames.AspNet/EpicGames.AspNet.csproj b/Engine/Source/Programs/Shared/EpicGames.AspNet/EpicGames.AspNet.csproj
--- a/Engine/Source/Programs/Shared/EpicGames.AspNet/EpicGames.AspNet.csproj
+++ b/Engine/Source/Programs/Shared/EpicGames.AspNet/EpicGames.AspNet.csproj
@@ -1,4 +1,6 @@
 <Project Sdk="Microsoft.NET.Sdk.Web">
+  <Import Project="..\UnrealEngine.csproj.props" />
+
   <PropertyGroup>
     <TargetFramework>net8.0</TargetFramework>
     <Configuration Condition=" '$(Configuration)' == '' ">Development</Configuration>
diff --git a/Engine/Source/Programs/Shared/EpicGames.MongoDB/EpicGames.MongoDB.csproj b/Engine/Source/Programs/Shared/EpicGames.MongoDB/EpicGames.MongoDB.csproj
--- a/Engine/Source/Programs/Shared/EpicGames.MongoDB/EpicGames.MongoDB.csproj
+++ b/Engine/Source/Programs/Shared/EpicGames.MongoDB/EpicGames.MongoDB.csproj
@@ -1,4 +1,6 @@
 <Project Sdk="Microsoft.NET.Sdk">
+  <Import Project="..\UnrealEngine.csproj.props" />
+
   <PropertyGroup>
     <TargetFramework>net8.0</TargetFramework>
     <Configuration Condition=" '$(Configuration)' == '' ">Development</Configuration>
diff --git a/Engine/Source/Programs/Shared/EpicGames.Perforce.Fixture/EpicGames.Perforce.Fixture.csproj b/Engine/Source/Programs/Shared/EpicGames.Perforce.Fixture/EpicGames.Perforce.Fixture.csproj
--- a/Engine/Source/Programs/Shared/EpicGames.Perforce.Fixture/EpicGames.Perforce.Fixture.csproj
+++ b/Engine/Source/Programs/Shared/EpicGames.Perforce.Fixture/EpicGames.Perforce.Fixture.csproj
@@ -1,4 +1,5 @@
 <Project Sdk="Microsoft.NET.Sdk">
+  <Import Project="..\UnrealEngine.csproj.props" />
 
   <PropertyGroup>
     <TargetFramework>net8.0</TargetFramework>
diff --git a/Engine/Source/Programs/Shared/EpicGames.Perforce.Managed.Tests/EpicGames.Perforce.Managed.Tests.csproj b/Engine/Source/Programs/Shared/EpicGames.Perforce.Managed.Tests/EpicGames.Perforce.Managed.Tests.csproj
--- a/Engine/Source/Programs/Shared/EpicGames.Perforce.Managed.Tests/EpicGames.Perforce.Managed.Tests.csproj
+++ b/Engine/Source/Programs/Shared/EpicGames.Perforce.Managed.Tests/EpicGames.Perforce.Managed.Tests.csproj
@@ -1,4 +1,5 @@
 <Project Sdk="Microsoft.NET.Sdk">
+  <Import Project="..\UnrealEngine.csproj.props" />
 
   <PropertyGroup>
     <TargetFramework>net8.0</TargetFramework>
diff --git a/Engine/Source/Programs/Shared/EpicGames.Perforce.Managed/EpicGames.Perforce.Managed.csproj b/Engine/Source/Programs/Shared/EpicGames.Perforce.Managed/EpicGames.Perforce.Managed.csproj
--- a/Engine/Source/Programs/Shared/EpicGames.Perforce.Managed/EpicGames.Perforce.Managed.csproj
+++ b/Engine/Source/Programs/Shared/EpicGames.Perforce.Managed/EpicGames.Perforce.Managed.csproj
@@ -1,4 +1,5 @@
 <Project Sdk="Microsoft.NET.Sdk">
+  <Import Project="..\UnrealEngine.csproj.props" />
 
   <PropertyGroup>
     <TargetFramework>net8.0</TargetFramework>
diff --git a/Engine/Source/Programs/Shared/EpicGames.Perforce.Tests/EpicGames.Perforce.Tests.csproj b/Engine/Source/Programs/Shared/EpicGames.Perforce.Tests/EpicGames.Perforce.Tests.csproj
--- a/Engine/Source/Programs/Shared/EpicGames.Perforce.Tests/EpicGames.Perforce.Tests.csproj
+++ b/Engine/Source/Programs/Shared/EpicGames.Perforce.Tests/EpicGames.Perforce.Tests.csproj
@@ -1,4 +1,5 @@
 <Project Sdk="Microsoft.NET.Sdk">
+  <Import Project="..\UnrealEngine.csproj.props" />
 
   <PropertyGroup>
     <TargetFramework>net8.0</TargetFramework>
diff --git a/Engine/Source/Programs/Shared/EpicGames.Redis.Tests/EpicGames.Redis.Tests.csproj b/Engine/Source/Programs/Shared/EpicGames.Redis.Tests/EpicGames.Redis.Tests.csproj
--- a/Engine/Source/Programs/Shared/EpicGames.Redis.Tests/EpicGames.Redis.Tests.csproj
+++ b/Engine/Source/Programs/Shared/EpicGames.Redis.Tests/EpicGames.Redis.Tests.csproj
@@ -1,4 +1,5 @@
 <Project Sdk="Microsoft.NET.Sdk">
+  <Import Project="..\UnrealEngine.csproj.props" />
 
   <PropertyGroup>
     <TargetFramework>net8.0</TargetFramework>
diff --git a/Engine/Source/Programs/Shared/EpicGames.Redis/EpicGames.Redis.csproj b/Engine/Source/Programs/Shared/EpicGames.Redis/EpicGames.Redis.csproj
--- a/Engine/Source/Programs/Shared/EpicGames.Redis/EpicGames.Redis.csproj
+++ b/Engine/Source/Programs/Shared/EpicGames.Redis/EpicGames.Redis.csproj
@@ -1,4 +1,6 @@
 <Project Sdk="Microsoft.NET.Sdk">
+  <Import Project="..\UnrealEngine.csproj.props" />
+  
   <PropertyGroup>
     <TargetFramework>net8.0</TargetFramework>
     <Configuration Condition=" '$(Configuration)' == '' ">Development</Configuration>
diff --git a/Engine/Source/Programs/Shared/EpicGames.Tracing.Tests/EpicGames.Tracing.Tests.csproj b/Engine/Source/Programs/Shared/EpicGames.Tracing.Tests/EpicGames.Tracing.Tests.csproj
--- a/Engine/Source/Programs/Shared/EpicGames.Tracing.Tests/EpicGames.Tracing.Tests.csproj
+++ b/Engine/Source/Programs/Shared/EpicGames.Tracing.Tests/EpicGames.Tracing.Tests.csproj
@@ -1,4 +1,5 @@
 <Project Sdk="Microsoft.NET.Sdk">
+  <Import Project="..\UnrealEngine.csproj.props" />
 
   <PropertyGroup>
     <TargetFramework>net8.0</TargetFramework>
diff --git a/Engine/Source/Programs/Shared/EpicGames.Tracing/EpicGames.Tracing.csproj b/Engine/Source/Programs/Shared/EpicGames.Tracing/EpicGames.Tracing.csproj
--- a/Engine/Source/Programs/Shared/EpicGames.Tracing/EpicGames.Tracing.csproj
+++ b/Engine/Source/Programs/Shared/EpicGames.Tracing/EpicGames.Tracing.csproj
@@ -1,4 +1,6 @@
 <Project Sdk="Microsoft.NET.Sdk">
+  <Import Project="..\UnrealEngine.csproj.props" />
+
   <PropertyGroup>
     <TargetFramework>net8.0</TargetFramework>
     <Configuration Condition=" '$(Configuration)' == '' ">Development</Configuration>
diff --git a/Engine/Source/Programs/UnrealGameSync/UnrealGameSync/UnrealGameSync.csproj b/Engine/Source/Programs/UnrealGameSync/UnrealGameSync/UnrealGameSync.csproj
--- a/Engine/Source/Programs/UnrealGameSync/UnrealGameSync/UnrealGameSync.csproj
+++ b/Engine/Source/Programs/UnrealGameSync/UnrealGameSync/UnrealGameSync.csproj
@@ -1,4 +1,6 @@
 <Project Sdk="Microsoft.NET.Sdk">
+  <Import Project="..\..\Shared\UnrealEngine.csproj.props" />
+
   <PropertyGroup>
     <OutputType>WinExe</OutputType>
     <TargetFramework>net8.0-windows</TargetFramework>
diff --git a/Engine/Source/Programs/UnrealGameSync/UnrealGameSyncAutomation/UnrealGameSyncAutomation.csproj b/Engine/Source/Programs/UnrealGameSync/UnrealGameSyncAutomation/UnrealGameSyncAutomation.csproj
--- a/Engine/Source/Programs/UnrealGameSync/UnrealGameSyncAutomation/UnrealGameSyncAutomation.csproj
+++ b/Engine/Source/Programs/UnrealGameSync/UnrealGameSyncAutomation/UnrealGameSyncAutomation.csproj
@@ -1,5 +1,7 @@
 <?xml version="1.0" encoding="utf-8"?>
 <Project Sdk="Microsoft.NET.Sdk">
+  <Import Project="..\..\Shared\UnrealEngine.csproj.props" />
+
   <PropertyGroup>
     <OutputType>WinExe</OutputType>
     <TargetFramework>net8.0-windows</TargetFramework>
diff --git a/Engine/Source/Programs/UnrealGameSync/UnrealGameSyncCmd/UnrealGameSyncCmd.csproj b/Engine/Source/Programs/UnrealGameSync/UnrealGameSyncCmd/UnrealGameSyncCmd.csproj
--- a/Engine/Source/Programs/UnrealGameSync/UnrealGameSyncCmd/UnrealGameSyncCmd.csproj
+++ b/Engine/Source/Programs/UnrealGameSync/UnrealGameSyncCmd/UnrealGameSyncCmd.csproj
@@ -1,4 +1,5 @@
 <Project Sdk="Microsoft.NET.Sdk">
+  <Import Project="..\..\Shared\UnrealEngine.csproj.props" />
 
   <PropertyGroup>
     <OutputType>Exe</OutputType>
diff --git a/Engine/Source/Programs/UnrealGameSync/UnrealGameSyncShared/UnrealGameSyncShared.csproj b/Engine/Source/Programs/UnrealGameSync/UnrealGameSyncShared/UnrealGameSyncShared.csproj
--- a/Engine/Source/Programs/UnrealGameSync/UnrealGameSyncShared/UnrealGameSyncShared.csproj
+++ b/Engine/Source/Programs/UnrealGameSync/UnrealGameSyncShared/UnrealGameSyncShared.csproj
@@ -1,4 +1,5 @@
 <Project Sdk="Microsoft.NET.Sdk">
+  <Import Project="..\..\Shared\UnrealEngine.csproj.props" />
 
   <PropertyGroup>
     <TargetFramework>net8.0</TargetFramework>
diff --git a/Engine/Source/Programs/UnrealToolbox/UnrealToolbox.csproj b/Engine/Source/Programs/UnrealToolbox/UnrealToolbox.csproj
--- a/Engine/Source/Programs/UnrealToolbox/UnrealToolbox.csproj
+++ b/Engine/Source/Programs/UnrealToolbox/UnrealToolbox.csproj
@@ -1,4 +1,5 @@
 ﻿<Project Sdk="Microsoft.NET.Sdk">
+  <Import Project="..\Shared\UnrealEngine.csproj.props" />
 
   <ItemGroup>
     <FrameworkReference Include="Microsoft.AspNetCore.App" />
