Horde: Remove ResolvedBy from being the default Root Cause Analysis Owner

#rnx
#Synced-CL Unknown

diff --git a/Engine/Source/Programs/Horde/Plugins/Build/HordeServer.Build/Issues/IssueCollection.cs b/Engine/Source/Programs/Horde/Plugins/Build/HordeServer.Build/Issues/IssueCollection.cs
--- a/Engine/Source/Programs/Horde/Plugins/Build/HordeServer.Build/Issues/IssueCollection.cs
+++ b/Engine/Source/Programs/Horde/Plugins/Build/HordeServer.Build/Issues/IssueCollection.cs
@@ -1340,8 +1340,11 @@ namespace HordeServer.Issues
 			}
 			if (options.ResolvedById != null)
 			{
+				if (options.OwnerId != null && options.RootCauseOwnerId == null && issueDocument.RootCauseOwnerId == null)
+				{
+					options = options with { RootCauseOwnerId = options.OwnerId };
+				}
 				options = options with { OwnerId = options.OwnerId ?? options.ResolvedById, Acknowledged = options.Acknowledged ?? true };
-				options = options with { RootCauseOwnerId = options.RootCauseOwnerId ?? options.OwnerId ?? options.ResolvedById };
 			}
 			if (options.OwnerId != null)
 			{
diff --git a/Engine/Source/Programs/Horde/Plugins/Build/HordeServer.Build/Notifications/Sinks/SlackNotificationSink.cs b/Engine/Source/Programs/Horde/Plugins/Build/HordeServer.Build/Notifications/Sinks/SlackNotificationSink.cs
--- a/Engine/Source/Programs/Horde/Plugins/Build/HordeServer.Build/Notifications/Sinks/SlackNotificationSink.cs
+++ b/Engine/Source/Programs/Horde/Plugins/Build/HordeServer.Build/Notifications/Sinks/SlackNotificationSink.cs
@@ -3100,7 +3100,7 @@ namespace HordeServer.Notifications.Sinks
 				return null;
 			}
 
-			UserId? ownerId = issue.RootCauseOwnerId ?? issue.OwnerId ?? issue.ResolvedById;
+			UserId? ownerId = issue.RootCauseOwnerId ?? issue.OwnerId;
 			if (ownerId == null || ownerId == UserId.Empty)
 			{
 				_logger.LogInformation(CanceledMessage, notificationName, issueId, "No Root Cause Owner");
