Display alternate run on same commit in history graphs

#rnx

#Synced-CL 48452052

diff --git a/Engine/Source/Programs/Horde/HordeDashboard/plugins/test-automation/phaseHistoryGraph.ts b/Engine/Source/Programs/Horde/HordeDashboard/plugins/test-automation/phaseHistoryGraph.ts
--- a/Engine/Source/Programs/Horde/HordeDashboard/plugins/test-automation/phaseHistoryGraph.ts
+++ b/Engine/Source/Programs/Horde/HordeDashboard/plugins/test-automation/phaseHistoryGraph.ts
@@ -6,7 +6,8 @@ import dashboard, { StatusColor } from 'horde/backend/Dashboard';
 import { getHumanTime, getShortNiceTime } from 'horde/base/utilities/timeUtils';
 import { getHordeStyling } from 'horde/styles/Styles';
 import { TestOutcome } from "./api";
-import { getPhaseSessionStatusColor, getStatusColors } from "./testAutomationCommon";
+import { getPhaseSessionStatusColor, getStatusColors, SessionListSelector } from "./testAutomationCommon";
+import { renderToString } from "react-dom/server";
 
 // Handle bad "@types/d3" types, fix if addressed upstream
 const _d3 = d3 as any;
@@ -141,16 +142,22 @@ export class PhaseHistoryGraph {
             .attr("x2", ([, I]: any) => width - this.margin.right);
 
         g.selectAll("circle")
-            .data(([, I]: any) => I)
+            .data(([, I]: any) => _d3.group(I, i => X[i]))
             .enter()
             .append("g")
-            .attr("id", i => `circle${sessions[i as any].id}`)
-            .attr("transform", i => `translate(${xScale(X[i as any])},0)`)
+            .attr("id", ([, items]: any) => `circle${sessions[items[0]].id}`)
+            .attr("transform", ([, items]: any) => `translate(${xScale(X[items[0] as any])},0)`)
             .append("g")
             .attr("class", "circleScale")
-            .attr("id", i => `circleScale${sessions[i as any].id}`)
+            .attr("id", ([, items]: any) => `circleScale${sessions[items[0]].id}`)
             .style("cursor", "pointer")
-            .each((d, i, nodes) => SessionCircle(nodes[i], sessions[d as any], 5, !!this.phaseSessionId && sessions[d as any].id === this.phaseSessionId));
+            .each(([, items]: any, i, nodes) => SessionCircle(nodes[i], items.map(i => sessions[i]), this.phaseSessionId, 5, !!this.phaseSessionId && items.some(i => sessions[i].id === this.phaseSessionId)));
+
+        const reDrawCircles = (items: PhaseSessionResult[], index: number) => {
+            const circle = svg.select(`#circleScale${items[0].id}`);
+            circle.selectAll("*").remove();
+            SessionCircle(circle.node() as any, items, items[index].id, 5, !!this.phaseSessionId && items.some(i => i.id === this.phaseSessionId));
+        }
 
         const tooltip = this.tooltip = d3.select(container)
             .append("div")
@@ -165,7 +172,7 @@ export class PhaseHistoryGraph {
             .style("position", "absolute")
             .style("pointer-events", "none");
 
-        const closestData = (x: number, y: number): PhaseSessionResult | undefined => {
+        const closestData = (x: number, y: number): PhaseSessionResult[] | undefined => {
 
             if (x < this.margin.left - 16) {
                 return undefined;
@@ -181,57 +188,79 @@ export class PhaseHistoryGraph {
                 const lengthSqr = absy * absy + absx * absx;
 
                 if (lengthSqr < best.value) {
-                    return { index: i, value: lengthSqr };
-                }
-                else {
+                    return { indexes: [i], value: lengthSqr };
+                } else {
+                     if (lengthSqr === best.value) best.indexes.push(i);
                     return best;
                 }
-            }, { index: 0, value: Number.MAX_SAFE_INTEGER });
+            }, { indexes: [] as number[], value: Number.MAX_SAFE_INTEGER });
 
-            if (closest) {
-                return sessions[closest.index];
+            if (closest.indexes.length) {
+                return closest.indexes.map(i => sessions[i]);
             }
 
             return undefined;
 
         }
 
+        const getInfoSession = (items: PhaseSessionResult[], index: number) => {
+            const closest = items[index];
+            const date = getShortNiceTime(closest.start, true, true);
+            let desc = "";
+            if (items.length > 1) {
+                desc += renderToString(SessionListSelector(items, index, '(mouse wheel to select)'));
+            }
+            desc += `#${items.length > 1? index + 1 : ''} ${closest.outcome} <br/>`;
+            desc += `on Commit ${closest.commitId} <br/>`;
+            desc += `${date} <br/>`;
+
+            return desc;
+        }
+
         const handleMouseMove = (event: any) => {
 
             const mouseX = _d3.pointer(event)[0];
             const mouseY = _d3.pointer(event)[1];
 
-            const closest = closestData(mouseX, mouseY);
-            if (!closest) {
+            const closestItems = closestData(mouseX, mouseY);
+            if (!closestItems) {
                 handleMouseLeave(undefined);
                 return;
             }
 
-            if (this.tooltipSessionId !== closest.id) {
-                this.tooltipSessionId = closest.id;
+            const closest = closestItems[0];
+            if (this.tooltipItem?.sessionId !== closest.id) {
+                handleMouseLeave(undefined);
+
                 const metaIndex = 0;
+                let index = closestItems.length > 1? closestItems.findIndex(item => item.id === this.phaseSessionId) : 0;
+                if (index < 0) index = 0;
+                this.tooltipItem = { items: closestItems, index: index, sessionId: closest.id, delta: 0 };
 
-                svg.selectAll(".circleScale").attr("transform", `scale(1)`);
                 svg.select(`#circleScale${closest.id}`).attr("transform", `scale(2)`);
                 svg.select(`#circle${closest.id}`).raise();
-                const date = getShortNiceTime(closest.start, true, true);
-                let desc = "";
-                desc += `${closest.outcome} <br/>`;
-                desc += `on Commit ${closest.commitId} <br/>`;
-                desc += `${date} <br/>`;
 
                 const timeStamp = handler.commitIdDates.get(closest.commitId)!.getTime() / 1000;
                 const tx = xScale(timeStamp);
-                const ty = yScale(metaIndex as any)!
+                const ty = yScale(metaIndex)!;
+                const desc = getInfoSession(closestItems, index);
 
-                this.updateTooltip(true, tx, ty, desc);
+                this.updateTooltip(tx, ty, desc);
             }
+
         }
 
         const handleMouseLeave = (event: any) => {
-            svg.selectAll(".circleScale").attr("transform", `scale(1)`);
-            tooltip.style("display", "none");
-            this.tooltipSessionId = undefined;
+
+            if (this.tooltipItem) {
+                svg.selectAll(".circleScale").attr("transform", `scale(1)`);
+                tooltip.style("display", "none");
+                let index = this.tooltipItem.items.length > 1? this.tooltipItem.items.findIndex(item => item.id === this.phaseSessionId) : 0;
+                if (index < 0) index = 0;
+                reDrawCircles(this.tooltipItem.items, index);
+                this.tooltipItem = undefined;
+            }
+
         }
 
         const handleMouseClick = (event: any) => {
@@ -239,33 +268,60 @@ export class PhaseHistoryGraph {
             const mouseX = _d3.pointer(event)[0];
             const mouseY = _d3.pointer(event)[1];
 
-            const closest = closestData(mouseX, mouseY);
-            if (closest && this.onClickCallback) {
-                this.onClickCallback(closest);
+            if(this.tooltipItem) {
+                const closestItems = closestData(mouseX, mouseY);
+                if (closestItems && this.onClickCallback) {
+                    this.onClickCallback(closestItems[this.tooltipItem.index]);
+                }
             }
+
         }
 
-        svg.on("mousemove", (event) => handleMouseMove(event))
-        svg.on("mouseleave", (event) => handleMouseLeave(event))
-        svg.on("click", (event) => handleMouseClick(event))
+        const handleMouseWheel = (event: any) => {
+
+            if (this.tooltipItem && this.tooltipItem.items.length > 1) {
+                this.tooltipItem.delta += event.wheelDelta;
+                if (Math.abs(this.tooltipItem.delta) >= 240) {
+                    this.tooltipItem.index += this.tooltipItem.delta > 0? 1 : -1;
+                    this.tooltipItem.delta = 0;
+
+                    if (this.tooltipItem.index < 0) {
+                        this.tooltipItem.index += this.tooltipItem.items.length;
+                    } else if (this.tooltipItem.index >= this.tooltipItem.items.length) {
+                        this.tooltipItem.index -= this.tooltipItem.items.length;
+                    }
+
+                    reDrawCircles(this.tooltipItem.items, this.tooltipItem.index);
+                    const desc = getInfoSession(this.tooltipItem.items, this.tooltipItem.index);
+                    this.updateTooltip(undefined, undefined, desc);
+                }
+            }
+            
+        }
+
+        svg.on("mousemove", (event) => handleMouseMove(event));
+        svg.on("mouseleave", (event) => handleMouseLeave(event));
+        svg.on("click", (event) => handleMouseClick(event));
+        svg.on("wheel", (event) => handleMouseWheel(event));
 
     }
 
-    updateTooltip(show: boolean, x?: number, y?: number, html?: string) {
+    updateTooltip(x?: number, y?: number, html?: string) {
         if (!this.tooltip) {
             return;
         }
 
-        x = x ?? 0;
-        y = y ?? 0;
+        if (x !== undefined && y != undefined) {
+            this.tooltip
+                .style("top", `${y}px`)
+                .style("left", `${x}px`);
+        }
 
         this.tooltip
             .style("display", "block")
             .html(html ?? "")
             .style("position", `absolute`)
             .style("width", `max-content`)
-            .style("top", `${y}px`)
-            .style("left", `${x}px`)
             .style("transform", "translate(10%, 52%)")
             .style("font-family", "Horde Open Sans Semibold")
             .style("font-size", "10px")
@@ -297,10 +353,11 @@ export class PhaseHistoryGraph {
     sessions: PhaseSessionResult[] = [];
 
     tooltip?: DivSelectionType;
-    tooltipSessionId?: string;
+    tooltipItem?: {index: number, items: PhaseSessionResult[], sessionId: string, delta: number};
 }
 
-const SessionCircle = (container: SVGGElement, session: PhaseSessionResult, radius: number, selected?: boolean) => {
+const SessionCircle = (container: SVGGElement, sessions: PhaseSessionResult[], sessionId: string | undefined, radius: number, selected?: boolean) => {
+    const session = sessions.length === 1 || !sessionId? sessions[0] : sessions.find(s => s.id === sessionId) ?? sessions[0];
     const color = getPhaseSessionStatusColor(session);
     const g = d3.select(container);
     g.append("circle")
@@ -315,5 +372,17 @@ const SessionCircle = (container: SVGGElement, session: PhaseSessionResult, radi
             .attr("r", radius + 3);
     }
 
+    if (sessions.length > 1) {
+        g.append("text")
+            .attr("text-anchor", "start")
+            .style("alignment-baseline", "left")
+            .style("font-family", "Horde Open Sans Regular")
+            .style("font-size", 10)
+            .attr("y", -4)
+            .attr("x", 8)
+            .attr("fill", dashboard.darktheme ? "#E0E0E0" : "#2D3F5F")
+            .text(`${sessions.length}`);
+    }
+
     return g;
 }
diff --git a/Engine/Source/Programs/Horde/HordeDashboard/plugins/test-automation/testAutomationCommon.tsx b/Engine/Source/Programs/Horde/HordeDashboard/plugins/test-automation/testAutomationCommon.tsx
--- a/Engine/Source/Programs/Horde/HordeDashboard/plugins/test-automation/testAutomationCommon.tsx
+++ b/Engine/Source/Programs/Horde/HordeDashboard/plugins/test-automation/testAutomationCommon.tsx
@@ -394,6 +394,18 @@ export const getTestSessionStatusColor = (test: TestSessionResult | TestSessionD
     return color;
 }
 
+const isTestSessionResult = (item: TestSessionResult | PhaseSessionResult): item is TestSessionResult => {
+    return !!(item as TestSessionResult).nameRef;
+}
+
+export const getSessionColor = (item: TestSessionResult | PhaseSessionResult) => {
+    if(isTestSessionResult(item)) {
+        return getTestSessionStatusColor(item);
+    } else {
+        return getPhaseSessionStatusColor(item);
+    }
+}
+
 export const PhaseHistoryGraphWidget: React.FC<{test: TestNameRef, phaseKey: string, meta: MetadataRef, streamId: string, sessionId?: string, handler: TestDataHandler, onClick: PhaseSessionCallback, withoutHeader?: boolean}> = ({ test, phaseKey, meta, streamId, sessionId, handler, onClick, withoutHeader }) => {
     const [container, setContainer] = useState<HTMLDivElement | null>(null);
     const [state, setState] = useState<{ graph?: PhaseHistoryGraph }>({});
@@ -1003,3 +1015,20 @@ export const UserPickerModal: React.FC<{ defaultUser?: GetUserResponse, onSelect
             </Stack>
         </Modal>;
 };
+
+export const SessionListSelector = (items: TestSessionResult[] | PhaseSessionResult[], index: number, text?: string): JSX.Element => {
+    const selectedColor = dashboard.darktheme ? "white" : "black";
+    const backColor = dashboard.darktheme ? "#2D2B29" : "#EDEBE9";
+    const textColor = dashboard.darktheme ? "#979797ff" : "#707070ff";
+
+    return <Stack>
+                <Stack style={{height: 3, backgroundColor: backColor}} grow horizontal tokens={{childrenGap: items.length < 20? 1: 0}}>
+                    { items.map( (item, i) => {
+                            const itemColor = index === i? selectedColor : getSessionColor(item);
+                            return <Stack key={`s-${i}`} grow verticalFill style={{backgroundColor: itemColor}}></Stack>
+                        } )
+                    }
+                </Stack>
+                {!!text && <Text style={{color: textColor}} variant="xSmall">&ensp;{text}</Text>}
+            </Stack>
+}
diff --git a/Engine/Source/Programs/Horde/HordeDashboard/plugins/test-automation/testHistoryGraph.ts b/Engine/Source/Programs/Horde/HordeDashboard/plugins/test-automation/testHistoryGraph.ts
--- a/Engine/Source/Programs/Horde/HordeDashboard/plugins/test-automation/testHistoryGraph.ts
+++ b/Engine/Source/Programs/Horde/HordeDashboard/plugins/test-automation/testHistoryGraph.ts
@@ -5,7 +5,7 @@ import { TestDataHandler, TestNameRef, MetadataRef, TestMetaStatus, TestSessionR
 import dashboard, { StatusColor } from 'horde/backend/Dashboard';
 import { getHumanTime, getShortNiceTime, msecToElapsed } from 'horde/base/utilities/timeUtils';
 import { getHordeStyling } from 'horde/styles/Styles';
-import { getStatusColors, statusTexts, SessionValues } from "./testAutomationCommon";
+import { getStatusColors, statusTexts, SessionValues, SessionListSelector } from "./testAutomationCommon";
 import { renderToString } from "react-dom/server"
 import { TestOutcome } from "./api";
 
@@ -176,7 +176,6 @@ export class TestHistoryGraph {
             .attr("x1", ([, I]: any) => this.margin.left)
             .attr("x2", ([, I]: any) => width - this.margin.right);
 
-
         g.append("text")
             .attr("text-anchor", "start")
             .style("alignment-baseline", "left")
@@ -194,16 +193,23 @@ export class TestHistoryGraph {
             .text(([y]: any) => this.metaNames.get(this.metaRefs[y])!);
 
         g.selectAll("circle")
-            .data(([, I]: any) => I)
+            .data(([, I]: any) => _d3.group(I, i => X[i]))
             .enter()
             .append("g")
-            .attr("id", i => `circle${sessions[i as any].id}`)
-            .attr("transform", i => `translate(${xScale(X[i as any])},0)`)
+            .attr("id", ([, items]: any) => `circle${sessions[items[0]].id}`)
+            .attr("transform", ([, items]: any) => `translate(${xScale(X[items[0] as any])},0)`)
             .append("g")
             .attr("class", "circleScale")
-            .attr("id", i => `circleScale${sessions[i as any].id}`)
+            .attr("id", ([, items]: any) => `circleScale${sessions[items[0]].id}`)
             .style("cursor", "pointer")
-            .each((d, i, nodes) => SessionConcentricCircle(nodes[i], sessions[d as any], 5));
+            .each(([, items]: any, i, nodes) => SessionConcentricCircle(nodes[i], items.map(i => sessions[i]), 0, 5));
+
+
+        const reDrawCircles = (items: TestSessionResult[], index: number) => {
+            const circle = svg.select(`#circleScale${items[0].id}`);
+            circle.selectAll("*").remove();
+            SessionConcentricCircle(circle.node() as any, items, index, 5);
+        }
 
         const tooltip = this.tooltip = d3.select(container)
             .append("div")
@@ -219,9 +225,9 @@ export class TestHistoryGraph {
             .style("pointer-events", "none");
 
         // populate metaInfo y value
-        this.metaInfo.forEach(meta => meta.y = yScale(meta.index as any)!)
+        this.metaInfo.forEach(meta => meta.y = yScale(meta.index as any)!);
 
-        const closestData = (x: number, y: number): TestSessionResult | undefined => {
+        const closestData = (x: number, y: number): TestSessionResult[] | undefined => {
 
             if (x < this.margin.left - 16) {
                 return undefined;
@@ -237,63 +243,82 @@ export class TestHistoryGraph {
                 const lengthSqr = absy * absy + absx * absx;
 
                 if (lengthSqr < best.value) {
-                    return { index: i, value: lengthSqr };
-                }
-                else {
+                    return { indexes: [i], value: lengthSqr, meta: session.metadataId };
+                } else {
+                    if (lengthSqr === best.value && session.metadataId === best.meta) best.indexes.push(i);
                     return best;
                 }
-            }, { index: 0, value: Number.MAX_SAFE_INTEGER });
+            }, { indexes: [] as number[], value: Number.MAX_SAFE_INTEGER, meta: '' });
 
-            if (closest) {
-                return sessions[closest.index];
+            if (closest.indexes.length) {
+                return closest.indexes.map(i => sessions[i]);
             }
 
             return undefined;
 
         }
 
+        const getInfoSession = (items: TestSessionResult[], index: number) => {
+            const item = items[index];
+            const metaIndex = this.metaInfo.get(item.metadataId)!.index;
+            const cmeta = this.metaRefs[metaIndex].getValuesExcept(this.commonMetaKeys!).join(" / ");
+            const date = getShortNiceTime(item.start, true, true);
+            let desc = "";
+            if (items.length > 1) {
+                desc += renderToString(SessionListSelector(items, index, '(mouse wheel to select)'));
+            }
+            if (!!cmeta || items.length > 1) desc += `#${items.length > 1? index + 1 : ''} ${cmeta} <br/>`;
+            desc += `Commit ${item.commitId} <br/>`;
+            desc += `Duration ${msecToElapsed(item.duration * 1000, true, true)} <br/>`;
+            desc += `${date} <br/>`;
+
+            if (item.phasesSucceededCount + item.phasesFailedCount + item.phasesUnspecifiedCount > 1)
+                desc += renderToString(SessionValues(item, {fontSize: 9}));
+
+            return desc;
+        }
+
         const handleMouseMove = (event: any) => {
 
             const mouseX = _d3.pointer(event)[0];
             const mouseY = _d3.pointer(event)[1];
 
-            const closest = closestData(mouseX, mouseY);
-            if (!closest) {
+            const closestItems = closestData(mouseX, mouseY);
+            if (!closestItems) {
                 handleMouseLeave(undefined);
                 return;
             }
 
-            if (this.tooltipSessionId !== closest.id) {
-                this.tooltipSessionId = closest.id;
+            const closest = closestItems[0];
+            if (this.tooltipItem?.sessionId !== closest.id) {
+                handleMouseLeave(undefined);
+
+                this.tooltipItem = {sessionId: closest.id, index: 0, items: closestItems, delta: 0};
                 const metaIndex = this.metaInfo.get(closest.metadataId)!.index;
 
-                svg.selectAll(".circleScale").attr("transform", `scale(1)`);
                 svg.select(`#circleScale${closest.id}`).attr("transform", `scale(2.5)`);
                 svg.select(`#circle${closest.id}`).raise();
                 svg.select(`#meta${metaIndex}`).raise();
 
-                const cmeta = this.metaRefs[metaIndex].getValuesExcept(this.commonMetaKeys!).join(" / ");
-                const date = getShortNiceTime(closest.start, true, true);
-                let desc = "";
-                if (!!cmeta) desc += `# ${cmeta} <br/>`;
-                desc += `Commit ${closest.commitId} <br/>`;
-                desc += `Duration ${msecToElapsed(closest.duration * 1000, true, true)} <br/>`;
-                desc += `${date} <br/>`;
-                if (closest.phasesSucceededCount + closest.phasesFailedCount + closest.phasesUnspecifiedCount > 1)
-                    desc += renderToString(SessionValues(closest, {fontSize: 9}));
-
                 const timeStamp = handler.commitIdDates.get(closest.commitId)!.getTime() / 1000;
                 const tx = xScale(timeStamp);
                 const ty = yScale(metaIndex as any)!
+                const desc = getInfoSession(closestItems, 0);
 
-                this.updateTooltip(true, tx, ty, desc);
+                this.updateTooltip(tx, ty, desc);
             }
+
         }
 
         const handleMouseLeave = (event: any) => {
-            svg.selectAll(".circleScale").attr("transform", `scale(1)`);
-            tooltip.style("display", "none");
-            this.tooltipSessionId = undefined;
+
+            if (this.tooltipItem) {
+                svg.selectAll(".circleScale").attr("transform", `scale(1)`);
+                tooltip.style("display", "none");
+                reDrawCircles(this.tooltipItem.items, 0);
+                this.tooltipItem = undefined;
+            }
+
         }
 
         const handleMouseClick = (event: any) => {
@@ -301,39 +326,66 @@ export class TestHistoryGraph {
             const mouseX = _d3.pointer(event)[0];
             const mouseY = _d3.pointer(event)[1];
 
-            const closest = closestData(mouseX, mouseY);
-            if (closest && this.onClickCallback) {
-                this.onClickCallback(closest);
+            if (this.tooltipItem) {
+                const closestItems = closestData(mouseX, mouseY);
+                if (closestItems && this.onClickCallback) {
+                    this.onClickCallback(closestItems[this.tooltipItem.index]);
+                }
             }
+
         }
 
-        svg.on("mousemove", (event) => handleMouseMove(event))
-        svg.on("mouseleave", (event) => handleMouseLeave(event))
-        svg.on("click", (event) => handleMouseClick(event))
+        const handleMouseWheel = (event: any) => {
+
+            if (this.tooltipItem && this.tooltipItem.items.length > 1) {
+                this.tooltipItem.delta += event.wheelDelta;
+                if (Math.abs(this.tooltipItem.delta) >= 240) {
+                    this.tooltipItem.index += this.tooltipItem.delta > 0? 1 : -1;
+                    this.tooltipItem.delta = 0;
+
+                    if (this.tooltipItem.index < 0) {
+                        this.tooltipItem.index += this.tooltipItem.items.length;
+                    } else if (this.tooltipItem.index >= this.tooltipItem.items.length) {
+                        this.tooltipItem.index -= this.tooltipItem.items.length;
+                    }
+
+                    reDrawCircles(this.tooltipItem.items, this.tooltipItem.index);
+                    const desc = getInfoSession(this.tooltipItem.items, this.tooltipItem.index);
+                    this.updateTooltip(undefined, undefined, desc);
+                }
+            }
+            
+        }
+
+        svg.on("mousemove", (event) => handleMouseMove(event));
+        svg.on("mouseleave", (event) => handleMouseLeave(event));
+        svg.on("click", (event) => handleMouseClick(event));
+        svg.on("wheel", (event) => handleMouseWheel(event));
 
     }
 
-    updateTooltip(show: boolean, x?: number, y?: number, html?: string) {
+    updateTooltip(x?: number, y?: number, html?: string) {
         if (!this.tooltip) {
             return;
         }
 
-        x = x ?? 0;
-        y = y ?? 0;
+        if (x !== undefined && y != undefined) {
+            this.tooltip
+                .style("top", `${y}px`)
+                .style("left", `${x}px`);
+        }
 
         this.tooltip
             .style("display", "block")
             .html(html ?? "")
             .style("position", `absolute`)
             .style("width", `max-content`)
-            .style("top", `${y}px`)
-            .style("left", `${x}px`)
             .style("transform", "translate(-108%, -97%)")
             .style("font-family", "Horde Open Sans Semibold")
             .style("font-size", "10px")
             .style("line-height", "16px")
             .style("shapeRendering", "crispEdges")
-            .style("stroke", "none")
+            .style("stroke", "none");
 
     }
 
@@ -364,7 +416,7 @@ export class TestHistoryGraph {
     clipId: string;
 
     tooltip?: DivSelectionType;
-    tooltipSessionId?: string;
+    tooltipItem?: {index: number, items: TestSessionResult[], sessionId: string, delta: number};
 }
 
 type ItemStack = {value: number, name: string, color: string, title?: boolean};
@@ -436,11 +488,12 @@ const ConcentricCircle = (container: SVGGElement, data: ItemStack[], radius: num
     return g;
 }
 
-const SessionConcentricCircle = (container: SVGGElement, session: TestSessionResult, radius: number) => {
+const SessionConcentricCircle = (container: SVGGElement, sessions: TestSessionResult[], index: number, radius: number) => {
+    const latestSession = sessions[index];
 
-    const sessionTotalCount = session.phasesSucceededCount + session.phasesFailedCount + session.phasesUnspecifiedCount;
-    const metaFailedFactor = Math.ceil(clampMin(session.phasesFailedCount / (sessionTotalCount || 1)) * 5) / 5;
-    const metaUnspecifiedFactor = Math.ceil(clampMin(session.phasesUnspecifiedCount / (sessionTotalCount || 1)) * 5) / 5;
+    const sessionTotalCount = latestSession.phasesSucceededCount + latestSession.phasesFailedCount + latestSession.phasesUnspecifiedCount;
+    const metaFailedFactor = Math.ceil(clampMin(latestSession.phasesFailedCount / (sessionTotalCount || 1)) * 5) / 5;
+    const metaUnspecifiedFactor = Math.ceil(clampMin(latestSession.phasesUnspecifiedCount / (sessionTotalCount || 1)) * 5) / 5;
     let metaSuccessFactor = 1 - (metaFailedFactor + metaUnspecifiedFactor);
     metaSuccessFactor = metaSuccessFactor < 0? 0 : metaSuccessFactor
 
@@ -464,5 +517,20 @@ const SessionConcentricCircle = (container: SVGGElement, session: TestSessionRes
         },
     ];
 
-    return ConcentricCircle(container, stack, radius);
+    const g = ConcentricCircle(container, stack, radius);
+
+    if (sessions.length > 1) {
+        d3.select(container)
+            .append("text")
+                .attr("text-anchor", "start")
+                .style("alignment-baseline", "left")
+                .style("font-family", "Horde Open Sans Regular")
+                .style("font-size", 10)
+                .attr("y", -2)
+                .attr("x", 6)
+                .attr("fill", dashboard.darktheme ? "#E0E0E0" : "#2D3F5F")
+                .text(`${sessions.length}`);
+    }
+
+    return g;
 }
