Fix missing elipsis on Cadence column

#rnx

#Synced-CL 48452052

diff --git a/Engine/Source/Programs/Horde/HordeDashboard/plugins/test-automation/testAutomationHealth.tsx b/Engine/Source/Programs/Horde/HordeDashboard/plugins/test-automation/testAutomationHealth.tsx
--- a/Engine/Source/Programs/Horde/HordeDashboard/plugins/test-automation/testAutomationHealth.tsx
+++ b/Engine/Source/Programs/Horde/HordeDashboard/plugins/test-automation/testAutomationHealth.tsx
@@ -213,8 +213,8 @@ const getHealthColumns = (items: TestHealthItem[], columns: TestHealthColumn[],
             isSorted: false,
             isSortedDescending: true,
             onRender: (item: TestHealthItem) => {
-                return <Stack style={{backgroundColor: cadenceColors.get(item.cadence), padding: 3, borderRadius: 4, cursor: 'default'}} horizontalAlign="center">
-                            <Text>{item.cadence}</Text>
+                return <Stack style={{backgroundColor: cadenceColors.get(item.cadence), padding: 3, borderRadius: 4, cursor: 'default'}} horizontalAlign="center" horizontal>
+                            <Text title={item.cadence} style={{textOverflow: 'ellipsis', overflow: 'hidden', whiteSpace: 'nowrap'}}>{item.cadence}</Text>
                         </Stack>
             },
             onColumnClick: (_, column) => onClickColumn(column),
