[Horde][BuildHealth][Server] Only require RCA category, don't escalate for missing summary

diff --git a/Engine/Source/Programs/Horde/Plugins/Build/HordeServer.Build/Notifications/Sinks/SlackNotificationSink.cs b/Engine/Source/Programs/Horde/Plugins/Build/HordeServer.Build/Notifications/Sinks/SlackNotificationSink.cs
--- a/Engine/Source/Programs/Horde/Plugins/Build/HordeServer.Build/Notifications/Sinks/SlackNotificationSink.cs
+++ b/Engine/Source/Programs/Horde/Plugins/Build/HordeServer.Build/Notifications/Sinks/SlackNotificationSink.cs
@@ -3107,9 +3107,9 @@ namespace HordeServer.Notifications.Sinks
 				return null;
 			}
 
-			// RootCauseCategory and RootCauseSummary are optional inputs in the UI because the issue closer may not have all the
-			// information when marking as fixed; they are all required to have a value for Root Cause Analysis to be considered closed.
-			if (!String.IsNullOrEmpty(issue.RootCauseCategory) && !String.IsNullOrEmpty(issue.RootCauseSummary))
+			// RootCauseCategory is an optional input in the UI because the issue closer may not have the information 
+			// when marking as fixed, but it is required to have a value for Root Cause Analysis to be considered closed.
+			if (!String.IsNullOrEmpty(issue.RootCauseCategory))
 			{
 				_logger.LogInformation(CanceledMessage, notificationName, issueId, "RCA Completed");
 				return null;
