Horde: Timeout P4 issue rank queries after 2 minutes

diff --git a/Engine/Source/Programs/Horde/Plugins/Build/HordeServer.Build/Issues/IssueService.cs b/Engine/Source/Programs/Horde/Plugins/Build/HordeServer.Build/Issues/IssueService.cs
--- a/Engine/Source/Programs/Horde/Plugins/Build/HordeServer.Build/Issues/IssueService.cs
+++ b/Engine/Source/Programs/Horde/Plugins/Build/HordeServer.Build/Issues/IssueService.cs
@@ -1539,7 +1539,8 @@ namespace HordeServer.Issues
 			
 			int rank = -1;
 
-			IReadOnlyList<string> commitFiles = await commit.GetFilesAsync(MaxFilesForSuspects, CancellationToken.None);
+			using CancellationTokenSource cts = new CancellationTokenSource(TimeSpan.FromMinutes(2));
+			IReadOnlyList<string> commitFiles = await commit.GetFilesAsync(MaxFilesForSuspects, cts.Token);
 			foreach (string commitFile in commitFiles)
 			{
 				if (filter.Matches(commitFile))
