Clean up code - remove color reference for horizontal lines in History graph

#rnx

#Synced-CL 48452052

diff --git a/Engine/Source/Programs/Horde/HordeDashboard/plugins/test-automation/phaseHistoryGraph.ts b/Engine/Source/Programs/Horde/HordeDashboard/plugins/test-automation/phaseHistoryGraph.ts
--- a/Engine/Source/Programs/Horde/HordeDashboard/plugins/test-automation/phaseHistoryGraph.ts
+++ b/Engine/Source/Programs/Horde/HordeDashboard/plugins/test-automation/phaseHistoryGraph.ts
@@ -48,14 +48,6 @@ export class PhaseHistoryGraph {
         const sessions = this.sessions;
         const width = 1000
 
-        const scolors = getStatusColors();
-        const colors = new Map<TestOutcome, string>([
-            [TestOutcome.Success, scolors.get(StatusColor.Success)!],
-            [TestOutcome.Failure, scolors.get(StatusColor.Failure)!],
-            [TestOutcome.Unspecified, scolors.get(StatusColor.Unspecified)!],
-            [TestOutcome.Skipped, scolors.get(StatusColor.Skipped)!]
-        ]);
-
         const X = _d3.map(sessions, (r) => handler.commitIdDates.get(r.commitId)!.getTime() / 1000);
         const Y = _d3.map(sessions, (r) => 0);
 
@@ -138,8 +130,8 @@ export class PhaseHistoryGraph {
             .attr("stroke-width", 1)
             .attr("stroke-linecap", 4)
             .attr("stroke-opacity", dashboard.darktheme ? 0.35 : 0.25)
-            .attr("x1", ([, I]: any) => this.margin.left)
-            .attr("x2", ([, I]: any) => width - this.margin.right);
+            .attr("x1", this.margin.left)
+            .attr("x2", width - this.margin.right);
 
         g.selectAll("circle")
             .data(([, I]: any) => _d3.group(I, i => X[i]))
diff --git a/Engine/Source/Programs/Horde/HordeDashboard/plugins/test-automation/testHistoryGraph.ts b/Engine/Source/Programs/Horde/HordeDashboard/plugins/test-automation/testHistoryGraph.ts
--- a/Engine/Source/Programs/Horde/HordeDashboard/plugins/test-automation/testHistoryGraph.ts
+++ b/Engine/Source/Programs/Horde/HordeDashboard/plugins/test-automation/testHistoryGraph.ts
@@ -48,7 +48,6 @@ export class TestHistoryGraph {
                 // skip if it filters by tags and no tag match
                 return;
             }
-            this.metaStatus.set(meta, !metaLastSession? TestOutcome.Unspecified : metaLastSession.outcome);
 
             const elements: string[] = meta.getValuesExcept(commonMetaKeys);
             this.metaNames.set(meta, elements.join(" / "));
@@ -76,14 +75,6 @@ export class TestHistoryGraph {
         const sessions = this.sessions;
         const width = 1000
 
-        const scolors = getStatusColors();
-        const colors = new Map<TestOutcome, string>([
-            [TestOutcome.Success, scolors.get(StatusColor.Success)!],
-            [TestOutcome.Failure, scolors.get(StatusColor.Failure)!],
-            [TestOutcome.Unspecified, scolors.get(StatusColor.Unspecified)!],
-            [TestOutcome.Skipped, scolors.get(StatusColor.Skipped)!]
-        ]);
-
         const X = _d3.map(sessions, (r) => handler.commitIdDates.get(r.commitId)!.getTime() / 1000);
         const Y = _d3.map(sessions, (r) => this.metaInfo.get(r.metadataId)!.index);
 
@@ -163,18 +154,12 @@ export class TestHistoryGraph {
             .attr("transform", ([y]: any) => `translate(0,${(yScale(y) as any) + 16})`);
 
         g.append("line")
-            .attr("stroke", (y: any) => {
-                const status = this.metaStatus.get(this.metaRefs[y])!;
-                if (status === TestOutcome.Failure) {
-                    return colors[this.metaStatus.get(this.metaRefs[y])!]
-                }
-                return dashboard.darktheme ? "#6D6C6B" : "#4D4C4B";
-            })
+            .attr("stroke", dashboard.darktheme ? "#6D6C6B" : "#4D4C4B")
             .attr("stroke-width", 1)
             .attr("stroke-linecap", 4)
             .attr("stroke-opacity", dashboard.darktheme ? 0.35 : 0.25)
-            .attr("x1", ([, I]: any) => this.margin.left)
-            .attr("x2", ([, I]: any) => width - this.margin.right);
+            .attr("x1", this.margin.left)
+            .attr("x2", width - this.margin.right);
 
         g.append("text")
             .attr("text-anchor", "start")
@@ -183,13 +168,7 @@ export class TestHistoryGraph {
             .style("font-size", 10)
             .attr("dy", "0.15em") // center stream name
             .attr("x", () => 0)
-            .attr("fill", ([y]: any) => {
-                const status = this.metaStatus.get(this.metaRefs[y])!;
-                if (!dashboard.darktheme && status === TestOutcome.Failure) {
-                    return colors[this.metaStatus.get(this.metaRefs[y])!]
-                }
-                return dashboard.darktheme ? "#E0E0E0" : "#2D3F5F";
-            })
+            .attr("fill", dashboard.darktheme ? "#E0E0E0" : "#2D3F5F")
             .text(([y]: any) => this.metaNames.get(this.metaRefs[y])!);
 
         g.selectAll("circle")
@@ -408,8 +387,6 @@ export class TestHistoryGraph {
     metaRefs: MetadataRef[] = [];
     metaInfo: Map<string, {y: number, index: number}> = new Map();
     sessions: TestSessionResult[] = [];
-    // metaRef => status string
-    metaStatus: Map<MetadataRef, string> = new Map();
     // metaRef => meta string
     metaNames: Map<MetadataRef, string> = new Map();
 
