Horde: Support JobStepId tags in log renders

diff --git a/Engine/Source/Programs/Horde/HordeDashboard/src/components/logs/LogRender.tsx b/Engine/Source/Programs/Horde/HordeDashboard/src/components/logs/LogRender.tsx
--- a/Engine/Source/Programs/Horde/HordeDashboard/src/components/logs/LogRender.tsx
+++ b/Engine/Source/Programs/Horde/HordeDashboard/src/components/logs/LogRender.tsx
@@ -12,6 +12,7 @@ enum TagType {
    MSDNCode,
    Link,
    JobId,
+   JobStepId,
    LogId,
    ArtifactId,
    AgentId,
@@ -77,7 +78,7 @@ const renderTags = (navigate: NavigateFunction, line: LogLine, lineNumber: numbe
    if (!line || !line.format || !line.properties) {
       return <Stack styles={{ root: { color: "#000000", paddingLeft: 8, whiteSpace: "pre", tabSize: "3" } }}>Internal log line format error</Stack>;
    }
-
+   
    const properties = line.properties;
    const lineProp = properties["line"];
 
@@ -201,6 +202,17 @@ const renderTags = (navigate: NavigateFunction, line: LogLine, lineNumber: numbe
             }
             break;
 
+         case TagType.JobStepId:            
+            const jobId = (line as any)?.properties?.JobId?.$text
+            if (jobId) {
+               const stepUrl = `/job/${jobId}?step=${text}`;
+               linkProps = {
+                  href: stepUrl,
+                  onClick: (ev) => { ev.stopPropagation(); ev.preventDefault(); navigate(stepUrl) }
+               }
+            }
+            break;
+
          case TagType.MSDNCode:
             linkProps = {
                href: `https://msdn.microsoft.com/query/dev16.query?appId=Dev16IDEF1&l=EN-US&k=k(${text.toLowerCase()})&rd=true`,
