[Horde][BuildHealth][Dashboard] Indicate RCA category is required on dashboard
#rnx

#Synced-CL Unknown

diff --git a/Engine/Source/Programs/Horde/HordeDashboard/src/components/IssueViewV2.tsx b/Engine/Source/Programs/Horde/HordeDashboard/src/components/IssueViewV2.tsx
--- a/Engine/Source/Programs/Horde/HordeDashboard/src/components/IssueViewV2.tsx
+++ b/Engine/Source/Programs/Horde/HordeDashboard/src/components/IssueViewV2.tsx
@@ -2500,7 +2500,7 @@ const MarkFixedModal: React.FC<{ onClose: () => void, autoFix?:boolean } > = ({
                      disabled={disableRoot}
                      selectedKey={selectedCatKey}
                      placeholder="Select an option"
-                     label="Root Cause Category"
+                     label="Root Cause Category (required)"
                      options={rcaCategoryOptions}
                      onChange={(event, option) => {
                         if (option?.key === "rca_clear") {
