[Horde][BuildHealth][Dashboard] Implement feedback from older swarm: code cleanliness

#Synced-CL 48440999

diff --git a/Engine/Source/Programs/Horde/HordeDashboard/src/components/buildhealth/BuildHealthOptionsController.ts b/Engine/Source/Programs/Horde/HordeDashboard/src/components/buildhealth/BuildHealthOptionsController.ts
--- a/Engine/Source/Programs/Horde/HordeDashboard/src/components/buildhealth/BuildHealthOptionsController.ts
+++ b/Engine/Source/Programs/Horde/HordeDashboard/src/components/buildhealth/BuildHealthOptionsController.ts
@@ -447,8 +447,7 @@ export class BuildHealthOptionsWriter implements IBuildHealthOptionWriter {
         if (this.uiState.linkedSelectionModeEnabled) {
             Object.keys(this.state.enabledSteps).forEach(stepKey => {
                 let stepName = this.state.enabledSteps[stepKey];
-                let result = this.toggleLinkedSteps(stepKey, stepName, enabled, queryOnly);
-                requiresLinking ||= result;
+                requiresLinking ||= this.toggleLinkedSteps(stepKey, stepName, enabled, queryOnly);
             });
         }
 
