[Horde][BuildHealth][Server] Modify RCA escalation message to inform user that only category is required
#Synced-CL Unknown

diff --git a/Engine/Source/Programs/Horde/Plugins/Build/HordeServer.Build/Notifications/Sinks/SlackNotificationSink.cs b/Engine/Source/Programs/Horde/Plugins/Build/HordeServer.Build/Notifications/Sinks/SlackNotificationSink.cs
--- a/Engine/Source/Programs/Horde/Plugins/Build/HordeServer.Build/Notifications/Sinks/SlackNotificationSink.cs
+++ b/Engine/Source/Programs/Horde/Plugins/Build/HordeServer.Build/Notifications/Sinks/SlackNotificationSink.cs
@@ -3158,7 +3158,7 @@ namespace HordeServer.Notifications.Sinks
 			_logger.LogInformation("Sending escalation message for {NotificationName} {IssueId} {User} {Channel}", notificationName, issueId, mention, state.MessageId);
 
 			Uri issueUrl = GetIssueUrl(issue, span.FirstFailure);
-			string message = $"{mention} this root cause analysis requires your attention - Issue <{issueUrl}|{issue.Id}> - Select ...More > Modify Issue to complete the root cause analysis.";
+			string message = $"{mention} this root cause analysis requires your attention - Issue <{issueUrl}|{issue.Id}> - Select ...More > Modify Issue and choose a category to complete the root cause analysis.";
 			await _slackClient.PostMessageToThreadAsync(state.MessageId, message, cancellationToken);
 
 			return GetNextEscalationTime(issue, issue.ResolvedAt.Value, utcNow, workflow.EscalateRootCauseAnalysisTimes, notificationName);
