[Horde][BuildHealth][Server] Null check for defensive programming.

#Synced-CL 48440999

diff --git a/Engine/Source/Programs/Horde/Plugins/Build/HordeServer.Build/Jobs/JobCollection.cs b/Engine/Source/Programs/Horde/Plugins/Build/HordeServer.Build/Jobs/JobCollection.cs
--- a/Engine/Source/Programs/Horde/Plugins/Build/HordeServer.Build/Jobs/JobCollection.cs
+++ b/Engine/Source/Programs/Horde/Plugins/Build/HordeServer.Build/Jobs/JobCollection.cs
@@ -2090,6 +2090,11 @@ namespace HordeServer.Jobs
 
 		private async Task UpdateRefStateAsync(IJob job, IGraph graph, IJobStepRefCollection jobStepRefCollection)
 		{
+			if(job == null)
+			{
+				return;
+			}
+
 			ICollection<JobStepRefId> jobStepRefSet;
 			ICollection<string> jobBatchSet = new HashSet<string>(job.Batches.Count);
 
