[Horde][BuildHealth][Server] Disable JobStepRef emissions.

#Synced-CL 48440999

diff --git a/Engine/Source/Programs/Horde/Plugins/Build/HordeServer.Build/Jobs/JobCollection.cs b/Engine/Source/Programs/Horde/Plugins/Build/HordeServer.Build/Jobs/JobCollection.cs
--- a/Engine/Source/Programs/Horde/Plugins/Build/HordeServer.Build/Jobs/JobCollection.cs
+++ b/Engine/Source/Programs/Horde/Plugins/Build/HordeServer.Build/Jobs/JobCollection.cs
@@ -1726,9 +1726,10 @@ namespace HordeServer.Jobs
 			if (refreshDependentJobSteps)
 			{
 				RefreshDependentJobSteps(jobDocument, graph, updates, _logger, GetJobLogger(jobDocument.Id));
-
+				/*
 				IJob jobObj = CreateJobObject(jobDocument, graph);
 				await UpdateRefStateAsync(jobObj, graph, _jobStepRefs);
+				*/
 
 				RefreshJobPriority(jobDocument, updates, _logger, GetJobLogger(jobDocument.Id));
 			}
@@ -2088,67 +2089,28 @@ namespace HordeServer.Jobs
 			return await TryUpdateAsync(jobDocument, updateBuilder.Combine(updates), cancellationToken);
 		}
 
-		private async Task UpdateRefStateAsync(IJob job, IGraph graph, IJobStepRefCollection jobStepRefCollection)
-		{
-			if(job == null)
-			{
-				return;
-			}
-
-			ICollection<JobStepRefId> jobStepRefSet;
-			ICollection<string> jobBatchSet = new HashSet<string>(job.Batches.Count);
-
-			foreach (IJobStepBatch batch in job.Batches)
-			{
-				jobBatchSet.Add($"{job.Id}:{batch.Id}");
-				jobStepRefSet = new HashSet<JobStepRefId>(batch.Steps.Count);
-
-				foreach (IJobStep step in batch.Steps)
-				{
-					await jobStepRefCollection.UpdateAsync(job, batch, step, graph);
-					JobStepRefId refId = new JobStepRefId(job.Id, batch.Id, step.Id);
-					jobStepRefSet.Add(refId);
-				}
-
-				// verify against known steps
-				List<IJobStepRef> jobBatchStepRefItems = await _jobStepRefs.FindAsync(job.Id, batch.Id);
-
-				// Some early cancel steps may have been removed; mark closed on those not contained within the set.
-				if (jobBatchStepRefItems.Count != batch.Steps.Count)
-				{
-					foreach (IJobStepRef item in jobBatchStepRefItems)
-					{
-						if (!jobStepRefSet.Contains(item.Id))
-						{
-							await _jobStepRefs.UpdateAsync(job, batch, item);
-						}
-					}
-				}
-			}
-
-			// Som batches may have been entirely culled, attempt to update at the job level.
-			List<IJobStepRef> jobStepRefItems = await _jobStepRefs.FindAsync(job.Id);
-			foreach (IJobStepRef item in jobStepRefItems)
-			{
-				if (!jobBatchSet.Contains($"{item.Id.JobId}:{item.Id.BatchId}"))
-				{
-					await _jobStepRefs.UpdateAsync(job, item);
-				}
-			}
-		}
-
+#pragma warning disable CA1822 // Mark members as static
 		private async Task UpdateBatchesInternalAsync(JobDocument job, IGraph graph, List<UpdateDefinition<JobDocument>> updates, ILogger logger, IAuditLogChannel<JobId> jobAudit)
+#pragma warning restore CA1822 // Mark members as static
 		{
 			UpdateBatches(job, graph, updates, logger, jobAudit);
+			/* UE-354742 - Reviewing JobStepRef Emission
 			IJob jobObj = CreateJobObject(job, graph);
 			await UpdateRefStateAsync(jobObj, graph, _jobStepRefs);
+			*/
+			await ValueTask.CompletedTask;
 		}
 
+#pragma warning disable CA1822 // Mark members as static
 		private async Task CreateBatchesInternalAsync(JobDocument job, IGraph graph, ILogger logger, IAuditLogChannel<JobId> jobAudit)
+#pragma warning restore CA1822 // Mark members as static
 		{
 			CreateBatches(job, graph, logger, jobAudit);
+			/* UE-354742 - Reviewing JobStepRef Emission
 			IJob jobObj = CreateJobObject(job, graph);
 			await UpdateRefStateAsync(jobObj, graph, _jobStepRefs);
+			*/
+			await ValueTask.CompletedTask;
 		}
 
 		/// <summary>
diff --git a/Engine/Source/Programs/Horde/Plugins/Build/HordeServer.Build/Jobs/JobStepRefCollection.cs b/Engine/Source/Programs/Horde/Plugins/Build/HordeServer.Build/Jobs/JobStepRefCollection.cs
--- a/Engine/Source/Programs/Horde/Plugins/Build/HordeServer.Build/Jobs/JobStepRefCollection.cs
+++ b/Engine/Source/Programs/Horde/Plugins/Build/HordeServer.Build/Jobs/JobStepRefCollection.cs
@@ -215,21 +215,27 @@ namespace HordeServer.Jobs
 		/// <inheritdoc/>
 		public async Task<List<IJobStepRef>> FindAsync(JobId jobId, JobStepBatchId batchId)
 		{
+			/* UE-354742 - Reviewing JobStepRef Emission
 			string prefix = $"^{jobId}:{batchId}:";
 			FilterDefinition<JobStepRef> filter = Builders<JobStepRef>.Filter.Regex(x => x.Id, new BsonRegularExpression(prefix));
 
 			List<JobStepRef> results = await _jobStepRefs.Find(filter).ToListAsync();
-			return results.Cast<IJobStepRef>().ToList();
+			*/
+			await ValueTask.CompletedTask;
+			return [];
 		}
 
 		/// <inheritdoc/>
 		public async Task<List<IJobStepRef>> FindAsync(JobId jobId)
 		{
+			/* UE-354742 - Reviewing JobStepRef Emission
 			string prefix = $"^{jobId}:";
 			FilterDefinition<JobStepRef> filter = Builders<JobStepRef>.Filter.Regex(x => x.Id, new BsonRegularExpression(prefix));
 
 			List<JobStepRef> results = await _jobStepRefs.Find(filter).ToListAsync();
-			return results.Cast<IJobStepRef>().ToList();
+			*/
+			await ValueTask.CompletedTask;
+			return [];
 		}
 
 		/// <inheritdoc/>
