Horde: Improve VSCode support

diff --git a/Engine/Source/Programs/Horde/.vscode/launch.json b/Engine/Source/Programs/Horde/.vscode/launch.json
new file mode 100644
--- /dev/null
+++ b/Engine/Source/Programs/Horde/.vscode/launch.json
@@ -0,0 +1,36 @@
+{
+    "version": "0.2.0",
+    "configurations": [
+        {
+            "name": "HordeServer",
+            "type": "coreclr",
+            "request": "launch",
+            "program": "${workspaceFolder}/HordeServer/bin/Debug/net8.0/HordeServer.exe",
+            "args": [],
+            "cwd": "${workspaceFolder}/HordeServer/bin/Debug/net8.0/",
+            "stopAtEntry": false,
+            "console": "externalTerminal"
+        },
+        {
+            "name": "HordeAgent",
+            "type": "coreclr",
+            "request": "launch",
+            "program": "${workspaceFolder}/HordeAgent/bin/Debug/net8.0/HordeAgent.exe",
+            "args": [],
+            "cwd": "${workspaceFolder}/HordeAgent/bin/Debug/net8.0",
+            "stopAtEntry": false,
+            "console": "externalTerminal"
+        }
+    ],
+    "compounds": [
+        {
+            "name": "Server and Agent",
+            "configurations": [
+                "HordeServer",
+                "HordeAgent"
+            ],
+            // Optional: Controls if debugging should stop if one process exits.
+            "stopAll": true
+        }
+    ]
+}
\ No newline at end of file
diff --git a/Engine/Source/Programs/Horde/Drivers/JobDriver.Tests/JobDriver.Tests.csproj b/Engine/Source/Programs/Horde/Drivers/JobDriver.Tests/JobDriver.Tests.csproj
--- a/Engine/Source/Programs/Horde/Drivers/JobDriver.Tests/JobDriver.Tests.csproj
+++ b/Engine/Source/Programs/Horde/Drivers/JobDriver.Tests/JobDriver.Tests.csproj
@@ -5,8 +5,6 @@
     <TargetFramework>net8.0</TargetFramework>
     <Nullable>enable</Nullable>
     <IsPackable>false</IsPackable>
-    <DebugType>full</DebugType>
-    <DebugType>pdbonly</DebugType>
     <Configurations>Debug;Release;Development;Analyze</Configurations>
     <NoWarn>$(NoWarn);1701;1702;1591</NoWarn>
   </PropertyGroup>
diff --git a/Engine/Source/Programs/Horde/Drivers/JobDriver/JobDriver.csproj b/Engine/Source/Programs/Horde/Drivers/JobDriver/JobDriver.csproj
--- a/Engine/Source/Programs/Horde/Drivers/JobDriver/JobDriver.csproj
+++ b/Engine/Source/Programs/Horde/Drivers/JobDriver/JobDriver.csproj
@@ -9,7 +9,6 @@
     <AssemblyName>JobDriver</AssemblyName>
     <PublishRepositoryUrl>true</PublishRepositoryUrl>
     <EmbedUntrackedSources>false</EmbedUntrackedSources>
-    <DebugType>pdbonly</DebugType>
     <Configurations>Debug;Release;Analyze</Configurations>
     <NoWarn>$(NoWarn);1701;1702;NU3003</NoWarn>
   </PropertyGroup>
diff --git a/Engine/Source/Programs/Horde/HordeAgent.Tests/HordeAgent.Tests.csproj b/Engine/Source/Programs/Horde/HordeAgent.Tests/HordeAgent.Tests.csproj
--- a/Engine/Source/Programs/Horde/HordeAgent.Tests/HordeAgent.Tests.csproj
+++ b/Engine/Source/Programs/Horde/HordeAgent.Tests/HordeAgent.Tests.csproj
@@ -5,8 +5,6 @@
     <TargetFramework>net8.0</TargetFramework>
     <Nullable>enable</Nullable>
     <IsPackable>false</IsPackable>
-    <DebugType>full</DebugType>
-    <DebugType>pdbonly</DebugType>
     <Configurations>Debug;Release;Development;Analyze</Configurations>
     <NoWarn>$(NoWarn);1701;1702;1591</NoWarn>
     <RunAnalyzersDuringBuild>False</RunAnalyzersDuringBuild>
diff --git a/Engine/Source/Programs/Horde/HordeAgent/HordeAgent.csproj b/Engine/Source/Programs/Horde/HordeAgent/HordeAgent.csproj
--- a/Engine/Source/Programs/Horde/HordeAgent/HordeAgent.csproj
+++ b/Engine/Source/Programs/Horde/HordeAgent/HordeAgent.csproj
@@ -11,7 +11,6 @@
     <AssemblyName>HordeAgent</AssemblyName>
     <PublishRepositoryUrl>true</PublishRepositoryUrl>
     <EmbedUntrackedSources>false</EmbedUntrackedSources>
-    <DebugType>pdbonly</DebugType>
     <Configurations>Debug;Release;Analyze</Configurations>
     <NoWarn>$(NoWarn);1701;1702;NU3003</NoWarn>
     <RunAnalyzersDuringBuild>False</RunAnalyzersDuringBuild>
