Horde: Improve debug name for server and agent
#submitool bypass
#rnx

diff --git a/Engine/Source/Programs/Horde/.vscode/launch.json b/Engine/Source/Programs/Horde/.vscode/launch.json
--- a/Engine/Source/Programs/Horde/.vscode/launch.json
+++ b/Engine/Source/Programs/Horde/.vscode/launch.json
@@ -24,7 +24,7 @@
     ],
     "compounds": [
         {
-            "name": "Server and Agent",
+            "name": "Debug Server and Agent",
             "configurations": [
                 "HordeServer",
                 "HordeAgent"
