EpicGames.Perforce: Improve fixture with deletes from shelve

diff --git a/Engine/Source/Programs/Horde/Drivers/JobDriver.Tests/PerforceMaterializerTests.cs b/Engine/Source/Programs/Horde/Drivers/JobDriver.Tests/PerforceMaterializerTests.cs
--- a/Engine/Source/Programs/Horde/Drivers/JobDriver.Tests/PerforceMaterializerTests.cs
+++ b/Engine/Source/Programs/Horde/Drivers/JobDriver.Tests/PerforceMaterializerTests.cs
@@ -74,11 +74,11 @@ public class PerforceMaterializerTests : BasePerforceFixtureTest
 		using PerforceMaterializer pm = CreatePm(Fixture.StreamFooMain);
 		
 		await pm.SyncAsync(IWorkspaceMaterializer.LatestChangeNumber, -1, new SyncOptions(), CancellationToken.None);
-		Fixture.StreamFooMain.GetChangelist(7).AssertDepotFiles(pm.SyncDir.FullName);
+		Fixture.StreamFooMain.Head.AssertDepotFiles(pm.SyncDir.FullName);
 		await pm.FinalizeAsync(CancellationToken.None);
 
 		PerforceMaterializer.State? state = await pm.LoadStateForTestAsync(CancellationToken.None);
-		Assert.AreEqual(7, state!.Changelist);
+		Assert.AreEqual(9, state!.Changelist);
 	}
 	
 	[TestMethod]
@@ -106,7 +106,7 @@ public class PerforceMaterializerTests : BasePerforceFixtureTest
 		await pm.FinalizeAsync(CancellationToken.None);
 
 		PerforceMaterializer.State? state = await pm.LoadStateForTestAsync(CancellationToken.None);
-		Assert.AreEqual(7, state!.Changelist);
+		Assert.AreEqual(9, state!.Changelist);
 	}
 	
 	[TestMethod]
@@ -116,7 +116,7 @@ public class PerforceMaterializerTests : BasePerforceFixtureTest
 		
 		await pm.SyncAsync(IWorkspaceMaterializer.LatestChangeNumber, -1, new SyncOptions(), CancellationToken.None);
 		await pm.SyncAsync(IWorkspaceMaterializer.LatestChangeNumber, -1, new SyncOptions(), CancellationToken.None);
-		Fixture.StreamFooMain.GetChangelist(7).AssertDepotFiles(pm.SyncDir.FullName);
+		Fixture.StreamFooMain.Head.AssertDepotFiles(pm.SyncDir.FullName);
 		await pm.FinalizeAsync(CancellationToken.None);
 
 		PerforceMaterializer.State? state = await pm.LoadStateForTestAsync(CancellationToken.None);
@@ -211,8 +211,10 @@ public class PerforceMaterializerTests : BasePerforceFixtureTest
 	[TestMethod]
 	public async Task TestAllStatesAsync()
 	{
-		StreamState[] mainState = [Main(2), Main(4), Main(5), Main(7), Main(7, 8)];
-		StreamState[] releaseStates = [Release(10), Release(11)];
+		StreamState[] mainState = [Main(2), Main(4), Main(5), Main(7), Main(7, 8), Main(9)];
+		StreamState[] releaseStates = [Release(11), Release(12)];
+		mainState = mainState.Select(x => x with { RemoveUntracked = false }).ToArray();
+		releaseStates = releaseStates.Select(x => x with { RemoveUntracked = false }).ToArray();
 		
 		StreamState?[] localStates = [null, ..mainState, ..releaseStates];
 		StreamState[] targetStates = [..mainState, ..releaseStates];
@@ -233,9 +235,9 @@ public class PerforceMaterializerTests : BasePerforceFixtureTest
 		}
 	}
 
-	private StreamState Main(int cl, int shelvedCl = 0) => new (Fixture.StreamFooMain, cl, shelvedCl);
-	private StreamState Release(int cl, int shelvedCl = 0) => new (Fixture.StreamFooRelease, cl, shelvedCl);
-	private record StreamState(StreamFixture Stream, int Cl, int ShelvedCl = 0)
+	private StreamState Main(int cl, int shelvedCl = -1) => new (Fixture.StreamFooMain, cl, shelvedCl);
+	private StreamState Release(int cl, int shelvedCl = -1) => new (Fixture.StreamFooRelease, cl, shelvedCl);
+	private record StreamState(StreamFixture Stream, int Cl, int ShelvedCl = -1, bool RemoveUntracked = true)
 	{
 		public override string ToString()
 		{
@@ -273,7 +275,7 @@ public class PerforceMaterializerTests : BasePerforceFixtureTest
 				Directory.Move(stateDir, backupDir);
 			}
 			
-            using PerforceMaterializer pm = CreatePm(scenario.ServerState.Stream, identifier: id);
+			using PerforceMaterializer pm = CreatePm(scenario.ServerState.Stream, identifier: id);
 			await SyncToStateAsync(pm, scenario.ServerState, cancellationToken);
 
 			// Now that the workspace has been synced, reset the local files back
@@ -296,7 +298,7 @@ public class PerforceMaterializerTests : BasePerforceFixtureTest
 
 	private static async Task SyncToStateAsync(PerforceMaterializer pm, StreamState state, CancellationToken cancellationToken)
 	{
-		await pm.SyncAsync(state.Cl, state.ShelvedCl, new SyncOptions(), cancellationToken);
+		await pm.SyncAsync(state.Cl, state.ShelvedCl, new SyncOptions { RemoveUntracked = state.RemoveUntracked }, cancellationToken);
 		state.Stream.GetChangelist(state.ShelvedCl > 0 ? state.ShelvedCl : state.Cl).AssertDepotFiles(pm.SyncDir.FullName);
 		await pm.FinalizeAsync(cancellationToken);
 		
diff --git a/Engine/Source/Programs/Shared/EpicGames.Perforce.Fixture/PerforceFixture.cs b/Engine/Source/Programs/Shared/EpicGames.Perforce.Fixture/PerforceFixture.cs
--- a/Engine/Source/Programs/Shared/EpicGames.Perforce.Fixture/PerforceFixture.cs
+++ b/Engine/Source/Programs/Shared/EpicGames.Perforce.Fixture/PerforceFixture.cs
@@ -220,14 +220,13 @@ public class StreamFixture
 			}
 		}
 
-		// Ignore null check to make tests more readable
-		return null!;
+		throw new Exception($"Change {changeNum} not found in {Root}");
 	}
 }
 
 public class PerforceFixture
 {
-	public StreamFixture StreamFooMain { get; } = new("//Foo/Main", 7,
+	public StreamFixture StreamFooMain { get; } = new("//Foo/Main", 9,
 		new List<ChangelistFixture>
 		{
 			new(0, ChangelistFixture.Placeholder, new List<DepotFileFixture>()),
@@ -289,15 +288,25 @@ public class PerforceFixture
 					new("//Foo/Main/main.h", "main.h", "This is change main.h #3\n", 3),
 					new("//Foo/Main/shared.h", "shared.h", "This is change common.h #1\n", 1),
 					new("//Foo/Main/shelved.cpp", "shelved.cpp", "This is change shelved.cpp #1\n", 1),
+					new("//Foo/Main/unused.cpp", "unused.cpp", "This is change unused.cpp #2\n", 1),
 					new("//Foo/Main/Data/data.txt", "Data/data.txt", "This is change data.txt #1\n", 1),
 					new("//Foo/Main/Data/moredata.txt", "Data/moredata.txt", "This is change moredata.txt #1\n", 1),
 				}, true),
+			new(9, "Did more changes to main",
+				new List<DepotFileFixture>()
+				{
+					new("//Foo/Main/main.cpp", "main.cpp", "This is change main.cpp #4\n", 3),
+					new("//Foo/Main/main.h", "main.h", "This is change main.h #4\n", 2),
+					new("//Foo/Main/shared.h", "shared.h", "This is change common.h #1\n", 1),
+					new("//Foo/Main/Data/data.txt", "Data/data.txt", "This is change data.txt #1\n", 1),
+					new("//Foo/Main/Data/moredata.txt", "Data/moredata.txt", "This is change moredata.txt #1\n", 1),
+				})
 		});
 
 	public StreamFixture StreamFooRelease { get; } = new("//Foo/Release", 11,
 		new List<ChangelistFixture>
 		{
-			new(10, "Initial import",
+			new(11, "Initial import",
 				new List<DepotFileFixture>()
 				{
 					new("//Foo/Main/main.cpp", "main.cpp", "This is change main.cpp #2\n", 1),
@@ -305,7 +314,7 @@ public class PerforceFixture
 					new("//Foo/Main/shared.h", "shared.h", "This is change common.h #1\n", 1),
 					// data.txt is excluded in view
 				}),
-			new(11, "Initial import",
+			new(12, "Initial import",
 				new List<DepotFileFixture>()
 				{
 					new("//Foo/Main/main.cpp", "main.cpp", "This is change main.cpp #2 in //Foo/Release\n", 1),
@@ -316,14 +325,14 @@ public class PerforceFixture
 				}),
 		});
 	
-	public StreamFixture StreamFooSparseDev { get; } = new("//Foo/Main-SparseDev", 12,
+	public StreamFixture StreamFooSparseDev { get; } = new("//Foo/Main-SparseDev", 13,
 		new List<ChangelistFixture>
 		{
-			new(12, "Stream //Foo/Main-SparseDev saved.", // TODO: add action
+			new(13, "Stream //Foo/Main-SparseDev saved.", // TODO: add action
 				new List<DepotFileFixture>()
 				{
-					new("//Foo/Main/main.cpp", "main.cpp", "This is change main.cpp #3\n", 3),
-					new("//Foo/Main/main.h", "main.h", "This is change main.h #2\n", 2),
+					new("//Foo/Main/main.cpp", "main.cpp", "This is change main.cpp #4\n", 3),
+					new("//Foo/Main/main.h", "main.h", "This is change main.h #4\n", 2),
 					new("//Foo/Main/shared.h", "shared.h", "This is change common.h #1\n", 1),
 					new("//Foo/Main/Data/data.txt", "Data/data.txt", "This is change data.txt #1\n", 1),
 					new("//Foo/Main/Data/moredata.txt", "Data/moredata.txt", "This is change moredata.txt #1\n", 1),
diff --git a/Engine/Source/Programs/Shared/EpicGames.Perforce.Fixture/setup.sh b/Engine/Source/Programs/Shared/EpicGames.Perforce.Fixture/setup.sh
--- a/Engine/Source/Programs/Shared/EpicGames.Perforce.Fixture/setup.sh
+++ b/Engine/Source/Programs/Shared/EpicGames.Perforce.Fixture/setup.sh
@@ -122,10 +122,20 @@ p4 change -i << EOF
 Change: new
 Description: A shelved CL
 EOF
+# Re-add a file that's been deleted. This is known to cause tracking issues.
+echo "This is change unused.cpp #2" > $CLIENT_ROOT/unused.cpp && p4 add -c 8 $CLIENT_ROOT/unused.cpp
 echo "This is change shelved.cpp #1" > $CLIENT_ROOT/shelved.cpp && p4 add -c 8 $CLIENT_ROOT/shelved.cpp
 p4 edit -c 8 main.h && echo "This is change main.h #3" > $CLIENT_ROOT/main.h
 p4 delete -c 8 main.cpp
 p4 shelve -c 8
+p4 revert -c 8 //...
+
+# Changelist #9
+p4 edit main.cpp
+p4 edit main.h
+echo "This is change main.cpp #4" > $CLIENT_ROOT/main.cpp
+echo "This is change main.h #4" > $CLIENT_ROOT/main.h
+p4 submit -d "Did more changes to main"
 
 
 # Create alternative release stream
@@ -157,11 +167,11 @@ EOF
 mkdir -p $CLIENT_ROOT
 cd $CLIENT_ROOT
 
-# Changelist #10
+# Changelist #11
 p4 populate -d "Populate from CL 5" //Foo/Main/...@5 //Foo/Release/...
 p4 sync
 
-# Changelist #11
+# Changelist #12
 p4 edit main.cpp
 echo "This is change main.cpp #2 in //Foo/Release" > $CLIENT_ROOT/main.cpp
 echo "This is change release.cpp #1 in //Foo/Release" > $CLIENT_ROOT/release.cpp && p4 add $CLIENT_ROOT/release.cpp
