Horde: Adding Epic Sandbox project to soluthon and improving conditional

diff --git a/Engine/Source/Programs/Horde/HordeServer/HordeServer.csproj b/Engine/Source/Programs/Horde/HordeServer/HordeServer.csproj
--- a/Engine/Source/Programs/Horde/HordeServer/HordeServer.csproj
+++ b/Engine/Source/Programs/Horde/HordeServer/HordeServer.csproj
@@ -123,13 +123,14 @@
     <ProjectReference Include="..\Plugins\Tools\HordeServer.Tools\HordeServer.Tools.csproj" />
   </ItemGroup>
 
-   <Choose>
-      <When Condition="Exists('..\..\..\..\Restricted\NotForLicensees\Source\Programs\Horde\EpicPlugins\BuildHealth\HordeServer.Epic.BuildHealth\HordeServer.Epic.BuildHealth.csproj')">
-         <ItemGroup>
-            <ProjectReference Include="..\..\..\..\Restricted\NotForLicensees\Source\Programs\Horde\EpicPlugins\BuildHealth\HordeServer.Epic.BuildHealth\HordeServer.Epic.BuildHealth.csproj" />
-         </ItemGroup>
-      </When>
-   </Choose>		
+  <ItemGroup Condition="Exists('$(MSBuildThisFileDirectory)..\..\..\..\Restricted\NotForLicensees\Source\Programs\Horde\EpicPlugins\BuildHealth\HordeServer.Epic.BuildHealth\HordeServer.Epic.BuildHealth.csproj')">
+    <ProjectReference Include="..\..\..\..\Restricted\NotForLicensees\Source\Programs\Horde\EpicPlugins\BuildHealth\HordeServer.Epic.BuildHealth\HordeServer.Epic.BuildHealth.csproj" />
+  </ItemGroup>
+
+  <ItemGroup Condition="Exists('$(MSBuildThisFileDirectory)..\..\..\..\Restricted\NotForLicensees\Source\Programs\Horde\EpicPlugins\EpicSandbox\HordeServer.Epic.EpicSandbox\HordeServer.Epic.EpicSandbox.csproj')">
+    <ProjectReference Include="..\..\..\..\Restricted\NotForLicensees\Source\Programs\Horde\EpicPlugins\EpicSandbox\HordeServer.Epic.EpicSandbox\HordeServer.Epic.EpicSandbox.csproj" />
+  </ItemGroup>
+
   <ItemGroup>
     <Content Include="..\ThirdParty\**" Condition="$([MSBuild]::IsOSPlatform('Windows'))">
       <CopyToOutputDirectory>Always</CopyToOutputDirectory>
