Horde: Add schedule filter to summary pane

diff --git a/Engine/Source/Programs/Horde/HordeDashboard/src/components/SchedulePane.tsx b/Engine/Source/Programs/Horde/HordeDashboard/src/components/SchedulePane.tsx
--- a/Engine/Source/Programs/Horde/HordeDashboard/src/components/SchedulePane.tsx
+++ b/Engine/Source/Programs/Horde/HordeDashboard/src/components/SchedulePane.tsx
@@ -1,9 +1,9 @@
 // Copyright Epic Games, Inc. All Rights Reserved.
-import { mergeStyleSets, Stack } from '@fluentui/react';
+import { IIconProps, mergeStyleSets, SearchBox, Stack } from '@fluentui/react';
 import { DetailsList, DetailsListLayoutMode, DetailsRow, IColumn, IDetailsListProps, SelectionMode } from '@fluentui/react/lib/DetailsList';
 import moment from 'moment-timezone';
 import React, { useState } from 'react';
-import { Link, useNavigate } from 'react-router-dom';
+import { Link } from 'react-router-dom';
 import backend from '../backend';
 import { GetJobResponse, GetTemplateRefResponse, JobState } from '../backend/Api';
 import dashboard from '../backend/Dashboard';
@@ -30,7 +30,7 @@ const scheduledJobs = new Map<string, GetJobResponse | null>();
 export const SchedulePane: React.FC<{ streamId: string, templates: GetTemplateRefResponse[] }> = ({ streamId, templates }) => {
 
    let [updated, setUpdated] = useState(new Date().getTime());
-   const navigate = useNavigate();
+   const [searchText, setSearchText] = useState("");
 
    const ids = new Set<string>();
 
@@ -66,6 +66,25 @@ export const SchedulePane: React.FC<{ streamId: string, templates: GetTemplateRe
 
    const items = templates.filter(t => !!t.schedule?.enabled).map(t => {
       return { template: t };
+   }).filter(item => {
+      if (!searchText) return true;
+
+      const lowerSearch = searchText.toLowerCase();
+      const template = item.template;
+      const schedule = template.schedule!;
+
+      if (template.name.toLowerCase().includes(lowerSearch)) return true;
+
+      const gateTemplateId = schedule.gate?.templateId;
+      if (gateTemplateId) {
+         const gateTemplate = templates.find(t => t.id === gateTemplateId);
+         if (gateTemplate?.name.toLowerCase().includes(lowerSearch)) return true;
+      }
+
+      const target = schedule.gate?.target;
+      if (target?.toLowerCase().includes(lowerSearch)) return true;
+
+      return false;
    });
 
    const columns = [
@@ -96,7 +115,7 @@ export const SchedulePane: React.FC<{ streamId: string, templates: GetTemplateRe
          let gateTemplate = templates.find(t => t.id === templateId);
 
 
-         return <Stack horizontalAlign="start" verticalFill verticalAlign="center" tokens={{ childrenGap: 4 }} style={{paddingLeft: 12}}>
+         return <Stack horizontalAlign="start" verticalFill verticalAlign="center" tokens={{ childrenGap: 4 }} style={{ paddingLeft: 12 }}>
             <Stack style={{ fontWeight: 600 }}>{name}</Stack>
             {!!gateTemplate && <Stack tokens={{ childrenGap: 4 }}>
                <Stack horizontal tokens={{ childrenGap: 8 }}>
@@ -115,10 +134,10 @@ export const SchedulePane: React.FC<{ streamId: string, templates: GetTemplateRe
       if (name === "audit") {
          return <Link to={`/audit/template/${encodeURIComponent(streamId)}/${encodeURIComponent(template.id)}`}>
             <Stack horizontalAlign="end" verticalFill verticalAlign="center" tokens={{ childrenGap: 4 }}>
-            <Stack tokens={{ childrenGap: 4 }}>
-               <Stack >View Log</Stack>
+               <Stack tokens={{ childrenGap: 4 }}>
+                  <Stack >View Log</Stack>
+               </Stack>
             </Stack>
-         </Stack>
          </Link>
       }
 
@@ -195,22 +214,37 @@ export const SchedulePane: React.FC<{ streamId: string, templates: GetTemplateRe
       return null;
    };
 
-   return <div className={classes.detailsRow} style={{ overflowY: 'auto', overflowX: 'hidden', maxHeight: "820px" }} data-is-scrollable={true}>
-
-      <DetailsList
-         styles={{ root: { overflow: "unset" } }}
-         compact={true}
-         isHeaderVisible={false}
-         indentWidth={0}
-         items={items}
-         columns={columns}
-         setKey="set"
-         selectionMode={SelectionMode.none}
-         layoutMode={DetailsListLayoutMode.justified}
-         onRenderItemColumn={renderItem}
-         onRenderRow={renderRow}
-      />
-   </div>
+   
+   const filterIcon: IIconProps = { iconName: 'Filter' };
+
+   return <Stack>
+      <Stack horizontalAlign="end">
+         <SearchBox
+            iconProps={filterIcon}
+            placeholder="Filter schedules..."
+            value={searchText}
+            onChange={(_, newValue) => setSearchText(newValue || "")}
+            onClear={() => setSearchText("")}
+            styles={{ root: { marginBottom: 8, width: 400 } }}
+         />
+      </Stack>
+      <div className={classes.detailsRow} style={{ overflowY: 'auto', overflowX: 'hidden', maxHeight: "820px" }} data-is-scrollable={true}>
+
+         <DetailsList
+            styles={{ root: { overflow: "unset" } }}
+            compact={true}
+            isHeaderVisible={false}
+            indentWidth={0}
+            items={items}
+            columns={columns}
+            setKey="set"
+            selectionMode={SelectionMode.none}
+            layoutMode={DetailsListLayoutMode.justified}
+            onRenderItemColumn={renderItem}
+            onRenderRow={renderRow}
+         />
+      </div>
+   </Stack>
 
 
 }
