[Horde][BuildHealth][Server] Removal of old JobStepRef "skipped state" emission code in entirety in preparation for V2.

#Synced-CL 48440999

diff --git a/Engine/Source/Programs/Horde/Plugins/Build/HordeServer.Build/Jobs/IJobStepRefCollection.cs b/Engine/Source/Programs/Horde/Plugins/Build/HordeServer.Build/Jobs/IJobStepRefCollection.cs
--- a/Engine/Source/Programs/Horde/Plugins/Build/HordeServer.Build/Jobs/IJobStepRefCollection.cs
+++ b/Engine/Source/Programs/Horde/Plugins/Build/HordeServer.Build/Jobs/IJobStepRefCollection.cs
@@ -61,21 +61,6 @@ namespace HordeServer.Jobs
 		/// <returns></returns>
 		Task<IJobStepRef?> FindAsync(JobId jobId, JobStepBatchId batchId, JobStepId stepId);
 
-		/// <summary>
-		/// Gets job step references given a jobId and batchId.
-		/// </summary>
-		/// <param name="jobId">The job to obtain the steps from.</param>
-		/// <param name="batchId">The batch to obtain the steps from.</param>
-		/// <returns>The list of job step refs.</returns>
-		Task<List<IJobStepRef>> FindAsync(JobId jobId, JobStepBatchId batchId);
-
-		/// <summary>
-		/// Gets job step references given a jobId.
-		/// </summary>
-		/// <param name="jobId">The job to obtain the steps from.</param>
-		/// <returns>The list of job step refs.</returns>
-		Task<List<IJobStepRef>> FindAsync(JobId jobId);
-
 		/// <summary>
 		/// Gets job step references given an array of ids
 		/// </summary>
@@ -180,109 +165,5 @@ namespace HordeServer.Jobs
 				await jobStepRefs.InsertOrReplaceAsync(new JobStepRefId(job.Id, batch.Id, step.Id), job.Name, nodeName, job.StreamId, job.TemplateId, job.CommitId, step.LogId, batch.PoolId, batch.AgentId, step.State, outcome, job.UpdateIssues, lastSuccess, lastWarning, waitTime, initTime, job.CreateTimeUtc, step.StartTimeUtc ?? DateTime.UtcNow, step.FinishTimeUtc, step.Metadata?.ToList());
 			}
 		}
-
-		/// <summary>
-		/// Updates a job step ref.
-		/// </summary>
-		/// <param name="jobStepRefs">The job step ref collection.</param>
-		/// <param name="job">The job which the job step belongs to..</param>
-		/// <param name="batch">The batch which the job step belongs to.</param>
-		/// <param name="existingStepRef">The existing job step ref to update.</param>
-		/// <param name="logger">The logger.</param>
-		/// <returns>Task object associated with the operation.</returns>
-		public static async Task UpdateAsync(this IJobStepRefCollection jobStepRefs, IJob job, IJobStepBatch batch, IJobStepRef existingStepRef, ILogger? logger = null)
-		{
-			if (job.PreflightCommitId == null)
-			{
-				float waitTime = (float)(batch.GetWaitTime() ?? TimeSpan.Zero).TotalSeconds;
-				float initTime = (float)(batch.GetInitTime() ?? TimeSpan.Zero).TotalSeconds;
-
-				string nodeName = existingStepRef.NodeName;
-
-				CommitIdWithOrder? lastSuccess = null;
-				CommitIdWithOrder? lastWarning = null;
-				JobStepOutcome? newOutcome = existingStepRef.Outcome;
-				JobStepState? newState = existingStepRef.State;
-
-				if (newOutcome != JobStepOutcome.Success)
-				{
-					IJobStepRef? prevStep = await jobStepRefs.GetPrevStepForNodeAsync(job.StreamId, job.TemplateId, nodeName, job.CommitId);
-					if (prevStep != null)
-					{
-						lastSuccess = prevStep.LastSuccess;
-						if (newOutcome != JobStepOutcome.Warnings)
-						{
-							lastWarning = prevStep.LastWarning;
-						}
-					}
-				}
-
-				if (job.AbortedByUserId != null)
-				{
-					newOutcome = JobStepOutcome.Failure;
-					newState = JobStepState.Aborted;
-				}
-
-				if (batch.Error != JobStepBatchError.None)
-				{
-					newOutcome = JobStepOutcome.Failure;
-				}
-
-				if (logger != null)
-				{
-					logger.LogInformation("Updating step reference {StepId} for job {JobId}, batch {BatchId}, with outcome {JobStepOutcome}", existingStepRef.Id, job.Id, batch.Id, newOutcome);
-				}
-
-				await jobStepRefs.InsertOrReplaceAsync(existingStepRef.Id, job.Name, nodeName, job.StreamId, job.TemplateId, job.CommitId, existingStepRef.LogId, batch.PoolId, batch.AgentId, newState, newOutcome, job.UpdateIssues, lastSuccess, lastWarning, waitTime, initTime, job.CreateTimeUtc, existingStepRef.StartTimeUtc, existingStepRef.FinishTimeUtc, existingStepRef.Metadata?.ToList());
-			}
-		}
-
-		/// <summary>
-		/// Updates a job step ref.
-		/// </summary>
-		/// <param name="jobStepRefs">The job step ref collection.</param>
-		/// <param name="job">The job which the job step belongs to..</param>
-		/// <param name="existingStepRef">The existing job step ref to update.</param>
-		/// <param name="logger">The logger.</param>
-		/// <returns>Task object associated with the operation.</returns>
-		public static async Task UpdateAsync(this IJobStepRefCollection jobStepRefs, IJob job, IJobStepRef existingStepRef, ILogger? logger = null)
-		{
-			if (job.PreflightCommitId == null)
-			{
-
-				string nodeName = existingStepRef.NodeName;
-
-				CommitIdWithOrder? lastSuccess = null;
-				CommitIdWithOrder? lastWarning = null;
-				JobStepOutcome? newOutcome = existingStepRef.Outcome;
-				JobStepState? newState = existingStepRef.State;
-
-				if (newOutcome != JobStepOutcome.Success)
-				{
-					IJobStepRef? prevStep = await jobStepRefs.GetPrevStepForNodeAsync(job.StreamId, job.TemplateId, nodeName, job.CommitId);
-					if (prevStep != null)
-					{
-						lastSuccess = prevStep.LastSuccess;
-						if (newOutcome != JobStepOutcome.Warnings)
-						{
-							lastWarning = prevStep.LastWarning;
-						}
-					}
-				}
-
-				if (job.AbortedByUserId != null)
-				{
-					newOutcome = JobStepOutcome.Failure;
-					newState = JobStepState.Aborted;
-				}
-
-				if (logger != null)
-				{
-					//logger.LogInformation("Updating step reference {StepId} for job {JobId}, batch {BatchId}, with outcome {JobStepOutcome}", , job.Id, batch.Id, newOutcome);
-				}
-
-				await jobStepRefs.InsertOrReplaceAsync(existingStepRef.Id, job.Name, nodeName, job.StreamId, job.TemplateId, job.CommitId, existingStepRef.LogId, existingStepRef.PoolId, existingStepRef.AgentId, newState, newOutcome, job.UpdateIssues, lastSuccess, lastWarning, existingStepRef.BatchWaitTime, existingStepRef.BatchInitTime, job.CreateTimeUtc, existingStepRef.StartTimeUtc, existingStepRef.FinishTimeUtc, existingStepRef.Metadata?.ToList());
-			}
-		}
 	}
 }
diff --git a/Engine/Source/Programs/Horde/Plugins/Build/HordeServer.Build/Jobs/JobCollection.cs b/Engine/Source/Programs/Horde/Plugins/Build/HordeServer.Build/Jobs/JobCollection.cs
--- a/Engine/Source/Programs/Horde/Plugins/Build/HordeServer.Build/Jobs/JobCollection.cs
+++ b/Engine/Source/Programs/Horde/Plugins/Build/HordeServer.Build/Jobs/JobCollection.cs
@@ -914,7 +914,7 @@ namespace HordeServer.Jobs
 
 			JobAudit(jobId, "Job {JobId} added, stream {StreamId}, templateRefId {TemplateId}, name {Name}, commitId {CommitId}", jobId, streamId.Id, templateRefId, name, commitId);
 
-			await CreateBatchesInternalAsync(newJob, graph, _logger, GetJobLogger(jobId));
+			CreateBatchesInternal(newJob, graph, _logger, GetJobLogger(jobId));
 
 			await _jobs.InsertOneAsync(newJob, null, cancellationToken);
 
@@ -1382,7 +1382,7 @@ namespace HordeServer.Jobs
 			// Update the batches
 			if (updateBatches)
 			{
-				await UpdateBatchesInternalAsync(jobDocument, graph, updates, _logger, GetJobLogger(jobId));
+				UpdateBatchesInternal(jobDocument, graph, updates, _logger, GetJobLogger(jobId));
 			}
 
 			// Update the new list of job steps
@@ -1468,7 +1468,7 @@ namespace HordeServer.Jobs
 
 				// Force an update of all batches in the job. This will fail or reschedule any nodes that can no longer be executed in this batch.
 				updates.Clear();
-				await UpdateBatchesInternalAsync(jobDocument, graph, updates, _logger, GetJobLogger(jobDocument.Id));
+				UpdateBatchesInternal(jobDocument, graph, updates, _logger, GetJobLogger(jobDocument.Id));
 
 				if (retriedNodes.Count > 0)
 				{
@@ -1719,18 +1719,13 @@ namespace HordeServer.Jobs
 			if (refreshBatches)
 			{
 				updates.Clear(); // UpdateBatches will update the entire batches list. We need to remove all the individual step updates to avoid an exception.
-				await UpdateBatchesInternalAsync(jobDocument, graph, updates, _logger, GetJobLogger(jobDocument.Id));
+				UpdateBatchesInternal(jobDocument, graph, updates, _logger, GetJobLogger(jobDocument.Id));
 			}
 
 			// Update the state of dependent jobsteps
 			if (refreshDependentJobSteps)
 			{
 				RefreshDependentJobSteps(jobDocument, graph, updates, _logger, GetJobLogger(jobDocument.Id));
-				/*
-				IJob jobObj = CreateJobObject(jobDocument, graph);
-				await UpdateRefStateAsync(jobObj, graph, _jobStepRefs);
-				*/
-
 				RefreshJobPriority(jobDocument, updates, _logger, GetJobLogger(jobDocument.Id));
 			}
 
@@ -1867,7 +1862,7 @@ namespace HordeServer.Jobs
 			jobDocument.GraphHash = newGraph.Id;
 			updates.Add(updateBuilder.Set(x => x.GraphHash, jobDocument.GraphHash));
 
-			await UpdateBatchesInternalAsync(jobDocument, newGraph, updates, _logger, GetJobLogger(jobDocument.Id));
+			UpdateBatchesInternal(jobDocument, newGraph, updates, _logger, GetJobLogger(jobDocument.Id));
 
 			// Update the new list of job steps
 			return await TryUpdateAsync(jobDocument, updates, cancellationToken);
@@ -1947,7 +1942,7 @@ namespace HordeServer.Jobs
 			}
 
 			List<UpdateDefinition<JobDocument>> updates = new List<UpdateDefinition<JobDocument>>();
-			await UpdateBatchesInternalAsync(jobDocument, graph, updates, _logger, GetJobLogger(jobDocument.Id));
+			UpdateBatchesInternal(jobDocument, graph, updates, _logger, GetJobLogger(jobDocument.Id));
 
 			return await TryUpdateAsync(jobDocument, updates, cancellationToken);
 		}
@@ -1978,7 +1973,7 @@ namespace HordeServer.Jobs
 			}
 
 			List<UpdateDefinition<JobDocument>> updates = new List<UpdateDefinition<JobDocument>>();
-			await UpdateBatchesInternalAsync(jobDocument, graph, updates, _logger, GetJobLogger(jobDocument.Id));
+			UpdateBatchesInternal(jobDocument, graph, updates, _logger, GetJobLogger(jobDocument.Id));
 
 			return await TryUpdateAsync(jobDocument, updates, cancellationToken);
 		}
@@ -2090,27 +2085,17 @@ namespace HordeServer.Jobs
 		}
 
 #pragma warning disable CA1822 // Mark members as static
-		private async Task UpdateBatchesInternalAsync(JobDocument job, IGraph graph, List<UpdateDefinition<JobDocument>> updates, ILogger logger, IAuditLogChannel<JobId> jobAudit)
+		private void UpdateBatchesInternal(JobDocument job, IGraph graph, List<UpdateDefinition<JobDocument>> updates, ILogger logger, IAuditLogChannel<JobId> jobAudit)
 #pragma warning restore CA1822 // Mark members as static
 		{
 			UpdateBatches(job, graph, updates, logger, jobAudit);
-			/* UE-354742 - Reviewing JobStepRef Emission
-			IJob jobObj = CreateJobObject(job, graph);
-			await UpdateRefStateAsync(jobObj, graph, _jobStepRefs);
-			*/
-			await ValueTask.CompletedTask;
 		}
 
 #pragma warning disable CA1822 // Mark members as static
-		private async Task CreateBatchesInternalAsync(JobDocument job, IGraph graph, ILogger logger, IAuditLogChannel<JobId> jobAudit)
+		private void CreateBatchesInternal(JobDocument job, IGraph graph, ILogger logger, IAuditLogChannel<JobId> jobAudit)
 #pragma warning restore CA1822 // Mark members as static
 		{
 			CreateBatches(job, graph, logger, jobAudit);
-			/* UE-354742 - Reviewing JobStepRef Emission
-			IJob jobObj = CreateJobObject(job, graph);
-			await UpdateRefStateAsync(jobObj, graph, _jobStepRefs);
-			*/
-			await ValueTask.CompletedTask;
 		}
 
 		/// <summary>
diff --git a/Engine/Source/Programs/Horde/Plugins/Build/HordeServer.Build/Jobs/JobStepRefCollection.cs b/Engine/Source/Programs/Horde/Plugins/Build/HordeServer.Build/Jobs/JobStepRefCollection.cs
--- a/Engine/Source/Programs/Horde/Plugins/Build/HordeServer.Build/Jobs/JobStepRefCollection.cs
+++ b/Engine/Source/Programs/Horde/Plugins/Build/HordeServer.Build/Jobs/JobStepRefCollection.cs
@@ -212,32 +212,6 @@ namespace HordeServer.Jobs
 			return await _jobStepRefs.FindOneAndUpdateAsync(x => x.Id.Equals(id), updateBuilder.Combine(updates));
 		}
 
-		/// <inheritdoc/>
-		public async Task<List<IJobStepRef>> FindAsync(JobId jobId, JobStepBatchId batchId)
-		{
-			/* UE-354742 - Reviewing JobStepRef Emission
-			string prefix = $"^{jobId}:{batchId}:";
-			FilterDefinition<JobStepRef> filter = Builders<JobStepRef>.Filter.Regex(x => x.Id, new BsonRegularExpression(prefix));
-
-			List<JobStepRef> results = await _jobStepRefs.Find(filter).ToListAsync();
-			*/
-			await ValueTask.CompletedTask;
-			return [];
-		}
-
-		/// <inheritdoc/>
-		public async Task<List<IJobStepRef>> FindAsync(JobId jobId)
-		{
-			/* UE-354742 - Reviewing JobStepRef Emission
-			string prefix = $"^{jobId}:";
-			FilterDefinition<JobStepRef> filter = Builders<JobStepRef>.Filter.Regex(x => x.Id, new BsonRegularExpression(prefix));
-
-			List<JobStepRef> results = await _jobStepRefs.Find(filter).ToListAsync();
-			*/
-			await ValueTask.CompletedTask;
-			return [];
-		}
-
 		/// <inheritdoc/>
 		public async Task<IJobStepRef?> FindAsync(JobId jobId, JobStepBatchId batchId, JobStepId stepId)
 		{
