Horde: Use p4 print for unshelving in PerforceMaterializer

- Avoids issues with files being locked and mimics how ManagedWorkspace works
- Also change to snapshot-based testing

diff --git a/Engine/Source/Programs/Horde/Drivers/JobDriver.Tests/PerforceMaterializerTests.cs b/Engine/Source/Programs/Horde/Drivers/JobDriver.Tests/PerforceMaterializerTests.cs
--- a/Engine/Source/Programs/Horde/Drivers/JobDriver.Tests/PerforceMaterializerTests.cs
+++ b/Engine/Source/Programs/Horde/Drivers/JobDriver.Tests/PerforceMaterializerTests.cs
@@ -14,7 +14,6 @@ using Microsoft.Extensions.DependencyInjection;
 using Microsoft.Extensions.Logging;
 using Microsoft.VisualStudio.TestTools.UnitTesting;
 using OpenTelemetry.Trace;
-using FileSet = System.Collections.Generic.HashSet<(string clientFile, long size, string digest)>;
 using SyncOptions = JobDriver.Execution.SyncOptions;
 
 namespace JobDriver.Tests;
@@ -30,13 +29,13 @@ public class PerforceMaterializerTests : BasePerforceFixtureTest
 	public async Task WalkChangeListsWithReusedMaterializerAsync()
 	{
 		CancellationToken ct = CancellationToken.None;
-		using PerforceMaterializer pm = CreatePm(Fixture.StreamFooMain);
-		List<StreamState> states = [Main(2), Main(3), Main(4), Main(5), Main(6), Main(7), Main(7, 8)];
+		using PerforceMaterializer pm = CreatePm(PerforceSnapshot.FooMain);
+		List<WorkspaceState> states = new(PerforceSnapshot.FooMainStates);
 		states.AddRange(states.SkipLast(1).Reverse());
-		foreach (StreamState ss in states)
+		foreach (WorkspaceState ss in states)
 		{
-			await pm.SyncAsync(ss.Cl, ss.ShelvedCl, new SyncOptions(), ct);
-			Fixture.StreamFooMain.GetChangelist(ss.ShelvedCl > 0 ? ss.ShelvedCl : ss.Cl).AssertDepotFiles(pm.SyncDir.FullName);
+			await pm.SyncAsync(ss.Cl, ss.ShelvedCl, new SyncOptions() { RemoveUntracked = false }, ct);
+			PerforceSnapshot.GetSnapshot().GetStream(ss.Stream).GetWorkspace(ss.Cl, ss.ShelvedCl).AssertDepotFiles(pm.SyncDir.FullName);
 			await pm.FinalizeAsync(ct);
 		}
 	}
@@ -45,115 +44,53 @@ public class PerforceMaterializerTests : BasePerforceFixtureTest
 	[Ignore] // Disabled for now as it mutates the stream. Need to obliterate after submit, or similar kind of restore
 	public async Task Submit_Async()
 	{
-		using PerforceMaterializer pm = CreatePm(Fixture.StreamFooMain);
+		using PerforceMaterializer pm = CreatePm(PerforceSnapshot.FooMain);
 		
 		await pm.SyncAsync(IWorkspaceMaterializer.LatestChangeNumber, -1, new SyncOptions(), CancellationToken.None);
-		Fixture.StreamFooMain.GetChangelist(7).AssertDepotFiles(pm.SyncDir.FullName);
+		PerforceSnapshot.GetSnapshot().GetStream(PerforceSnapshot.FooMain).GetWorkspace(7).AssertDepotFiles(pm.SyncDir.FullName);
 		await SubmitFileAsync(pm, "Adding new file", "hello.txt", "Hello world!");
 		await pm.FinalizeAsync(CancellationToken.None);
 
 		Assert.IsTrue(true, "Submit went through without exception being thrown");
 	}
 
-	private static async Task SubmitFileAsync(PerforceMaterializer pm, string clDescription, string filename, string fileContent)
-	{
-		string tempFile = Path.Join(pm.SyncDir.FullName, filename);
-		await File.WriteAllTextAsync(tempFile, fileContent);
-
-		IPerforceConnection perforce = await pm.GetPerforceWithClientAsync();
-		List<AddRecord> addRecords = await perforce.AddAsync(-1, tempFile);
-		ChangeRecord change = new () { Description = clDescription };
-		change.Files.Add(addRecords[0].DepotFile);
-		change = await perforce.CreateChangeAsync(change);
-		await perforce.SubmitAsync(change.Number, SubmitOptions.None);
-	}
-
 	[TestMethod]
-	public async Task Sync_Latest_Async()
+	[DataRow(PerforceSnapshot.FooMain, IWorkspaceMaterializer.LatestChangeNumber, 9)]
+	[DataRow(PerforceSnapshot.FooMain, IWorkspaceMaterializer.LatestChangeNumber, 9, true)]
+	[DataRow(PerforceSnapshot.FooMainSparseDev, IWorkspaceMaterializer.LatestChangeNumber, 9)]
+	[DataRow(PerforceSnapshot.FooMainSparseDev, 13, 13)]
+	public async Task Sync_Async(string streamName, int changeNum, int expectedChangeNum, bool syncTwice = false)
 	{
-		using PerforceMaterializer pm = CreatePm(Fixture.StreamFooMain);
+		using PerforceMaterializer pm = CreatePm(streamName);
 		
-		await pm.SyncAsync(IWorkspaceMaterializer.LatestChangeNumber, -1, new SyncOptions(), CancellationToken.None);
-		Fixture.StreamFooMain.Head.AssertDepotFiles(pm.SyncDir.FullName);
-		await pm.FinalizeAsync(CancellationToken.None);
-
-		PerforceMaterializer.State? state = await pm.LoadStateForTestAsync(CancellationToken.None);
-		Assert.AreEqual(9, state!.Changelist);
-	}
-	
-	[TestMethod]
-	public async Task Sync_LatestManuallySpecified_SparseDev_Async()
-	{
-		StreamFixture stream = Fixture.StreamFooSparseDev;
-		using PerforceMaterializer pm = CreatePm(stream);
-		
-		await pm.SyncAsync(stream.Head.Number, -1, new SyncOptions(), CancellationToken.None);
-		stream.GetChangelist(stream.Head.Number).AssertDepotFiles(pm.SyncDir.FullName);
-		await pm.FinalizeAsync(CancellationToken.None);
-
-		PerforceMaterializer.State? state = await pm.LoadStateForTestAsync(CancellationToken.None);
-		Assert.AreEqual(stream.Head.Number, state!.Changelist);
-	}
-	
-	[TestMethod]
-	public async Task Sync_Latest_SparseDev_Async()
-	{
-		StreamFixture stream = Fixture.StreamFooSparseDev;
-		using PerforceMaterializer pm = CreatePm(stream);
-		
-		await pm.SyncAsync(IWorkspaceMaterializer.LatestChangeNumber, -1, new SyncOptions(), CancellationToken.None);
-		stream.GetChangelist(stream.Head.Number).AssertDepotFiles(pm.SyncDir.FullName);
-		await pm.FinalizeAsync(CancellationToken.None);
-
-		PerforceMaterializer.State? state = await pm.LoadStateForTestAsync(CancellationToken.None);
-		Assert.AreEqual(9, state!.Changelist);
-	}
-	
-	[TestMethod]
-	public async Task Sync_SameChangelistTwice_Async()
-	{
-		using PerforceMaterializer pm = CreatePm(Fixture.StreamFooMain);
-		
-		await pm.SyncAsync(IWorkspaceMaterializer.LatestChangeNumber, -1, new SyncOptions(), CancellationToken.None);
-		await pm.SyncAsync(IWorkspaceMaterializer.LatestChangeNumber, -1, new SyncOptions(), CancellationToken.None);
-		Fixture.StreamFooMain.Head.AssertDepotFiles(pm.SyncDir.FullName);
+		await pm.SyncAsync(changeNum, -1, new SyncOptions(), CancellationToken.None);
+		if (syncTwice)
+		{
+			await pm.SyncAsync(changeNum, -1, new SyncOptions(), CancellationToken.None);
+		}
+		PerforceSnapshot.GetSnapshot().GetStream(streamName).GetWorkspace(expectedChangeNum).AssertDepotFiles(pm.SyncDir.FullName);
 		await pm.FinalizeAsync(CancellationToken.None);
 
 		PerforceMaterializer.State? state = await pm.LoadStateForTestAsync(CancellationToken.None);
-		Assert.AreEqual(Fixture.StreamFooMain.Head.Number, state!.Changelist);
-	}
-	
-	[TestMethod]
-	[Ignore] // Views not implemented
-	public async Task SyncWithViewAsync()
-	{
-		using PerforceMaterializer pm = CreatePm(Fixture.StreamFooMain);
-		await pm.SyncAsync(7, -1, new SyncOptions(), CancellationToken.None);
-		ChangelistFixture cl = Fixture.StreamFooMain.GetChangelist(7);
-		ChangelistFixture clWithView = new (cl.Number, cl.Description, cl.StreamFiles.Where(x => x.ClientFile != "Data/data.txt").ToList());
-		clWithView.AssertDepotFiles(pm.SyncDir.FullName);
-		await pm.FinalizeAsync(CancellationToken.None);
-		
-		PerforceMaterializer.State? state = await pm.LoadStateForTestAsync(CancellationToken.None);
-		Assert.AreEqual(7, state!.Changelist);
+		Assert.AreEqual(expectedChangeNum, state!.Changelist);
 	}
 	
 	[TestMethod]
 	public async Task Conform_RemoveUntracked_Async()
 	{
-		using PerforceMaterializer pm = CreatePm(Fixture.StreamFooMain);
+		using PerforceMaterializer pm = CreatePm(PerforceSnapshot.FooMain);
 		await pm.SyncAsync(5, -1, new SyncOptions(), CancellationToken.None);
 		await File.WriteAllTextAsync(Path.Join(pm.SyncDir.FullName, _untrackedFile.clientFile), _untrackedFile.content);
 		
 		await pm.ConformAsync(true, CancellationToken.None);
-		Fixture.StreamFooMain.Head.AssertDepotFiles(pm.SyncDir.FullName);
+		PerforceSnapshot.GetSnapshot().GetStream(PerforceSnapshot.FooMain).GetWorkspace(9).AssertDepotFiles(pm.SyncDir.FullName);
 	}
 	
 	[TestMethod]
 	public async Task Conform_KeepUntracked_Async()
 	{
 		// Arrange
-		using PerforceMaterializer pm = CreatePm(Fixture.StreamFooMain);
+		using PerforceMaterializer pm = CreatePm(PerforceSnapshot.FooMain);
 		await pm.SyncAsync(5, -1, new SyncOptions(), CancellationToken.None);
 		await File.WriteAllTextAsync(Path.Join(pm.SyncDir.FullName, _untrackedFile.clientFile), _untrackedFile.content);
 		
@@ -161,9 +98,10 @@ public class PerforceMaterializerTests : BasePerforceFixtureTest
 		await pm.ConformAsync(false, CancellationToken.None);
 		
 		// Assert
-		(FileSet localFiles, FileSet expectedFiles) = Fixture.StreamFooMain.Head.GetFileSets(pm.SyncDir.FullName);
-		expectedFiles.Add((_untrackedFile.clientFile, _untrackedFile.size, _untrackedFile.digest));
-		ChangelistFixture.AssertFileSets(localFiles, expectedFiles);
+		HashSet<WorkspaceFile> expectedWorkspaceFiles = new (PerforceSnapshot.GetSnapshot().GetStream(PerforceSnapshot.FooMain).GetWorkspace(9).Files); // Copy set
+		HashSet<WorkspaceFile> actualLocalFiles = PerforceSnapshot.GetLocalFileSet(pm.SyncDir.FullName);
+		expectedWorkspaceFiles.Add(WorkspaceFile.Create(_untrackedFile.clientFile, _untrackedFile.content));
+		WorkspaceSnapshot.AssertFileSets(actualLocalFiles, expectedWorkspaceFiles);
 	}
 	
 	[TestMethod]
@@ -178,7 +116,7 @@ public class PerforceMaterializerTests : BasePerforceFixtureTest
 			UserName = PerforceConnection.UserName,
 			Password = PerforceConnection.Settings.Password,
 			Identifier = $"{TestGuid}-{identifier}",
-			Stream = Fixture.StreamFooMain.Root,
+			Stream = PerforceSnapshot.FooMain,
 			Method = "name=perforce"
 		};
 		
@@ -192,9 +130,6 @@ public class PerforceMaterializerTests : BasePerforceFixtureTest
 		await PerforceExecutor.ConformMaterializersAsync(factories, TempDir, [raw], removeUntrackedFiles, [], NoOpTracer, logger, CancellationToken.None);
 	}
 	
-	private const string Stream = "//Foo/Main";
-	private const string ClientName = "MyComputer";
-	
 	[TestMethod]
 	[DataRow("Platforms/...", "//Foo/Main/Platforms/... //MyComputer/Platforms/...")]
 	[DataRow(" Platforms/...", "//Foo/Main/Platforms/... //MyComputer/Platforms/...")]
@@ -204,30 +139,31 @@ public class PerforceMaterializerTests : BasePerforceFixtureTest
 	[DataRow("/Engine/Build/...", "//Foo/Main/Engine/Build/... //MyComputer/Engine/Build/...")]
 	public void ConvertHordeViews(string hordeView, string expected)
 	{
-		string actual = PerforceMaterializer.ConvertHordeViewToClientView(hordeView, Stream, ClientName);
+		string actual = PerforceMaterializer.ConvertHordeViewToClientView(hordeView, "//Foo/Main", "MyComputer");
 		Assert.AreEqual(expected, actual);
 	}
 
 	[TestMethod]
 	public async Task TestAllStatesAsync()
 	{
-		StreamState[] mainState = [Main(2), Main(4), Main(5), Main(7), Main(7, 8), Main(9)];
-		StreamState[] releaseStates = [Release(11), Release(12)];
-		mainState = mainState.Select(x => x with { RemoveUntracked = false }).ToArray();
-		releaseStates = releaseStates.Select(x => x with { RemoveUntracked = false }).ToArray();
+		List<WorkspaceState> mainStates = new(PerforceSnapshot.FooMainStates);
+		List<WorkspaceState> releaseStates = new(PerforceSnapshot.FooReleaseStates);
+		
+		mainStates = mainStates.Select(x => x with { RemoveUntracked = true }).ToList();
+		releaseStates = releaseStates.Select(x => x with { RemoveUntracked = true }).ToList();
 		
-		StreamState?[] localStates = [null, ..mainState, ..releaseStates];
-		StreamState[] targetStates = [..mainState, ..releaseStates];
+		WorkspaceState?[] localStates = [null, ..mainStates, ..releaseStates];
+		WorkspaceState[] targetStates = [..mainStates, ..releaseStates];
 
-		foreach (StreamState? localState in localStates)
+		foreach (WorkspaceState? localState in localStates)
 		{
 			// Server state can only be empty or match the local state for now
 			// This could theoretically differ, so checking as "serverStates = [..localStates]" would be more robust
 			// But requires more (slightly expensive) checks 
-			StreamState?[] serverStates = [null, localState];
-			foreach (StreamState? serverState in serverStates)
+			WorkspaceState?[] serverStates = [null, localState];
+			foreach (WorkspaceState? serverState in serverStates)
 			{
-				foreach (StreamState targetState in targetStates)
+				foreach (WorkspaceState targetState in targetStates)
 				{
 					await SyncWithScenarioAsync(localState, serverState, targetState, CancellationToken.None);
 				}
@@ -235,18 +171,9 @@ public class PerforceMaterializerTests : BasePerforceFixtureTest
 		}
 	}
 
-	private StreamState Main(int cl, int shelvedCl = -1) => new (Fixture.StreamFooMain, cl, shelvedCl);
-	private StreamState Release(int cl, int shelvedCl = -1) => new (Fixture.StreamFooRelease, cl, shelvedCl);
-	private record StreamState(StreamFixture Stream, int Cl, int ShelvedCl = -1, bool RemoveUntracked = true)
-	{
-		public override string ToString()
-		{
-			return $"Stream={Stream.Root} CL={Cl} Shelve={ShelvedCl}";
-		}
-	}
-	private record SyncScenario(StreamState? LocalState, StreamState? ServerState, StreamState TargetState);
+	private record SyncScenario(WorkspaceState? LocalState, WorkspaceState? ServerState, WorkspaceState TargetState);
 
-	private Task SyncWithScenarioAsync(StreamState? localState, StreamState? serverState, StreamState targetState, CancellationToken cancellationToken)
+	private Task SyncWithScenarioAsync(WorkspaceState? localState, WorkspaceState? serverState, WorkspaceState targetState, CancellationToken cancellationToken)
 	{
 		return SyncWithScenarioAsync(new SyncScenario(localState, serverState, targetState), cancellationToken);
 	}
@@ -296,20 +223,20 @@ public class PerforceMaterializerTests : BasePerforceFixtureTest
 		}
 	}
 
-	private static async Task SyncToStateAsync(PerforceMaterializer pm, StreamState state, CancellationToken cancellationToken)
+	private static async Task SyncToStateAsync(PerforceMaterializer pm, WorkspaceState state, CancellationToken cancellationToken)
 	{
 		await pm.SyncAsync(state.Cl, state.ShelvedCl, new SyncOptions { RemoveUntracked = state.RemoveUntracked }, cancellationToken);
-		state.Stream.GetChangelist(state.ShelvedCl > 0 ? state.ShelvedCl : state.Cl).AssertDepotFiles(pm.SyncDir.FullName);
+		PerforceSnapshot.GetSnapshot().GetStream(state.Stream).GetWorkspace(state.Cl, state.ShelvedCl).AssertDepotFiles(pm.SyncDir.FullName);
 		await pm.FinalizeAsync(cancellationToken);
 		
 		PerforceMaterializer.State? pmState = await pm.LoadStateForTestAsync(cancellationToken);
 		Assert.AreEqual(state.Cl, pmState!.Changelist);
 		Assert.AreEqual(state.ShelvedCl, pmState!.ShelvedChangelist);
-		Assert.AreEqual(state.Stream.Root, pmState.Stream);
+		Assert.AreEqual(state.Stream, pmState.Stream);
 		Assert.AreEqual(PerforceMaterializer.TransactionStatus.Clean, pmState.Status);
 	}
 
-	private PerforceMaterializer CreatePm(StreamFixture stream, string? dataDir = null, string? identifier = "pmTest", string[]? view = null)
+	private PerforceMaterializer CreatePm(string stream, string? dataDir = null, string? identifier = "pmTest", string[]? view = null)
 	{
 		ILogger logger = LoggerFactory.CreateLogger("PM");
 		RpcAgentWorkspace raw = new ()
@@ -318,11 +245,24 @@ public class PerforceMaterializerTests : BasePerforceFixtureTest
 			UserName = PerforceConnection.UserName,
 			Password = PerforceConnection.Settings.Password,
 			Identifier = $"{TestGuid}-{identifier}",
-			Stream = stream.Root,
+			Stream = stream,
 		};
 		raw.View.AddRange(view ?? []);
 		PerforceMaterializerOptions options = new(dataDir ?? TempDir.FullName, raw);
 		return new PerforceMaterializer(options, NoOpTracer, logger);
 	}
+	
+	private static async Task SubmitFileAsync(PerforceMaterializer pm, string clDescription, string filename, string fileContent)
+	{
+		string tempFile = Path.Join(pm.SyncDir.FullName, filename);
+		await File.WriteAllTextAsync(tempFile, fileContent);
+
+		IPerforceConnection perforce = await pm.GetPerforceWithClientAsync();
+		List<AddRecord> addRecords = await perforce.AddAsync(-1, tempFile);
+		ChangeRecord change = new () { Description = clDescription };
+		change.Files.Add(addRecords[0].DepotFile);
+		change = await perforce.CreateChangeAsync(change);
+		await perforce.SubmitAsync(change.Number, SubmitOptions.None);
+	}
 }
 
diff --git a/Engine/Source/Programs/Horde/Drivers/JobDriver/Execution/PerforceMaterializer.cs b/Engine/Source/Programs/Horde/Drivers/JobDriver/Execution/PerforceMaterializer.cs
--- a/Engine/Source/Programs/Horde/Drivers/JobDriver/Execution/PerforceMaterializer.cs
+++ b/Engine/Source/Programs/Horde/Drivers/JobDriver/Execution/PerforceMaterializer.cs
@@ -93,8 +93,11 @@ public sealed class PerforceMaterializer : IWorkspaceMaterializer
 	/// Name of this materializer
 	/// </summary>
 	public const string TypeName = "Perforce";
-	
+
+	internal const int CurrentStateVersion = 2;
 	internal enum TransactionStatus { Clean = 0, Dirty = 1 }
+
+	internal record WorkspaceFile(string ClientFile, string LocalFile);
 	internal record State(
 		int Version,
 		TransactionStatus Status,
@@ -102,20 +105,21 @@ public sealed class PerforceMaterializer : IWorkspaceMaterializer
 		string Identifier,
 		string Stream,
 		int Changelist,
-		int ShelvedChangelist);
+		int ShelvedChangelist,
+		List<WorkspaceFile> DirtyFiles);
 	
-    /// <inheritdoc/>
+	/// <inheritdoc/>
 	public DirectoryReference SyncDir { get; }
-    
+
 	/// <inheritdoc/>
 	public DirectoryReference BaseDir { get; }
 
-    /// <inheritdoc/>
+	/// <inheritdoc/>
 	public string Name => TypeName;
 
 	/// <inheritdoc/>
 	public string Identifier { get; }
-	
+
 	/// <inheritdoc/>
 	public IReadOnlyDictionary<string, string> EnvironmentVariables { get; private set; }
 	
@@ -133,6 +137,7 @@ public sealed class PerforceMaterializer : IWorkspaceMaterializer
 	private IPerforceConnection? _perforceWithClient;
 	private const string StateFilename = "State.json";
 	private const bool PreferNativePerforceClient = true;
+	private const int MaxFilesLogged = 100;
 	private readonly string _stateFile;
 	private readonly string _stateTempFile;
 
@@ -228,7 +233,12 @@ public sealed class PerforceMaterializer : IWorkspaceMaterializer
 		
 		State? state = await LoadStateAsync(cancellationToken);
 		
-		bool isDirty = state == null || state.Status == TransactionStatus.Dirty || state.Identifier != Identifier || state.Stream != _options.AgentWorkspace.Stream;
+		bool isDirty = state == null
+			|| state.Version != CurrentStateVersion
+			|| state.Status == TransactionStatus.Dirty
+			|| state.Identifier != Identifier
+			|| state.Stream != _options.AgentWorkspace.Stream;
+		
 		_logger.LogInformation("Status: {Status}\nClient: {PrevClient} -> {Client}\nIdentifier: {PrevId} -> {Id}\nStream: {PrevStream} -> {Stream}\nCL: {PrevCl} -> {Cl}\nShelve: {PrevShelve} -> {Shelve}\nDirty: {Dirty}",
 			state?.Status,
 			state?.Client, _perforceSettings.ClientName,
@@ -238,7 +248,6 @@ public sealed class PerforceMaterializer : IWorkspaceMaterializer
 			state?.ShelvedChangelist, shelveChangeNum,
 			isDirty);
 		
-		
 		if (isDirty)
 		{
 			using TelemetrySpan deleteSpan = _tracer.StartActiveSpan($"{nameof(PerforceMaterializer)}.{nameof(SyncAsync)}.Delete");
@@ -288,7 +297,7 @@ public sealed class PerforceMaterializer : IWorkspaceMaterializer
 			changeNum = changesRecords[0].Number;
 		}
 
-		await SaveStateAsync(TransactionStatus.Dirty, client.Name, changeNum, shelveChangeNum, cancellationToken);
+		await SaveStateAsync(TransactionStatus.Dirty, client.Name, changeNum, shelveChangeNum, [], cancellationToken);
 		
 		{
 			using TelemetrySpan syncSpan = _tracer.StartActiveSpan($"{nameof(PerforceMaterializer)}.{nameof(SyncAsync)}.Sync");
@@ -318,18 +327,12 @@ public sealed class PerforceMaterializer : IWorkspaceMaterializer
 
 		if (shelveChangeNum > 0)
 		{
-			using TelemetrySpan unshelveSpan = _tracer.StartActiveSpan($"{nameof(PerforceMaterializer)}.{nameof(SyncAsync)}.Unshelve");
-			using ILoggerProgress status = _logger.BeginProgressScope("Unshelving files...");
-			Stopwatch unshelveTimer = Stopwatch.StartNew();
-			
-			List<UnshelveRecord> unshelveRecords = await _perforceWithClient.UnshelveAsync(shelveChangeNum, -1, null, null, null, UnshelveOptions.ForceOverwrite, Array.Empty<string>(), cancellationToken);
-
-			unshelveSpan.SetAttribute("horde.pm.file_count", unshelveRecords.Count);
-			_logger.LogInformation("Unshelved {FileCount} files.", unshelveRecords.Count);
-			status.Progress = $"({unshelveTimer.Elapsed.TotalSeconds:0.0}s)";
+			await UnshelveAsync(_perforceWithClient, client, changeNum, shelveChangeNum, cancellationToken);
+		}
+		else
+		{
+			await SaveStateAsync(TransactionStatus.Clean, client.Name, changeNum, shelveChangeNum, [], cancellationToken);
 		}
-
-		await SaveStateAsync(TransactionStatus.Clean, client.Name, changeNum, shelveChangeNum, cancellationToken);
 	}
 	
 	/// <inheritdoc/>
@@ -370,6 +373,109 @@ public sealed class PerforceMaterializer : IWorkspaceMaterializer
 		await SyncAsync(IWorkspaceMaterializer.LatestChangeNumber, 0, syncOptions, cancellationToken);
 	}
 	
+		/// <summary>
+		/// Replays the effects of unshelving a changelist
+		/// But clobbering files in the workspace rather than actually unshelving them.
+		/// This to prevent problems with multiple machines locking them.
+		/// </summary>
+		public async Task UnshelveAsync(IPerforceConnection perforce, ClientRecord client, int changeNum, int shelveChangeNum, CancellationToken cancellationToken)
+		{
+			using TelemetrySpan span = _tracer.StartActiveSpan($"{nameof(PerforceMaterializer)}.{nameof(UnshelveAsync)}");
+			
+			// Need to mark those files as dirty - update the workspace with those files 
+			// Delete is fine, but need to flag anything added
+			Stopwatch timer = Stopwatch.StartNew();
+			_logger.LogInformation("Unshelving changelist {Change}...", shelveChangeNum);
+
+			// Query the contents of the shelved changelist
+			List<DescribeRecord> records = await perforce.DescribeAsync(DescribeOptions.Shelved, -1, [shelveChangeNum], cancellationToken);
+			if (records.Count != 1)
+			{
+				throw new WorkspaceMaterializationException($"Changelist {shelveChangeNum} is not shelved");
+			}
+			DescribeRecord lastRecord = records[0];
+			if (lastRecord.Files.Count == 0)
+			{
+				throw new WorkspaceMaterializationException($"Changelist {shelveChangeNum} does not contain any shelved files");
+			}
+
+			// Query the location of each file
+			List<PerforceResponse<WhereRecord>> whereResponseList = await perforce.TryWhereAsync(lastRecord.Files.Select(x => x.DepotFile).ToArray(), cancellationToken).ToListAsync(cancellationToken);
+			Dictionary<string, WhereRecord> whereRecords = whereResponseList.Where(x => x.Succeeded).Select(x => x.Data).ToDictionary(x => x.DepotFile, x => x, StringComparer.OrdinalIgnoreCase);
+
+			// Parse out all the list of deleted and modified files
+			List<WhereRecord> deleteFiles = [];
+			List<WhereRecord> writeFiles = [];
+			foreach (DescribeFileRecord fileRecord in lastRecord.Files)
+			{
+				if (!whereRecords.TryGetValue(fileRecord.DepotFile, out WhereRecord? whereRecord))
+				{
+					_logger.LogInformation("Unable to get location of {File} in current workspace; ignoring.", fileRecord.DepotFile);
+					continue;
+				}
+
+				switch (fileRecord.Action)
+				{
+					case FileAction.Delete:
+					case FileAction.MoveDelete:
+						deleteFiles.Add(whereRecord);
+						break;
+					case FileAction.Add:
+					case FileAction.Edit:
+					case FileAction.MoveAdd:
+					case FileAction.Branch:
+					case FileAction.Integrate:
+						writeFiles.Add(whereRecord);
+						break;
+					case FileAction.None:
+					case FileAction.Purge:
+					case FileAction.Archive:
+					case FileAction.Import:
+					case FileAction.Unknown:
+					default:
+						throw new Exception($"Unknown action '{fileRecord.Action}' for shelved file {fileRecord.DepotFile}");
+				}
+			}
+
+			List<WorkspaceFile> dirtyFiles = writeFiles
+				.Concat(deleteFiles)
+				.Select(x => new WorkspaceFile(x.ClientFile, x.Path))
+				.ToList();
+			
+			await SaveStateAsync(TransactionStatus.Dirty, client.Name, changeNum, shelveChangeNum, dirtyFiles, cancellationToken);
+
+			// Delete all the files
+			foreach (WhereRecord deleteFile in deleteFiles)
+			{
+				string localPath = deleteFile.Path;
+				if (File.Exists(localPath))
+				{
+					_logger.LogInformation("  Deleting {LocalPath}", localPath);
+					FileUtils.ForceDeleteFile(localPath);
+				}
+			}
+
+			// Use common paths with wild cards speed up the print operation with one call instead of many calls to print.
+			_logger.LogInformation("Writing files from shelved changelist {Change}", shelveChangeNum);
+			PerforceResponseList<PrintRecord> printResponse = await perforce.TryPrintAsync($"{SyncDir.FullName}{Path.DirectorySeparatorChar}...", $"//{perforce.Settings.ClientName}/...@={shelveChangeNum}", cancellationToken);
+			if (!printResponse.Succeeded)
+			{
+				_logger.LogWarning("Unable to print shelved changelist: {Error}", printResponse.ToString());
+			}
+
+			foreach (PerforceResponse<PrintRecord> pr in printResponse.Take(MaxFilesLogged))
+			{
+				_logger.LogInformation("Unshelved {Action} {File}", pr.Data.Action, pr.Data.DepotFile);
+			}
+			if (printResponse.Count > MaxFilesLogged)
+			{
+				_logger.LogInformation(".. with {MoreCount} more files", printResponse.Count - MaxFilesLogged);
+			}
+
+			_logger.LogInformation("Completed in {TimeSeconds}s", $"{timer.Elapsed.TotalSeconds:0.0}");
+			await SaveStateAsync(TransactionStatus.Clean, client.Name, changeNum, shelveChangeNum, dirtyFiles, cancellationToken);
+		}
+	
 	internal static async Task<IPerforceConnection> CreatePerforceWithoutClientAsync(RpcAgentWorkspace agentWorkspace, ILogger logger, CancellationToken cancellationToken)
 	{
 		RpcAgentWorkspace raw = agentWorkspace;
@@ -449,10 +555,40 @@ public sealed class PerforceMaterializer : IWorkspaceMaterializer
 
 	private async Task RevertInternalAsync(IPerforceConnection perforce, CancellationToken cancellationToken)
 	{
+		if (_perforceSettings == null || _perforceWithClient == null)
+		{
+			throw new WorkspaceMaterializationException("Perforce settings/client not initialized!");
+		}
+		
 		using TelemetrySpan span = _tracer.StartActiveSpan($"{nameof(PerforceMaterializer)}.{nameof(RevertInternalAsync)}");
 		using ILoggerProgress status = _logger.BeginProgressScope("Reverting files...");
 		Stopwatch timer = Stopwatch.StartNew();
 
+		State? state = await LoadStateAsync(cancellationToken);
+		if (state is { DirtyFiles.Count: > 0 })
+		{
+			foreach (WorkspaceFile wf in state.DirtyFiles)
+			{
+				FileUtils.ForceDeleteFile(wf.LocalFile);
+			}
+			
+			List<string> dirtyClientFiles = state.DirtyFiles.Select(x => $"{x.ClientFile}@{state.Changelist}").ToList();
+			List<SyncRecord> syncedFiles = await _perforceWithClient.SyncAsync(
+				EpicGames.Perforce.SyncOptions.Force, -1, _options.SyncThreads, _options.SyncBatch, _options.SyncBatchSize,
+				-1, -1, dirtyClientFiles, cancellationToken).ToListAsync(cancellationToken);
+			
+			foreach (SyncRecord sr in syncedFiles.Take(MaxFilesLogged))
+			{
+				_logger.LogInformation("Restored dirty file: {File}", sr.DepotFile);
+			}
+			if (syncedFiles.Count > MaxFilesLogged)
+			{
+				_logger.LogInformation(".. with {MoreCount} more files", syncedFiles.Count - MaxFilesLogged);
+			}
+			
+			await SaveStateAsync(TransactionStatus.Clean, state.Client, state.Changelist, state.ShelvedChangelist, [], cancellationToken);
+		}
+
 		List<RevertRecord> revertRecords = await perforce.RevertAsync(-1, null, RevertOptions.DeleteAddedFiles, "//...", cancellationToken);
 		_logger.LogInformation("Reverted {FileCount} files.", revertRecords.Count);
 		
@@ -540,9 +676,9 @@ public sealed class PerforceMaterializer : IWorkspaceMaterializer
 		return $"{exclude}{stream}{path} //{clientName}{path}";
 	}
 
-	private Task SaveStateAsync(TransactionStatus status, string client ,int changelist, int shelvedChangelist, CancellationToken cancellationToken)
+	private Task SaveStateAsync(TransactionStatus status, string client ,int changelist, int shelvedChangelist, List<WorkspaceFile> dirtyFiles, CancellationToken cancellationToken)
 	{
-		State pms = new (1, status, client, Identifier, _options.AgentWorkspace.Stream, changelist, shelvedChangelist);
+		State pms = new (CurrentStateVersion, status, client, Identifier, _options.AgentWorkspace.Stream, changelist, shelvedChangelist, dirtyFiles);
 		return SaveStateAsync(pms, cancellationToken);
 	}
 	
@@ -579,14 +715,6 @@ public sealed class PerforceMaterializer : IWorkspaceMaterializer
 		}
 	}
 
-	private void ResetState()
-	{
-		if (File.Exists(_stateFile))
-		{
-			File.Delete(_stateFile);
-		}
-	}
-
 	private static int GetDefaultThreadCount(int defaultValue) => Math.Min(defaultValue, Math.Max(Environment.ProcessorCount - 1, 1));
 	
 	internal async Task DeleteClientForTestAsync(CancellationToken cancellationToken)
diff --git a/Engine/Source/Programs/Shared/EpicGames.Perforce.Fixture/BasePerforceFixtureTest.cs b/Engine/Source/Programs/Shared/EpicGames.Perforce.Fixture/BasePerforceFixtureTest.cs
--- a/Engine/Source/Programs/Shared/EpicGames.Perforce.Fixture/BasePerforceFixtureTest.cs
+++ b/Engine/Source/Programs/Shared/EpicGames.Perforce.Fixture/BasePerforceFixtureTest.cs
@@ -17,7 +17,7 @@ namespace EpicGames.Perforce.Fixture;
 [TestClass]
 public abstract class BasePerforceFixtureTest : IDisposable
 {
-	private const string EnvVar = "EPICGAMES_P4_FIXTURE_SERVER_URL";
+	internal const string EnvVar = "EPICGAMES_P4_FIXTURE_SERVER_URL";
 	private const string UseDefaultMarker = "default";
 	private const string DefaultUrl = "p4://test.user@localhost:1666";
 
@@ -105,6 +105,11 @@ public abstract class BasePerforceFixtureTest : IDisposable
 	}
 
 	private PerforceConnection GetPerforceConnection(string connectionString)
+	{
+		return GetPerforceConnection(connectionString, _clientName, LoggerFactory.CreateLogger<PerforceConnection>());
+	}
+
+	internal static PerforceConnection GetPerforceConnection(string connectionString, string clientName, ILogger logger)
 	{
 		connectionString = connectionString == UseDefaultMarker ? DefaultUrl : connectionString;
 		Uri uri = new(connectionString);
@@ -117,7 +122,7 @@ public abstract class BasePerforceFixtureTest : IDisposable
 		PerforceSettings perforceSettings = new(serverAndPort, username)
 		{
 			AppName = "Perforce Test Fixture",
-			ClientName = _clientName,
+			ClientName = clientName,
 			PreferNativeClient = true
 		};
 
@@ -126,7 +131,7 @@ public abstract class BasePerforceFixtureTest : IDisposable
 			perforceSettings.Password = password;
 		}
 
-		return new PerforceConnection(perforceSettings, LoggerFactory.CreateLogger<PerforceConnection>());
+		return new PerforceConnection(perforceSettings, logger);
 	}
 
 	private static (string username, string? password) GetCredentialsFromUri(Uri uri)
diff --git a/Engine/Source/Programs/Shared/EpicGames.Perforce.Fixture/EpicGames.Perforce.Fixture.csproj b/Engine/Source/Programs/Shared/EpicGames.Perforce.Fixture/EpicGames.Perforce.Fixture.csproj
--- a/Engine/Source/Programs/Shared/EpicGames.Perforce.Fixture/EpicGames.Perforce.Fixture.csproj
+++ b/Engine/Source/Programs/Shared/EpicGames.Perforce.Fixture/EpicGames.Perforce.Fixture.csproj
@@ -16,6 +16,7 @@
     <IsPackable>false</IsPackable>
     <Nullable>enable</Nullable>
     <NoWarn>$(NoWarn);1591</NoWarn>
+    <StartupObject>EpicGames.Perforce.Fixture.Program</StartupObject>
   </PropertyGroup>
 
   <PropertyGroup Condition=" '$(Configuration)|$(Platform)' == 'Analyze|AnyCPU' ">
@@ -32,6 +33,10 @@
     </Compile>
   </ItemGroup>
 
+  <ItemGroup>
+    <EmbeddedResource Include="snapshot.json" />
+  </ItemGroup>
+
   <ItemGroup>
     <PackageReference Include="Microsoft.Extensions.Logging.Console" Version="8.0.1" />
     <PackageReference Include="Microsoft.NET.Test.Sdk" Version="17.14.1" />
