Skip RCA escalation on a weekend.

#Synced-CL 47213820

diff --git a/Engine/Source/Programs/Horde/Plugins/Build/HordeServer.Build/Notifications/Sinks/SlackNotificationSink.cs b/Engine/Source/Programs/Horde/Plugins/Build/HordeServer.Build/Notifications/Sinks/SlackNotificationSink.cs
--- a/Engine/Source/Programs/Horde/Plugins/Build/HordeServer.Build/Notifications/Sinks/SlackNotificationSink.cs
+++ b/Engine/Source/Programs/Horde/Plugins/Build/HordeServer.Build/Notifications/Sinks/SlackNotificationSink.cs
@@ -3153,13 +3153,20 @@ namespace HordeServer.Notifications.Sinks
 				return null;
 			}
 
-			string mention = await FormatMentionAsync(ownerId.Value, workflow.AllowMentions, cancellationToken);
+			if (DateTime.Now.DayOfWeek == DayOfWeek.Saturday || DateTime.Now.DayOfWeek == DayOfWeek.Sunday)
+			{
+				_logger.LogInformation("Skipping escalation message for {NotificationName} {IssueId} since it's a weekend", notificationName, issueId);
+			}
+			else
+			{
+				string mention = await FormatMentionAsync(ownerId.Value, workflow.AllowMentions, cancellationToken);
 
-			_logger.LogInformation("Sending escalation message for {NotificationName} {IssueId} {User} {Channel}", notificationName, issueId, mention, state.MessageId);
+				_logger.LogInformation("Sending escalation message for {NotificationName} {IssueId} {User} {Channel}", notificationName, issueId, mention, state.MessageId);
 
-			Uri issueUrl = GetIssueUrl(issue, span.FirstFailure);
-			string message = $"{mention} this root cause analysis requires your attention - Issue <{issueUrl}|{issue.Id}> - Select ...More > Modify Issue and choose a category to complete the root cause analysis.";
-			await _slackClient.PostMessageToThreadAsync(state.MessageId, message, cancellationToken);
+				Uri issueUrl = GetIssueUrl(issue, span.FirstFailure);
+				string message = $"{mention} this root cause analysis requires your attention - Issue <{issueUrl}|{issue.Id}> - Select ...More > Modify Issue and choose a category to complete the root cause analysis.";
+				await _slackClient.PostMessageToThreadAsync(state.MessageId, message, cancellationToken);
+			}
 
 			return GetNextEscalationTime(issue, issue.ResolvedAt.Value, utcNow, workflow.EscalateRootCauseAnalysisTimes, notificationName);
 		}
