[Horde][BuildHealth][Dashboard] Cleanup vuln issues.

#Synced-CL 48440999

diff --git a/Engine/Source/Programs/Horde/HordeDashboard/src/components/buildhealth/stepoutcome/StepOutcomeD3Component.tsx b/Engine/Source/Programs/Horde/HordeDashboard/src/components/buildhealth/stepoutcome/StepOutcomeD3Component.tsx
--- a/Engine/Source/Programs/Horde/HordeDashboard/src/components/buildhealth/stepoutcome/StepOutcomeD3Component.tsx
+++ b/Engine/Source/Programs/Horde/HordeDashboard/src/components/buildhealth/stepoutcome/StepOutcomeD3Component.tsx
@@ -10,7 +10,6 @@ import { getCellAccents, getCellClasses, getCellStyle, getColorRecords, Incremen
 import { getHordeTheme } from "horde/styles/theme";
 import { observer } from "mobx-react-lite";
 import { isTerminalState, toLocalISOString } from "./StepOutcomeUtilities";
-import { encodeStepNameFromStrings } from "../BuildHealthUtilities";
 
 // #region -- Styles --
 
@@ -181,11 +180,28 @@ function GenerateD3Table(stepOutcomeTable: StepOutcomeTable, handler: StepOutcom
         sectionHeaderDiv.style.display = "flex";
         sectionHeaderDiv.style.paddingTop = "10px";
         sectionHeaderDiv.style.justifyContent = "space-between";
-        sectionHeaderDiv.innerHTML = `<strong>${header}`;
+        const generateLabelTooltipRow = document.createElement("strong");
+        generateLabelTooltipRow.textContent = header;
+        sectionHeaderDiv.appendChild(generateLabelTooltipRow);
         sectionHeaderDiv.className = tooltipClasses.toolTipSection;
         tooltipContent.appendChild(sectionHeaderDiv);
     }
 
+    let generateTableTooltipRow = (tooltipContent: HTMLDivElement, label: string, value: string) => {
+        const tooltipRowElement = document.createElement("div");
+
+        const tooltipLabelElement = document.createElement("strong");
+        tooltipLabelElement.textContent = `${label}:`;
+
+        const tooltipValueElement = document.createElement("p");
+        tooltipValueElement.className = tooltipClasses.tooltipLine;
+        tooltipValueElement.textContent = value;
+
+        tooltipRowElement.appendChild(tooltipLabelElement);
+        tooltipRowElement.appendChild(tooltipValueElement);
+        tooltipContent.appendChild(tooltipRowElement);
+    }
+
     let computePassRatio = (entry: SummaryTableEntry) => {
         let ratio: number = -1;
 
@@ -357,11 +373,7 @@ function GenerateD3Table(stepOutcomeTable: StepOutcomeTable, handler: StepOutcom
 
                 const tooltipContent = document.createElement("div");
                 tooltipContent.style.paddingLeft = "4px";
-
-                const refreshRateDiv = document.createElement("div");
-                refreshRateDiv.innerHTML = `<strong>Refresh Rate:</strong><p class="${tooltipClasses.tooltipLine}">${formatted}</p><p>Click to refresh.</p>`;
-                tooltipContent.appendChild(refreshRateDiv);
-
+                generateTableTooltipRow(tooltipContent, "Refresh Rate", formatted);
                 attachTargettedToolTip(this as HTMLElement, scaffold.current!.tooltip, tooltipContent);
             })
             .on("mouseout", function () {
@@ -419,20 +431,9 @@ function GenerateD3Table(stepOutcomeTable: StepOutcomeTable, handler: StepOutcom
                     const changeHeader = { change: entry.change ?? "-", date: entry.date ?? "-" };
                     const original = new Date(changeHeader.date!);
 
-                    // UE-324856 - This is a temporary fix until the horde server is fixed.
-                    const corrected = new Date(original.getTime() - 4 * 60 * 60 * 1000);
-
-                    const changeDiv = document.createElement("div");
-                    changeDiv.innerHTML = `<strong>Change:</strong><p class="${tooltipClasses.tooltipLine}">${changeHeader.change}</p>`;
-                    tooltipContent.appendChild(changeDiv);
-
-                    const dateUTCDiv = document.createElement("div");
-                    dateUTCDiv.innerHTML = `<strong>Date (UTC):</strong><p class="${tooltipClasses.tooltipLine}">${corrected.toISOString()}</p>`;
-                    tooltipContent.appendChild(dateUTCDiv);
-
-                    const dateLocalDiv = document.createElement("div");
-                    dateLocalDiv.innerHTML = `<strong>Date (Local):</strong><p class="${tooltipClasses.tooltipLine}">${toLocalISOString(corrected)}</p>`;
-                    tooltipContent.appendChild(dateLocalDiv);
+                    generateTableTooltipRow(tooltipContent, "Change", changeHeader.change.toString());
+                    generateTableTooltipRow(tooltipContent, "Date (UTC)", original.toISOString());
+                    generateTableTooltipRow(tooltipContent, "Date (Local)", toLocalISOString(original));
 
                     attachTargettedToolTip(this as HTMLElement, scaffold.current!.tooltip, tooltipContent);
                 }
@@ -507,9 +508,7 @@ function GenerateD3Table(stepOutcomeTable: StepOutcomeTable, handler: StepOutcom
                     tooltipContent.style.paddingLeft = "4px";
 
                     if (stepNameHeader) {
-                        const streamDiv = document.createElement("div");
-                        streamDiv.innerHTML = `<strong>Stream:</strong><p class="${tooltipClasses.tooltipLine}">${stepNameHeader.streamId}</p>`;
-                        tooltipContent.appendChild(streamDiv);
+                        generateTableTooltipRow(tooltipContent, "Stream", stepNameHeader.streamId);
                     }
                     attachDynamicPositionToolTip(event, scaffold.current!.tooltip, tooltipContent);
                 });
@@ -554,7 +553,16 @@ function GenerateD3Table(stepOutcomeTable: StepOutcomeTable, handler: StepOutcom
                     exit => exit.remove()
                 );
 
-            cells.html(d => d.isLabel ? `<span style="background-color: ${theme.semanticColors.bodyBackgroundHovered};">${d.stepName}</span>` : `<span>${d.stepName}</span>`)
+            cells.each(function (d) {
+                const cell = d3.select(this);
+                cell.html("");
+
+                const span = cell.append('span').text(d.stepName);
+
+                if (d.isLabel) {
+                    span.style('background-color', theme.semanticColors.bodyBackgroundHovered);
+                }
+            })
                 .style("position", "sticky")
                 .style("left", "32px")
                 .style("z-index", "1")
@@ -573,20 +581,9 @@ function GenerateD3Table(stepOutcomeTable: StepOutcomeTable, handler: StepOutcom
                     tooltipContent.style.paddingLeft = "4px";
 
                     if (castedStepNameHeader) {
-                        const stepDiv = document.createElement("div");
-
-                        if (castedStepNameHeader.isLabel) {
-                            stepDiv.innerHTML = `<strong>Label:</strong><p class="${tooltipClasses.tooltipLine}">${castedStepNameHeader.stepName}</p>`;
-                        }
-                        else {
-                            stepDiv.innerHTML = `<strong>Step:</strong><p class="${tooltipClasses.tooltipLine}">${castedStepNameHeader.stepName}</p>`;
-                        }
-
-                        tooltipContent.appendChild(stepDiv);
-
-                        const streamDiv = document.createElement("div");
-                        streamDiv.innerHTML = `<strong>Stream:</strong><p class="${tooltipClasses.tooltipLine}">${castedStepNameHeader.streamId}</p>`;
-                        tooltipContent.appendChild(streamDiv);
+                        let calculatedLabel: string = castedStepNameHeader.isLabel ? "Label" : "Step";
+                        generateTableTooltipRow(tooltipContent, calculatedLabel, castedStepNameHeader.stepName);
+                        generateTableTooltipRow(tooltipContent, "Stream", castedStepNameHeader.streamId);
                     }
 
                     attachTargettedToolTip(this as HTMLElement, scaffold.current!.tooltip, tooltipContent);
@@ -630,53 +627,27 @@ function GenerateD3Table(stepOutcomeTable: StepOutcomeTable, handler: StepOutcom
                                     const tooltipContent = document.createElement("div");
                                     tooltipContent.style.paddingLeft = "4px";
 
-                                    const changeDiv = document.createElement("div");
-                                    changeDiv.innerHTML = `<strong>Change:</strong><p class="${tooltipClasses.tooltipLine}">${entry.change ?? "-"}</p>`;
-                                    tooltipContent.appendChild(changeDiv);
-
-                                    const dateDiv = document.createElement("div");
-                                    dateDiv.innerHTML = `<strong>Date:</strong><p class="${tooltipClasses.tooltipLine}">(UTC) ${original.toISOString() ?? "-"}</p>`;
-                                    dateDiv.innerHTML += `<p class="${tooltipClasses.tooltipLineWithBorder}">(Local) ${toLocalISOString(original) ?? "-"}</p>`
-                                    tooltipContent.appendChild(dateDiv);
+                                    generateTableTooltipRow(tooltipContent, "Change", entry.change !== undefined ? entry.change.toString() : "-");
+                                    generateTableTooltipRow(tooltipContent, "Date", `(UTC) ${original.toISOString() ?? "-"}`);
 
                                     if (isStepOutcome(entry)) {
-                                        const durationDiv = document.createElement("div");
-                                        durationDiv.innerHTML = `<strong>Duration:</strong><p class="${tooltipClasses.tooltipLine}">${entry.getDurationString?.() ?? "-"}</p>`;
-                                        tooltipContent.appendChild(durationDiv);
+                                        generateTableTooltipRow(tooltipContent, "Duration", `{${entry.getDurationString?.() ?? "-"}`);
                                     }
 
-                                    const stateDiv = document.createElement("div");
-                                    stateDiv.innerHTML = `<strong>State:</strong><p class="${tooltipClasses.tooltipLine}">${entry?.state ?? "-"}</p>`;
-                                    tooltipContent.appendChild(stateDiv);
+                                    generateTableTooltipRow(tooltipContent, "State", `${entry?.state ?? "-"}`);
 
                                     if (isTerminalState(entry)) {
-                                        const outcomeDiv = document.createElement("div");
-                                        outcomeDiv.innerHTML = `<strong>Outcome:</strong><p class="${tooltipClasses.tooltipLine}">${entry?.outcome ?? "-"}</p>`;
-
                                         // For simplicity, we will compeltely override the outcome of this step with the job's outcome, as the step was aborted.
-                                        if (entry.state === AbortedStr) {
-                                            outcomeDiv.innerHTML = `<strong>Job Outcome:</strong><p class="${tooltipClasses.tooltipLine}">${handler.getJobResponseData(entry.jobId)?.summarizedOutcome ?? "-"}</p>`;
-                                        }
-
-                                        tooltipContent.appendChild(outcomeDiv);
+                                        let calculatedLabel: string = entry.state === AbortedStr ? "Job Outcome" : "Outcome";
+                                        let calculatedValue: string = entry.state === AbortedStr ? `${handler.getJobResponseData(entry.jobId)?.summarizedOutcome ?? "-"}` : `${entry?.outcome ?? "-"}`;
+                                        generateTableTooltipRow(tooltipContent, calculatedLabel, calculatedValue);
                                     }
 
-                                    const streamDiv = document.createElement("div");
-                                    streamDiv.innerHTML = `<strong>Stream:</strong><p class="${tooltipClasses.tooltipLine}">${entry.streamId}</p>`;
-                                    tooltipContent.appendChild(streamDiv);
-
-                                    const jobDiv = document.createElement("div");
-                                    jobDiv.innerHTML = `<strong>Job:</strong><p class="${tooltipClasses.tooltipLine}">${entry.jobName ?? "-"}</p>`;
-                                    tooltipContent.appendChild(jobDiv);
+                                    generateTableTooltipRow(tooltipContent, "Stream", entry.streamId);
+                                    generateTableTooltipRow(tooltipContent, "Job", `${entry.jobName ?? "-"}`);
 
-                                    const stepDiv = document.createElement("div");
-                                    if (isLabelOutcome(entry)) {
-                                        stepDiv.innerHTML = `<strong>Label:</strong><p class="${tooltipClasses.tooltipLine}">${entry?.name ?? "-"}</p>`;
-                                    }
-                                    else {
-                                        stepDiv.innerHTML = `<strong>Step:</strong><p class="${tooltipClasses.tooltipLine}">${entry?.name ?? "-"}</p>`;
-                                    }
-                                    tooltipContent.appendChild(stepDiv);
+                                    let calculatedLabel: string = isLabelOutcome(entry) ? "Label" : "Step";
+                                    generateTableTooltipRow(tooltipContent, calculatedLabel, `${entry?.name ?? "-"}`);
 
                                     attachTargettedToolTip(this as HTMLElement, scaffold.current!.tooltip, tooltipContent);
                                 })
@@ -740,7 +711,7 @@ function GenerateD3Table(stepOutcomeTable: StepOutcomeTable, handler: StepOutcom
             const cells: TableEntry[] = cellsRaw ? [...cellsRaw] : [];
 
             const dataCells = tr.selectAll("td.step-outcome-data")
-                .data(cells, (d : TableEntry, i) => isSummary(d) ? `%summary-${i}`: `${d?.jobId}-${d?.id}`)
+                .data(cells, (d: TableEntry, i) => isSummary(d) ? `%summary-${i}` : `${d?.jobId}-${d?.id}`)
                 .join("td")
                 .attr("class", "step-outcome-data");
 
