Horde: Update release notes

#rnx

diff --git a/Engine/Source/Programs/Horde/Docs/ReleaseNotes.md b/Engine/Source/Programs/Horde/Docs/ReleaseNotes.md
--- a/Engine/Source/Programs/Horde/Docs/ReleaseNotes.md
+++ b/Engine/Source/Programs/Horde/Docs/ReleaseNotes.md
@@ -1,5 +1,37 @@
 # Release Notes
 
+## 2025-11-25
+
+* Perf & tests for commit regex (48627869)
+* Skip RCA escalation on a weekend. (48622101)
+* Use p4 print for unshelving in PerforceMaterializer (48607376)
+* EpicGames.Perforce: Add support for snapshot-based testing (48607305)
+* Removal of old JobStepRef "skipped state" emission code in entirety in preparation for V2. (48573938)
+* Add schedule filter to summary pane (48564369)
+* Adding Epic Sandbox project to soluthon and improving conditional (48563137)
+* EpicGames.Perforce: Improve fixture with deletes from shelve (48555287)
+* Improve VSCode support (48552892)
+* Null check for defensive programming. (48499986)
+* Modify RCA escalation message to inform user that only category is required (48482600)
+* Implement feedback from older swarm: code cleanliness (48479865)
+* Support JobStepId tags in log renders (48466761)
+* Timeout P4 issue rank queries after 2 minutes (48465686)
+* Only require RCA category, don't escalate for missing summary (48459016)
+* JobStepRef metadata to be emitted under more circumstances & ensure backwards compatibility on old result set. (48363489)
+* Initial work centralizing some Horde environment variables and dashboard routes used by various C# code (48352701)
+* Fix conform issue where mixing materializers caused deleted workspace dirs (48349168)
+* Adding additional style icons. (48347457)
+* Linked mode QoL features:
+  - Added link action button to streams
+  - Cascade link action button to the entire hierarchy (stream -> template -> step & label
+  - Edge case saitated where when a stream is added, implicitly link all the hierarhcy
+  - Only on the affirmative case, we do not issue this for the removal case (otherwise we would remove all hierarchy of the newly added stream)
+  - Optimizations made to only issue hierarchy refresh when receipts have diffs (48347284)
+* Toggle to filter out cancelled jobs from the toplevel history, allow for filtering of jobs based on whether it was scheduled or manually kicked, and a big fix for stale data bindings on D3. (48297767)
+* Commit telemetry integration into the p4 cache service. (48296707)
+* Improve issue matching across job steps in no suspects case (48293032)
+* Use rootCauseOwnerInfo exclusively (48290903)
+
 ## 2025-11-14
 
 * Fix compiler warning in IssueServiceTests (48281019)
