[Horde][Analytics][Converting commit telemetry stream to prod

#Synced-CL 48681143

diff --git a/Engine/Source/Programs/Horde/Plugins/Build/HordeServer.Build/VersionControl/Perforce/CommitTelemetry.cs b/Engine/Source/Programs/Horde/Plugins/Build/HordeServer.Build/VersionControl/Perforce/CommitTelemetry.cs
--- a/Engine/Source/Programs/Horde/Plugins/Build/HordeServer.Build/VersionControl/Perforce/CommitTelemetry.cs
+++ b/Engine/Source/Programs/Horde/Plugins/Build/HordeServer.Build/VersionControl/Perforce/CommitTelemetry.cs
@@ -14,7 +14,7 @@ namespace HordeServer.VersionControl.Perforce
 		/// <summary>
 		/// Default event name for the CommitTelemetry.
 		/// </summary>
-		public static readonly string DefaultEventName = "State.Change.Dev";
+		public static readonly string DefaultEventName = "State.Change";
 
 		/// <summary>
 		/// Current schema version.
