[Horde][Analytics] Align submitted naming with existing schema

#Synced-CL 48681143

diff --git a/Engine/Source/Programs/Horde/Plugins/Build/HordeServer.Build/VersionControl/Perforce/CommitTelemetry.cs b/Engine/Source/Programs/Horde/Plugins/Build/HordeServer.Build/VersionControl/Perforce/CommitTelemetry.cs
--- a/Engine/Source/Programs/Horde/Plugins/Build/HordeServer.Build/VersionControl/Perforce/CommitTelemetry.cs
+++ b/Engine/Source/Programs/Horde/Plugins/Build/HordeServer.Build/VersionControl/Perforce/CommitTelemetry.cs
@@ -9,7 +9,7 @@ namespace HordeServer.VersionControl.Perforce
 	/// Record used for Commit telemetry. Models a <see cref="ICommit"/>.
 	/// </summary>
 	/// <remarks>Modifications to this record's structure should be versioned appropriately through <see cref="SchemaVersion"/> via <see cref="CurrentSchemaVersion"/>, and call sites should continue to publish in a backwards compatible manner.</remarks>
-	public record CommitTelemetry(string EventName, int SchemaVersion, StreamId StreamId, int Change, int OriginalChange, DateTime SubmitDate, string? ExternalIssueKey, string? RobomergeAuthorTag, string? RobomergeSourceTag, string? RobomergeOwnerTag, string? RobomergeBotTag, string? ReviewTag)
+	public record CommitTelemetry(string EventName, int SchemaVersion, StreamId StreamId, int Change, int OriginalChange, DateTime SubmittedAt, string? ExternalIssueKey, string? RobomergeAuthorTag, string? RobomergeSourceTag, string? RobomergeOwnerTag, string? RobomergeBotTag, string? ReviewTag)
 	{
 		/// <summary>
 		/// Default event name for the CommitTelemetry.
