[Horde][BuildHealth][Dashboard] Duration label cleanup.

#Synced-CL 48788183

diff --git a/Engine/Source/Programs/Horde/HordeDashboard/src/components/buildhealth/stepoutcome/StepOutcomeD3Component.tsx b/Engine/Source/Programs/Horde/HordeDashboard/src/components/buildhealth/stepoutcome/StepOutcomeD3Component.tsx
--- a/Engine/Source/Programs/Horde/HordeDashboard/src/components/buildhealth/stepoutcome/StepOutcomeD3Component.tsx
+++ b/Engine/Source/Programs/Horde/HordeDashboard/src/components/buildhealth/stepoutcome/StepOutcomeD3Component.tsx
@@ -631,7 +631,7 @@ function GenerateD3Table(stepOutcomeTable: StepOutcomeTable, handler: StepOutcom
                                     generateTableTooltipRow(tooltipContent, "Date", `(UTC) ${original.toISOString() ?? "-"}`);
 
                                     if (isStepOutcome(entry)) {
-                                        generateTableTooltipRow(tooltipContent, "Duration", `{${entry.getDurationString?.() ?? "-"}`);
+                                        generateTableTooltipRow(tooltipContent, "Duration", `${entry.getDurationString?.() ?? "-"}`);
                                     }
 
                                     generateTableTooltipRow(tooltipContent, "State", `${entry?.state ?? "-"}`);
