Horde: Store list of files for replicated Perforce changelists

#Synced-CL Unknown

diff --git a/Engine/Source/Programs/Horde/Plugins/Build/HordeServer.Build.Tests/Perforce/PerforceLoadBalancerTests.cs b/Engine/Source/Programs/Horde/Plugins/Build/HordeServer.Build.Tests/Perforce/PerforceLoadBalancerTests.cs
--- a/Engine/Source/Programs/Horde/Plugins/Build/HordeServer.Build.Tests/Perforce/PerforceLoadBalancerTests.cs
+++ b/Engine/Source/Programs/Horde/Plugins/Build/HordeServer.Build.Tests/Perforce/PerforceLoadBalancerTests.cs
@@ -31,6 +31,15 @@ public class FakeHealthMonitor<T> : IHealthMonitor<T>
 	}
 }
 
+internal class HttpMessageHandlerStub(HttpStatusCode statusCode, string content) : HttpMessageHandler
+{
+	public string Content { get; set; } = content;
+	protected override Task<HttpResponseMessage> SendAsync(HttpRequestMessage request, CancellationToken cancellationToken)
+	{
+		return Task.FromResult(new HttpResponseMessage(statusCode) { Content = new StringContent(Content) });
+	}
+}
+
 [TestClass]
 public class PerforceLoadBalancerTests : BuildTestSetup
 {
@@ -105,7 +114,7 @@ public class PerforceLoadBalancerTests : BuildTestSetup
 		Assert.AreEqual(Server2, servers[0].ServerAndPort);
 	}
 
-	private static GlobalConfig GetConfig(params string[] perforceHosts)
+	internal static GlobalConfig GetConfig(params string[] perforceHosts)
 	{
 		PerforceCluster perforceCluster = CreateCluster(perforceHosts);
 		GlobalConfig gc = new();
@@ -147,18 +156,9 @@ public class PerforceLoadBalancerTests : BuildTestSetup
 	private PerforceLoadBalancer Create(GlobalConfig gc, string httpCheckResponse)
 	{
 #pragma warning disable CA2000 // Dispose objects before losing scope
-		HttpClient httpClient = new(new StubMessageHandler(HttpStatusCode.OK, httpCheckResponse));
+		HttpClient httpClient = new(new HttpMessageHandlerStub(HttpStatusCode.OK, httpCheckResponse));
 		return new(MongoService, GetRedisServiceSingleton(), LeaseCollection, Clock, httpClient, new TestOptionsMonitor<BuildConfig>(gc.Plugins.GetBuildConfig()), new FakeHealthMonitor<PerforceLoadBalancer>(), Tracer, _logger);
 #pragma warning restore CA2000 // Dispose objects before losing scope		
 	}
-
-	private class StubMessageHandler(HttpStatusCode statusCode, string content) : HttpMessageHandler
-	{
-		public string Content { get; set; } = content;
-		protected override Task<HttpResponseMessage> SendAsync(HttpRequestMessage request, CancellationToken cancellationToken)
-		{
-			return Task.FromResult(new HttpResponseMessage(statusCode) { Content = new StringContent(Content) });
-		}
-	}
 }
 
diff --git a/Engine/Source/Programs/Horde/Plugins/Build/HordeServer.Build.Tests/Perforce/PerforceServiceCacheTests.cs b/Engine/Source/Programs/Horde/Plugins/Build/HordeServer.Build.Tests/Perforce/PerforceServiceCacheTests.cs
new file mode 100644
--- /dev/null
+++ b/Engine/Source/Programs/Horde/Plugins/Build/HordeServer.Build.Tests/Perforce/PerforceServiceCacheTests.cs
@@ -0,0 +1,76 @@
+// Copyright Epic Games, Inc. All Rights Reserved.
+
+using System.Net;
+using EpicGames.Horde.Commits;
+using EpicGames.Horde.Streams;
+using EpicGames.Horde.Telemetry;
+using EpicGames.Perforce;
+using HordeServer.Acls;
+using HordeServer.Plugins;
+using HordeServer.Projects;
+using HordeServer.VersionControl.Perforce;
+using HordeServer.Server;
+using HordeServer.Streams;
+using HordeServer.Telemetry;
+using Microsoft.Extensions.Logging;
+
+namespace HordeServer.Tests.Perforce;
+
+public class TelemetryWriterStub : ITelemetryWriter
+{
+	public bool Enabled { get; } = false;
+	public void WriteEvent(TelemetryStoreId telemetryStoreId, object payload) => throw new NotImplementedException();
+	public void WriteEvent(TelemetryStoreId telemetryStoreId, TelemetryRecordMeta metadata, object payload) => throw new NotImplementedException();
+}
+
+[TestClass]
+public class PerforceServiceCacheTests : BuildTestSetup
+{
+	private readonly ILogger<PerforceLoadBalancer> _logger = CreateConsoleLogger<PerforceLoadBalancer>();
+
+	[TestMethod]
+	[Ignore] // Requires configuring the Perforce fixture server
+	public async Task CacheCommit_GetFiles_Async()
+	{
+		ILogger<PerforceServiceCache> pscLogger = CreateConsoleLogger<PerforceServiceCache>();
+		BuildServerConfig bsc = new ();
+		TestOptions<BuildServerConfig> bscWrapper = new (bsc);
+
+		List<string> servers = ["localhost"];
+		StreamConfig sc = new() { Id = new StreamId("FooMain"), Name = "//Foo/Main", ClusterName = "MyCluster" };
+		PerforceCluster perforceCluster = new() { Name = sc.ClusterName, Servers = servers.Select(x => new PerforceServer { ServerAndPort = x }).ToList() };
+		
+		BuildConfig buildConfig = new ();
+		buildConfig.Projects.Add(new ProjectConfig { Streams = [sc]});
+		buildConfig.PerforceClusters = [perforceCluster];
+
+		GlobalConfig gc = new();
+		gc.Plugins.AddBuildConfig(buildConfig);
+		gc.Plugins.AddComputeConfig(new ComputeConfig());
+		gc.PostLoad(new ServerSettings(), Array.Empty<ILoadedPlugin>(), Array.Empty<IDefaultAclModifier>());
+		TestOptionsMonitor<BuildConfig> bcWrapper = new (buildConfig);
+		
+		PerforceLoadBalancer plb = CreatePlb(gc, "notUsed");
+		await using PerforceServiceCache psc = new (
+			plb, MongoService, GetRedisServiceSingleton(), DowntimeService, UserCollection, Clock,
+			bscWrapper, bcWrapper, new TelemetryWriterStub(), Tracer, pscLogger);
+		
+		ICommitCollection cc = psc.GetCommits(sc);
+
+		PerforceServiceCache.StreamInfo streamInfo = new (sc, new PerforceViewMap([new PerforceViewMapEntry(true, sc.Name + "/...", "...")]), new PerforceChangeView());
+		await psc.UpdateClusterInternalAsync(sc.ClusterName, [streamInfo], null, CancellationToken.None);
+
+		ICommit commit = await cc.GetAsync(CommitIdWithOrder.FromPerforceChange(2));
+		IReadOnlyList<string> files = await commit.GetFilesAsync(null, null, CancellationToken.None);
+		CollectionAssert.AreEquivalent(new List<string> { "Data/data.txt", "common.h", "main.cpp", "main.h", "unused.cpp" }, files.ToList());
+	}
+	
+	private PerforceLoadBalancer CreatePlb(GlobalConfig gc, string httpCheckResponse)
+	{
+#pragma warning disable CA2000 // Dispose objects before losing scope
+		HttpClient httpClient = new(new HttpMessageHandlerStub(HttpStatusCode.OK, httpCheckResponse));
+		return new(MongoService, GetRedisServiceSingleton(), LeaseCollection, Clock, httpClient, new TestOptionsMonitor<BuildConfig>(gc.Plugins.GetBuildConfig()), new FakeHealthMonitor<PerforceLoadBalancer>(), Tracer, _logger);
+#pragma warning restore CA2000 // Dispose objects before losing scope
+	}
+}
+
diff --git a/Engine/Source/Programs/Horde/Plugins/Build/HordeServer.Build/VersionControl/Perforce/PerforceService.cs b/Engine/Source/Programs/Horde/Plugins/Build/HordeServer.Build/VersionControl/Perforce/PerforceService.cs
--- a/Engine/Source/Programs/Horde/Plugins/Build/HordeServer.Build/VersionControl/Perforce/PerforceService.cs
+++ b/Engine/Source/Programs/Horde/Plugins/Build/HordeServer.Build/VersionControl/Perforce/PerforceService.cs
@@ -32,6 +32,12 @@ namespace HordeServer.VersionControl.Perforce
 	/// </summary>
 	class PerforceService : IPerforceService, IAsyncDisposable
 	{
+		/// <summary>
+		/// Max number of files in a changelist to query for, or store
+		/// This is to prevent performance issues with large merge commits with lots of files.
+		/// </summary>
+		protected const int MaxFiles = 1000;
+		
 		protected sealed class PooledConnection : IDisposable
 		{
 			public IPerforceConnection Perforce { get; }
@@ -785,7 +791,6 @@ namespace HordeServer.VersionControl.Perforce
 
 			using (PooledConnectionHandle perforce = await ConnectAsync(cluster, null, cancellationToken))
 			{
-				const int MaxFiles = 1000;
 				DescribeRecord describeRecord = await perforce.DescribeAsync(DescribeOptions.None, MaxFiles, changeNumber, cancellationToken);
 
 				InfoRecord info = await perforce.GetInfoAsync(cancellationToken);
@@ -1274,8 +1279,6 @@ namespace HordeServer.VersionControl.Perforce
 								break;
 							}
 
-							const int MaxFiles = 1000;
-
 							List<DescribeRecord> describeRecords = await perforce.DescribeAsync(DescribeOptions.None, MaxFiles, changesRecords.Select(x => x.Number).ToArray(), cancellationToken);
 							foreach (DescribeRecord describeRecord in describeRecords)
 							{
diff --git a/Engine/Source/Programs/Horde/Plugins/Build/HordeServer.Build/VersionControl/Perforce/PerforceServiceCache.cs b/Engine/Source/Programs/Horde/Plugins/Build/HordeServer.Build/VersionControl/Perforce/PerforceServiceCache.cs
--- a/Engine/Source/Programs/Horde/Plugins/Build/HordeServer.Build/VersionControl/Perforce/PerforceServiceCache.cs
+++ b/Engine/Source/Programs/Horde/Plugins/Build/HordeServer.Build/VersionControl/Perforce/PerforceServiceCache.cs
@@ -48,7 +48,7 @@ namespace HordeServer.VersionControl.Perforce
 			public Dictionary<StreamId, int> MinChanges { get; set; } = new Dictionary<StreamId, int>();
 		}
 
-		class CommitTagInfo
+		internal class CommitTagInfo
 		{
 			public CommitTag Name { get; }
 			public FileFilter Filter { get; }
@@ -60,7 +60,7 @@ namespace HordeServer.VersionControl.Perforce
 			}
 		}
 
-		class StreamInfo
+		internal class StreamInfo
 		{
 			public StreamConfig StreamConfig { get; set; }
 			public PerforceViewMap View { get; }
@@ -151,7 +151,11 @@ namespace HordeServer.VersionControl.Perforce
 			public string BasePath { get; set; }
 			public DateTime DateUtc { get; set; }
 
-			public List<CommitTag> CommitTags { get; set; } = new List<CommitTag>();
+			public List<CommitTag> CommitTags { get; set; } = [];
+			public List<string> Files { get; set; } = [];
+			
+			[BsonElement("haf")]
+			public bool HasAllFiles { get; set; } = true;
 
 			[BsonConstructor]
 			CachedCommitDoc()
@@ -159,10 +163,11 @@ namespace HordeServer.VersionControl.Perforce
 				Description = null!;
 				BasePath = null!;
 
-				CommitTags = new List<CommitTag>();
+				CommitTags = [];
+				Files = [];
 			}
 
-			public CachedCommitDoc(ICommit commit, List<CommitTag> commitTags)
+			public CachedCommitDoc(ICommit commit, List<CommitTag> commitTags, List<string> files, bool hasAllFiles)
 			{
 				StreamId = commit.StreamId;
 				Number = commit.Id.GetPerforceChange();
@@ -181,12 +186,16 @@ namespace HordeServer.VersionControl.Perforce
 				DateUtc = commit.DateUtc;
 
 				CommitTags = commitTags;
+				Files = files;
+				HasAllFiles = hasAllFiles;
 			}
 
 			public static async Task<CachedCommitDoc> FromCommitAsync(ICommit commit, CancellationToken cancellationToken)
 			{
-				List<CommitTag> commitTags = new List<CommitTag>(await commit.GetTagsAsync(cancellationToken));
-				return new CachedCommitDoc(commit, commitTags);
+				List<CommitTag> commitTags = new (await commit.GetTagsAsync(cancellationToken));
+				List<string> files = new (await commit.GetFilesAsync(null, MaxFiles, cancellationToken));
+				bool hasAllFiles = files.Count < MaxFiles;
+				return new CachedCommitDoc(commit, commitTags, files, hasAllFiles);
 			}
 
 			public void PostLoad(PerforceServiceCache owner, StreamConfig streamConfig)
@@ -202,22 +211,27 @@ namespace HordeServer.VersionControl.Perforce
 
 			public async ValueTask<bool> MatchesFilterAsync(FileFilter filter, CancellationToken cancellationToken)
 			{
-				ICommit? other = await _owner!.GetChangeDetailsAsync(_streamConfig, Number, cancellationToken);
-				if (other == null)
+				if (HasAllFiles)
 				{
-					return false;
+					return filter.ApplyTo(Files).Any();
 				}
-				return await other.MatchesFilterAsync(filter, cancellationToken);
+				
+				IReadOnlyList<string> files = await GetFilesAsync(null, null, cancellationToken);
+				return filter.ApplyTo(files).Any();
 			}
 
 			public async ValueTask<IReadOnlyList<string>> GetFilesAsync(int? minFiles, int? maxFiles, CancellationToken cancellationToken)
 			{
-				ICommit? other = await _owner!.GetChangeDetailsAsync(_streamConfig, Number, cancellationToken);
-				if (other == null)
+				if (HasAllFiles)
 				{
-					return Array.Empty<string>();
+					return Files;
 				}
-				return await other.GetFilesAsync(minFiles, maxFiles, cancellationToken);
+
+				// The cached commit document does not have a complete list of files for the CL because it exceeded max files to be saved
+				// Query Perforce server for the complete list and return that, but do not store it
+				using IPooledPerforceConnection perforce = await _owner.ConnectAsync(_streamConfig.ClusterName, null, cancellationToken);
+				DescribeRecord describeRecord = await perforce.DescribeAsync(DescribeOptions.None, maxFiles ?? -1, Number, cancellationToken);
+				return await perforce.GetStreamFilesAsync(_streamConfig, describeRecord, cancellationToken);
 			}
 		}
 
@@ -413,9 +427,9 @@ namespace HordeServer.VersionControl.Perforce
 			}
 		}
 
-		async Task UpdateClusterInternalAsync(string clusterName, List<StreamInfo> streamInfos, SpanContext parentContext, CancellationToken cancellationToken)
+		internal async Task UpdateClusterInternalAsync(string clusterName, List<StreamInfo> streamInfos, SpanContext? parentContext, CancellationToken cancellationToken)
 		{
-			using TelemetrySpan span = _tracer.StartActiveSpan($"{nameof(PerforceServiceCache)}.{nameof(UpdateClusterInternalAsync)}", parentContext: parentContext);
+			using TelemetrySpan span = _tracer.StartActiveSpan($"{nameof(PerforceServiceCache)}.{nameof(UpdateClusterInternalAsync)}", parentContext: parentContext ?? default);
 
 			// Get the current state of this cluster
 			CacheState cacheState = await _mongoService.GetSingletonAsync<CacheState>(cancellationToken);
@@ -502,9 +516,6 @@ namespace HordeServer.VersionControl.Perforce
 				// Create a buffer for files
 				List<string> files = new List<string>();
 
-				// Maximum number of files to query by default
-				const int MaxFiles = 1000;
-
 				// Describe the changes and create records for them
 				PerforceResponseList<DescribeRecord> describeRecordResponses = await perforce.TryDescribeAsync(DescribeOptions.None, MaxFiles, changeNumbers.ToArray(), cancellationToken);
 				foreach (StreamInfo streamInfo in streamInfos)
@@ -545,8 +556,8 @@ namespace HordeServer.VersionControl.Perforce
 							_logger.LogInformation("Replicating {StreamId} commit {Change}", streamInfo.StreamConfig.Id, describeRecord.Number);
 
 							IReadOnlyList<CommitTag> commitTags = await commit.GetTagsAsync(cancellationToken);
-
-							CachedCommitDoc commitDoc = new CachedCommitDoc(commit, commitTags.ToList());
+							List<string> cappedFiles = files.Take(MaxFiles).ToList();
+							CachedCommitDoc commitDoc = new (commit, commitTags.ToList(), cappedFiles, hasAllFiles);
 							await AddCachedCommitAsync(commitDoc, cancellationToken);
 
 							await _redisService.PublishAsync(s_commitUpdateChannel, streamInfo.StreamConfig.Id, CommandFlags.FireAndForget);
@@ -738,7 +749,8 @@ namespace HordeServer.VersionControl.Perforce
 								cachedCommit.PostLoad(_owner, StreamConfig);
 								await _owner.AddCachedCommitAsync(cachedCommit, cancellationToken);
 								await _owner._mongoService.UpdateSingletonAsync<CacheState>(x => TryUpdateRange(x, StreamConfig, commit.Id.GetPerforceChange(), maxChange), cancellationToken);
-								_owner._logger.LogDebug("Adding new cached commit for {StreamId} at change {Change}", StreamConfig.Id, commit.Id);
+								_owner._logger.LogDebug("Adding new cached commit for {StreamId} at change {Change} with {NumFiles} (hasAllFiles={HasAllFiles})",
+									StreamConfig.Id, commit.Id, cachedCommit.Files.Count, cachedCommit.HasAllFiles);
 
 								if (tags == null || tags.Any(x => cachedCommit.CommitTags.Contains(x)))
 								{
