Horde: Fix .NET 10 warnings

diff --git a/Engine/Source/Programs/Horde/Plugins/Build/HordeServer.Build.Tests/Fleet/JobQueueStrategyTest.cs b/Engine/Source/Programs/Horde/Plugins/Build/HordeServer.Build.Tests/Fleet/JobQueueStrategyTest.cs
--- a/Engine/Source/Programs/Horde/Plugins/Build/HordeServer.Build.Tests/Fleet/JobQueueStrategyTest.cs
+++ b/Engine/Source/Programs/Horde/Plugins/Build/HordeServer.Build.Tests/Fleet/JobQueueStrategyTest.cs
@@ -161,7 +161,7 @@ namespace HordeServer.Tests.Fleet
 				await job.TryUpdateBatchAsync(job.Batches[0].Id, null, JobStepBatchState.Ready, null);
 			}
 
-			return (new JobQueueStrategy(JobCollection, GraphCollection, StreamCollection, Clock, Cache, isDowntimeActive, ServiceProvider.GetRequiredService<IOptionsMonitor<BuildConfig>>()), poolSize, pool, agents);
+			return (new JobQueueStrategy(JobCollection, GraphCollection, Clock, Cache, isDowntimeActive, ServiceProvider.GetRequiredService<IOptionsMonitor<BuildConfig>>()), poolSize, pool, agents);
 		}
 
 		private async Task<IJob> AddPlaceholderJobAsync(IGraph graph, StreamId streamId, string nodeNameToExecute)
diff --git a/Engine/Source/Programs/Horde/Plugins/Build/HordeServer.Build.Tests/Jobs/TestData/TestDataTests.cs b/Engine/Source/Programs/Horde/Plugins/Build/HordeServer.Build.Tests/Jobs/TestData/TestDataTests.cs
--- a/Engine/Source/Programs/Horde/Plugins/Build/HordeServer.Build.Tests/Jobs/TestData/TestDataTests.cs
+++ b/Engine/Source/Programs/Horde/Plugins/Build/HordeServer.Build.Tests/Jobs/TestData/TestDataTests.cs
@@ -688,7 +688,7 @@ namespace HordeServer.Tests.Jobs.TestData
 				Intent = "Test intent",
 				Team = "TeamTest",
 				Owner = userIdString,
-				Customers = [userIdString],
+				Customers = [userIdString!],
 				Notifications = new Dictionary<string, string>(){ { _mainStreamId.ToString(), "test-workflow1" } , { _releaseStreamId.ToString(), "test-workflow2" } }
 			};
 			testauditPostResult = await TestDataController.AddTestAuditV2Async(testAuditPostRequest);
diff --git a/Engine/Source/Programs/Horde/Plugins/Build/HordeServer.Build/Agents/Fleet/JobQueueStrategy.cs b/Engine/Source/Programs/Horde/Plugins/Build/HordeServer.Build/Agents/Fleet/JobQueueStrategy.cs
--- a/Engine/Source/Programs/Horde/Plugins/Build/HordeServer.Build/Agents/Fleet/JobQueueStrategy.cs
+++ b/Engine/Source/Programs/Horde/Plugins/Build/HordeServer.Build/Agents/Fleet/JobQueueStrategy.cs
@@ -30,7 +30,6 @@ namespace HordeServer.Agents.Fleet
 
 		private readonly IJobCollection _jobs;
 		private readonly IGraphCollection _graphs;
-		private readonly IStreamCollection _streamCollection;
 		private readonly IClock _clock;
 		private readonly IMemoryCache _cache;
 		private readonly bool _isDowntimeActive;
@@ -39,19 +38,10 @@ namespace HordeServer.Agents.Fleet
 		/// <summary>
 		/// Constructor
 		/// </summary>
-		/// <param name="jobs"></param>
-		/// <param name="graphs"></param>
-		/// <param name="streamCollection"></param>
-		/// <param name="clock"></param>
-		/// <param name="cache"></param>
-		/// <param name="isDowntimeActive"></param>
-		/// <param name="buildConfig"></param>
-		/// <param name="settings"></param>
-		public JobQueueStrategy(IJobCollection jobs, IGraphCollection graphs, IStreamCollection streamCollection, IClock clock, IMemoryCache cache, bool isDowntimeActive, IOptionsMonitor<BuildConfig> buildConfig, JobQueueSettings? settings = null)
+		public JobQueueStrategy(IJobCollection jobs, IGraphCollection graphs, IClock clock, IMemoryCache cache, bool isDowntimeActive, IOptionsMonitor<BuildConfig> buildConfig, JobQueueSettings? settings = null)
 		{
 			_jobs = jobs;
 			_graphs = graphs;
-			_streamCollection = streamCollection;
 			_clock = clock;
 			_cache = cache;
 			_isDowntimeActive = isDowntimeActive;
@@ -181,29 +171,16 @@ namespace HordeServer.Agents.Fleet
 		}
 	}
 
-	class JobQueueStrategyFactory : PoolSizeStrategyFactory<JobQueueSettings>
+	class JobQueueStrategyFactory(
+		IGraphCollection graphCollection,
+		IJobCollection jobCollection,
+		IDowntimeService downtimeService,
+		IClock clock,
+		IMemoryCache memoryCache,
+		IOptionsMonitor<BuildConfig> buildConfig)
+		: PoolSizeStrategyFactory<JobQueueSettings>(PoolSizeStrategy.JobQueue)
 	{
-		readonly IGraphCollection _graphCollection;
-		readonly IJobCollection _jobCollection;
-		readonly IStreamCollection _streamCollection;
-		readonly IDowntimeService _downtimeService;
-		readonly IClock _clock;
-		readonly IMemoryCache _memoryCache;
-		readonly IOptionsMonitor<BuildConfig> _buildConfig;
-
-		public JobQueueStrategyFactory(IGraphCollection graphCollection, IJobCollection jobCollection, IStreamCollection streamCollection, IDowntimeService downtimeService, IClock clock, IMemoryCache memoryCache, IOptionsMonitor<BuildConfig> buildConfig)
-			: base(PoolSizeStrategy.JobQueue)
-		{
-			_graphCollection = graphCollection;
-			_jobCollection = jobCollection;
-			_streamCollection = streamCollection;
-			_downtimeService = downtimeService;
-			_clock = clock;
-			_memoryCache = memoryCache;
-			_buildConfig = buildConfig;
-		}
-
 		public override IPoolSizeStrategy Create(JobQueueSettings settings)
-			=> new JobQueueStrategy(_jobCollection, _graphCollection, _streamCollection, _clock, _memoryCache, _downtimeService.IsDowntimeActive, _buildConfig, settings);
+			=> new JobQueueStrategy(jobCollection, graphCollection, clock, memoryCache, downtimeService.IsDowntimeActive, buildConfig, settings);
 	}
 }
diff --git a/Engine/Source/Programs/Horde/Plugins/Build/HordeServer.Build/Devices/DeviceService.cs b/Engine/Source/Programs/Horde/Plugins/Build/HordeServer.Build/Devices/DeviceService.cs
--- a/Engine/Source/Programs/Horde/Plugins/Build/HordeServer.Build/Devices/DeviceService.cs
+++ b/Engine/Source/Programs/Horde/Plugins/Build/HordeServer.Build/Devices/DeviceService.cs
@@ -7,7 +7,6 @@ using EpicGames.Horde.Devices;
 using EpicGames.Horde.Jobs;
 using EpicGames.Horde.Jobs.Graphs;
 using EpicGames.Horde.Projects;
-using EpicGames.Horde.Streams;
 using EpicGames.Horde.Users;
 using HordeServer.Auditing;
 using HordeServer.Jobs;
@@ -65,7 +64,6 @@ namespace HordeServer.Devices
 	{
 		readonly INotificationService _notificationService;
 		readonly JobService _jobService;
-		readonly IStreamCollection _streamCollection;
 		readonly IUserCollection _userCollection;
 		readonly Tracer _tracer;
 		readonly ILogger<DeviceService> _logger;
@@ -87,12 +85,11 @@ namespace HordeServer.Devices
 		/// <summary>
 		/// Device service constructor
 		/// </summary>
-		public DeviceService(IDeviceCollection devices, IUserCollection userCollection, JobService jobService, IStreamCollection streamCollection, IServerInfo serverInfo, IOptionsMonitor<BuildServerConfig> settings, IOptionsMonitor<BuildConfig> buildConfig, INotificationService notificationService, IClock clock, Tracer tracer, ILogger<DeviceService> logger, IAuditLog<DeviceId> auditLog)
+		public DeviceService(IDeviceCollection devices, IUserCollection userCollection, JobService jobService, IServerInfo serverInfo, IOptionsMonitor<BuildServerConfig> settings, IOptionsMonitor<BuildConfig> buildConfig, INotificationService notificationService, IClock clock, Tracer tracer, ILogger<DeviceService> logger, IAuditLog<DeviceId> auditLog)
 		{
 			_userCollection = userCollection;
 			_devices = devices;
 			_jobService = jobService;
-			_streamCollection = streamCollection;
 			_notificationService = notificationService;
 			_ticker = clock.AddSharedTicker<DeviceService>(TimeSpan.FromMinutes(1.0), TickAsync, logger);
 			_telemetryTicker = clock.AddSharedTicker("DeviceService.Telemetry", TimeSpan.FromMinutes(10.0), TickTelemetryAsync, logger);
diff --git a/Engine/Source/Programs/Horde/Plugins/Build/HordeServer.Build/Jobs/Bisect/BisectController.cs b/Engine/Source/Programs/Horde/Plugins/Build/HordeServer.Build/Jobs/Bisect/BisectController.cs
--- a/Engine/Source/Programs/Horde/Plugins/Build/HordeServer.Build/Jobs/Bisect/BisectController.cs
+++ b/Engine/Source/Programs/Horde/Plugins/Build/HordeServer.Build/Jobs/Bisect/BisectController.cs
@@ -7,7 +7,6 @@ using EpicGames.Horde.Jobs.Graphs;
 using EpicGames.Horde.Jobs.Templates;
 using EpicGames.Horde.Streams;
 using EpicGames.Horde.Users;
-using HordeServer.Commits;
 using HordeServer.Jobs.Graphs;
 using HordeServer.Streams;
 using HordeServer.Users;
@@ -263,7 +262,6 @@ namespace HordeServer.Jobs.Bisect
 		readonly IBisectTaskCollection _bisectTaskCollection;
 		readonly IJobCollection _jobCollection;
 		readonly JobService _jobService;
-		readonly ICommitService _commitService;
 		readonly IJobStepRefCollection _jobStepRefs;
 		readonly IGraphCollection _graphCollection;
 		readonly IUserCollection _userCollection;
@@ -274,11 +272,10 @@ namespace HordeServer.Jobs.Bisect
 		/// <summary>
 		/// Constructor
 		/// </summary>
-		public BisectTasksController(IBisectTaskCollection bisectTaskCollection, JobService jobService, ICommitService commitService, IJobCollection jobCollection, IJobStepRefCollection jobStepRefs, IGraphCollection graphCollection, IUserCollection userCollection, Tracer tracer, ILogger<BisectTasksController> logger, IOptionsSnapshot<BuildConfig> buildConfig)
+		public BisectTasksController(IBisectTaskCollection bisectTaskCollection, JobService jobService, IJobCollection jobCollection, IJobStepRefCollection jobStepRefs, IGraphCollection graphCollection, IUserCollection userCollection, Tracer tracer, ILogger<BisectTasksController> logger, IOptionsSnapshot<BuildConfig> buildConfig)
 		{
 			_bisectTaskCollection = bisectTaskCollection;
 			_jobService = jobService;
-			_commitService = commitService;
 			_jobCollection = jobCollection;
 			_jobStepRefs = jobStepRefs;
 			_graphCollection = graphCollection;
diff --git a/Engine/Source/Programs/Horde/Plugins/Build/HordeServer.Build/Jobs/JobCollection.cs b/Engine/Source/Programs/Horde/Plugins/Build/HordeServer.Build/Jobs/JobCollection.cs
--- a/Engine/Source/Programs/Horde/Plugins/Build/HordeServer.Build/Jobs/JobCollection.cs
+++ b/Engine/Source/Programs/Horde/Plugins/Build/HordeServer.Build/Jobs/JobCollection.cs
@@ -783,12 +783,13 @@ namespace HordeServer.Jobs
 		const int MaxRetries = 2;
 
 		readonly IMongoCollection<JobDocument> _jobs;
+		
+		[SuppressMessage("CodeQuality", "IDE0052:Remove unread private members")]
 		readonly MongoIndex<JobDocument> _streamThenTemplateThenCreationTimeIndex;
 		readonly MongoIndex<JobDocument> _startedByBisectTaskIdIndex;
 		readonly ITelemetryWriter _telemetryWriter;
 		readonly IClock _clock;
 		readonly IGraphCollection _graphCollection;
-		readonly IJobStepRefCollection _jobStepRefs;
 		readonly ILogCollection _logCollection;
 		readonly ICommitService _commitService;
 		readonly Tracer _tracer;
@@ -799,11 +800,10 @@ namespace HordeServer.Jobs
 		/// <summary>
 		/// Constructor
 		/// </summary>
-		public JobCollection(IMongoService mongoService, IClock clock, IGraphCollection graphCollection, ILogCollection logCollection, ICommitService commitService, ITelemetryWriter telemetryWriter, IOptionsMonitor<BuildConfig> buildConfig, Tracer tracer, ILogger<JobCollection> logger, IAuditLog<JobId> auditLog, IJobStepRefCollection jobStepRefCollection)
+		public JobCollection(IMongoService mongoService, IClock clock, IGraphCollection graphCollection, ILogCollection logCollection, ICommitService commitService, ITelemetryWriter telemetryWriter, IOptionsMonitor<BuildConfig> buildConfig, Tracer tracer, ILogger<JobCollection> logger, IAuditLog<JobId> auditLog)
 		{
 			_clock = clock;
 			_graphCollection = graphCollection;
-			_jobStepRefs = jobStepRefCollection;
 			_logCollection = logCollection;
 			_commitService = commitService;
 			_telemetryWriter = telemetryWriter;
diff --git a/Engine/Source/Programs/Horde/Plugins/Build/HordeServer.Build/Jobs/TestData/TestDataCollectionV2.cs b/Engine/Source/Programs/Horde/Plugins/Build/HordeServer.Build/Jobs/TestData/TestDataCollectionV2.cs
--- a/Engine/Source/Programs/Horde/Plugins/Build/HordeServer.Build/Jobs/TestData/TestDataCollectionV2.cs
+++ b/Engine/Source/Programs/Horde/Plugins/Build/HordeServer.Build/Jobs/TestData/TestDataCollectionV2.cs
@@ -7,7 +7,6 @@ using EpicGames.Horde.Jobs.Templates;
 using EpicGames.Horde.Jobs.TestData;
 using EpicGames.Horde.Streams;
 using EpicGames.Horde.Users;
-using HordeServer.Commits;
 using HordeServer.Server;
 using HordeServer.Utilities;
 using Microsoft.Extensions.Logging;
@@ -674,16 +673,14 @@ namespace HordeServer.Jobs.TestData
 		readonly IMongoCollection<TestPhaseSessionDocument> _testPhaseSessions;
 		readonly IMongoCollection<TestSessionStreamDocument> _testStreams;
 		readonly IMongoCollection<TestAuditDocument> _testAudits;
-		readonly ICommitService _commitService;
 		readonly Tracer _tracer;
 		readonly ILogger _logger;
 
 		/// <summary>
 		/// Constructor
 		/// </summary>
-		public TestDataCollectionV2(IMongoService mongoService, ICommitService commitService, Tracer tracer, ILogger<TestDataCollection> logger)
+		public TestDataCollectionV2(IMongoService mongoService, Tracer tracer, ILogger<TestDataCollection> logger)
 		{
-			_commitService = commitService;
 			_tracer = tracer;
 			_logger = logger;
 
