Horde: Capture CPU seconds spent during a lease on the agent

diff --git a/Engine/Source/Programs/Horde/HordeAgent/Leases/LeaseHandler.cs b/Engine/Source/Programs/Horde/HordeAgent/Leases/LeaseHandler.cs
--- a/Engine/Source/Programs/Horde/HordeAgent/Leases/LeaseHandler.cs
+++ b/Engine/Source/Programs/Horde/HordeAgent/Leases/LeaseHandler.cs
@@ -84,9 +84,9 @@ namespace HordeAgent.Leases
 		/// <summary>
 		/// Starts executing the lease
 		/// </summary>
-		public void Start(ISession session, Tracer tracer, ILogger logger, LeaseLoggerFactory leaseLoggerFactory)
+		public void Start(ISession session, ISystemMetrics systemMetrics, Tracer tracer, ILogger logger, LeaseLoggerFactory leaseLoggerFactory)
 		{
-			Result = Task.Run(() => HandleLeaseAsync(session, tracer, logger, leaseLoggerFactory));
+			Result = Task.Run(() => HandleLeaseAsync(session, systemMetrics, tracer, logger, leaseLoggerFactory));
 		}
 
 		/// <summary>
@@ -101,7 +101,7 @@ namespace HordeAgent.Leases
 		/// <summary>
 		/// Handle a lease request
 		/// </summary>
-		async Task<LeaseResult> HandleLeaseAsync(ISession session, Tracer tracer, ILogger logger, LeaseLoggerFactory leaseLoggerFactory)
+		async Task<LeaseResult> HandleLeaseAsync(ISession session, ISystemMetrics systemMetrics, Tracer tracer, ILogger logger, LeaseLoggerFactory leaseLoggerFactory)
 		{
 			using TelemetrySpan span = tracer.StartActiveSpan($"{nameof(LeaseHandler)}.{nameof(HandleLeaseAsync)}");
 			span.SetAttribute("horde.lease.id", RpcLease.Id.ToString());
@@ -109,6 +109,9 @@ namespace HordeAgent.Leases
 			span.SetAttribute("horde.agent.session_id", session.SessionId.ToString());
 			logger.LogInformation("Handling lease {LeaseId}: {LeaseTypeUrl}", Id, RpcLease.Payload.TypeUrl);
 
+			Stopwatch leaseDuration = Stopwatch.StartNew();
+			CpuTimeMetrics? cpuTimeStart = systemMetrics.GetCpuTime();
+			
 			// Get the lease outcome
 			LeaseResult result = new (LeaseOutcome.Failed, LeaseOutcomeReason.UnexpectedError);
 			try
@@ -126,6 +129,7 @@ namespace HordeAgent.Leases
 				span.SetStatus(Status.Error);
 				span.RecordException(ex);
 			}
+			CpuTimeMetrics? cpuTimeEnd = systemMetrics.GetCpuTime();
 
 			// Update the state of the lease
 			RpcLease newRpcLease = new RpcLease(RpcLease);
@@ -144,6 +148,37 @@ namespace HordeAgent.Leases
 				newRpcLease.OutcomeReasons.AddRange(RpcLease.ToProtobufOutcomeReasons(reason));
 				newRpcLease.Output = (result.Output != null) ? ByteString.CopyFrom(result.Output) : ByteString.Empty;
 			}
+			
+			newRpcLease.Metrics = new RpcLeaseMetrics
+			{
+				SetupTimeMs = 0, // TODO: Implement
+				TeardownTimeMs = 0, // TODO: Implement
+				GlobalCpuUtilization = 0.0f,
+				CpuCount = Environment.ProcessorCount,
+			};
+			
+			if (cpuTimeStart != null && cpuTimeEnd != null)
+			{
+				CpuTimeMetrics cpuTimeSpent = cpuTimeEnd.Subtract(cpuTimeStart);
+				TimeSpan totalCpuTimeSpent = cpuTimeSpent.User + cpuTimeSpent.System;
+				TimeSpan totalCpuTimeAvailable = Environment.ProcessorCount * leaseDuration.Elapsed;
+				
+				// Prevent the unlikely divide by zero below
+				totalCpuTimeAvailable = totalCpuTimeAvailable == TimeSpan.Zero ? TimeSpan.FromMilliseconds(1) : totalCpuTimeAvailable;
+				
+				newRpcLease.Metrics.GlobalCpuUtilization = totalCpuTimeSpent / totalCpuTimeAvailable;
+				newRpcLease.Metrics.GlobalCpuTime = new RpcLeaseMetricCpuTime
+				{
+					UserMs = Convert.ToUInt32(cpuTimeSpent.User.TotalMilliseconds),
+					SystemMs = Convert.ToUInt32(cpuTimeSpent.System.TotalMilliseconds),
+					IdleMs = Convert.ToUInt32(cpuTimeSpent.Idle.TotalMilliseconds),
+				};
+				span.SetAttribute("horde.lease.cpu.time_spent_ms", (uint)totalCpuTimeSpent.TotalMilliseconds);
+				span.SetAttribute("horde.lease.cpu.time_available_ms", (uint)totalCpuTimeAvailable.TotalMilliseconds);
+				span.SetAttribute("horde.lease.cpu.utilization", newRpcLease.Metrics.GlobalCpuUtilization);
+				span.SetAttribute("horde.lease.cpu.num_cores", newRpcLease.Metrics.CpuCount);
+			}
+			
 			RpcLease = newRpcLease;
 
 			logger.LogInformation("Transitioning lease {LeaseId} to {State}, outcome={Outcome}", Id, RpcLease.State, RpcLease.Outcome);
diff --git a/Engine/Source/Programs/Horde/HordeAgent/Leases/LeaseManager.cs b/Engine/Source/Programs/Horde/HordeAgent/Leases/LeaseManager.cs
--- a/Engine/Source/Programs/Horde/HordeAgent/Leases/LeaseManager.cs
+++ b/Engine/Source/Programs/Horde/HordeAgent/Leases/LeaseManager.cs
@@ -383,7 +383,7 @@ namespace HordeAgent.Leases
 									_logger.LogInformation("Adding lease {LeaseId}", serverLease.Id);
 
 									LeaseHandler leaseHandler = CreateLeaseHandler(serverLease);
-									leaseHandler.Start(_session, _tracer, _logger, _leaseLoggerFactory);
+									leaseHandler.Start(_session, _systemMetrics, _tracer, _logger, _leaseLoggerFactory);
 									leaseHandler.Result.ContinueWith((Task<LeaseResult> task) =>
 									{
 										LeaseResult result = task.Result;
diff --git a/Engine/Source/Programs/Horde/HordeAgent/Services/TelemetryService.cs b/Engine/Source/Programs/Horde/HordeAgent/Services/TelemetryService.cs
--- a/Engine/Source/Programs/Horde/HordeAgent/Services/TelemetryService.cs
+++ b/Engine/Source/Programs/Horde/HordeAgent/Services/TelemetryService.cs
@@ -8,6 +8,7 @@ using System.Text;
 using EpicGames.Core;
 using Microsoft.Extensions.Hosting;
 using Microsoft.Extensions.Logging;
+using static HordeAgent.Services.WindowsCapabilities;
 
 namespace HordeAgent.Services;
 
@@ -92,6 +93,41 @@ public class DiskMetrics
 	}
 }
 
+/// <summary>
+/// Metrics for CPU time spent
+/// </summary>
+public class CpuTimeMetrics
+{
+	/// <summary>
+	/// CPU time spent in user mode
+	/// </summary>
+	public TimeSpan User { get; set; }
+
+	/// <summary>
+	/// CPU time spent in system/kernel mode
+	/// </summary>
+	public TimeSpan System { get; set; }
+
+	/// <summary>
+	/// CPU time spent idle
+	/// </summary>
+	public TimeSpan Idle { get; set; }
+
+	/// <summary>
+	/// Subtract times from other object passed in
+	/// </summary>
+	public CpuTimeMetrics Subtract(CpuTimeMetrics other)
+	{
+		return new CpuTimeMetrics { User = User - other.User, System = System - other.System, Idle = Idle - other.Idle };
+	}
+	
+	/// <inheritdoc />
+	public override string ToString()
+	{
+		return $"User={(int)User.TotalMilliseconds} System={(int)System.TotalMilliseconds} Idle={(int)Idle.TotalMilliseconds}";
+	}
+}
+
 /// <summary>
 /// OS agnostic interface for retrieving system metrics (CPU, memory etc)
 /// </summary>
@@ -103,6 +139,12 @@ public interface ISystemMetrics
 	/// <returns>An object with CPU usage metrics</returns>
 	CpuMetrics? GetCpu();
 
+	/// <summary>
+	/// Get cumulative CPU time spent since system start
+	/// </summary>
+	/// <returns>An object with CPU time metrics</returns>
+	CpuTimeMetrics? GetCpuTime();
+	
 	/// <summary>
 	/// Get memory usage metrics
 	/// </summary>
@@ -122,16 +164,16 @@ public interface ISystemMetrics
 public sealed class DefaultSystemMetrics : ISystemMetrics
 {
 	/// <inheritdoc/>
-	public CpuMetrics? GetCpu()
-		=> null;
+	public CpuMetrics? GetCpu() => null;
+	
+	/// <inheritdoc/>
+	public CpuTimeMetrics? GetCpuTime() => null;
 
 	/// <inheritdoc/>
-	public MemoryMetrics? GetMemory()
-		=> null;
+	public MemoryMetrics? GetMemory() => null;
 
 	/// <inheritdoc/>
-	public DiskMetrics? GetDisk()
-		=> null;
+	public DiskMetrics? GetDisk() => null;
 }
 
 /// <summary>
@@ -155,29 +197,47 @@ public sealed class WindowsSystemMetrics : ISystemMetrics
 		GetCpu(); // Trigger this to ensure performance counter has a fetched value. Avoids an initial zero result when called later.
 	}
 
-    /// <inheritdoc />
-    public CpuMetrics GetCpu()
-    {
-        (float idle, float system, float user) = WindowsCapabilities.GetCpuMetrics();
-        return new CpuMetrics
-        {
-            User = user,
-            System = system,
-            Idle = idle
-        };
-    }
-
-    /// <inheritdoc />
-    public MemoryMetrics GetMemory()
+	/// <inheritdoc />
+	public CpuMetrics GetCpu()
+	{
+		(float idle, float system, float user) = GetCpuMetrics();
+		return new CpuMetrics { User = user, System = system, Idle = idle };
+	}
+
+	/// <inheritdoc />
+	public CpuTimeMetrics? GetCpuTime()
 	{
-		(ulong totalBytes, ulong availableBytes) = WindowsCapabilities.GetPhysicalMemory();
+		try
+		{
+			if (GetSystemTimes(out FILETIME idle, out FILETIME kernel, out FILETIME user))
+			{
+				return new CpuTimeMetrics
+				{
+					Idle = TimeSpan.FromMilliseconds(idle.ToUInt64() / 10000.0),
+					User = TimeSpan.FromMilliseconds(user.ToUInt64() / 10000.0),
+					// Kernel includes idle time, so must be subtracted
+					System = TimeSpan.FromMilliseconds((kernel.ToUInt64() - idle.ToUInt64()) / 10000.0),
+				};
+			}
+		}
+		catch
+		{
+			// Ignore any error.
+		}
+		return null;
+	}
+
+	/// <inheritdoc />
+	public MemoryMetrics GetMemory()
+	{
+		(ulong totalBytes, ulong availableBytes) = GetPhysicalMemory();
 		ulong usedBytes = totalBytes - availableBytes;
 		return new MemoryMetrics
 		{
 			Total = (uint)(totalBytes / 1024),
 			Available = (uint)(availableBytes / 1024),
 			Used = (uint)(usedBytes / 1024),
-			UsedPercentage = usedBytes / (float)totalBytes,
+			UsedPercentage = usedBytes / (float)(totalBytes > 0 ? totalBytes : 1),
 		};
 	}
 
@@ -220,10 +280,34 @@ public sealed class LinuxSystemMetrics : ISystemMetrics, IDisposable
 {
 	private const string ProcMeminfoPath = "/proc/meminfo";
 	private const string ProcStatPath = "/proc/stat";
+	private const int SC_CLK_TCK = 2;
 	private const int UserIndex = 1;
+	private const int NiceIndex = 2;
 	private const int SystemIndex = 3;
 	private const int IdleIndex = 4;
 	
+	[DllImport("libc", SetLastError = true)]
+	private static extern long sysconf(int name);
+
+	private static readonly Lazy<long> s_clockTicksPerSecond = new (() =>
+	{
+		try
+		{
+			long result = sysconf(SC_CLK_TCK);
+			if (result > 0)
+			{
+				return result;
+			}
+		}
+		catch
+		{
+			// Fall through to default
+		}
+		return 100; // Default fallback
+	});
+
+	
+	
 	private readonly DirectoryReference _workingDir;
 	
 	/// <summary>
@@ -259,6 +343,19 @@ public sealed class LinuxSystemMetrics : ISystemMetrics, IDisposable
 			Idle = (idleDelta * 100.0f) / totalDelta,
 		};
 	}
+
+	/// <inheritdoc/>
+	public CpuTimeMetrics GetCpuTime()
+	{
+		(long user, long system, long idle) cpuTimes = ReadCpuTimes();
+		long clockTicksPerSecond = s_clockTicksPerSecond.Value;
+		return new CpuTimeMetrics
+		{
+			User = TimeSpan.FromMilliseconds(cpuTimes.user * 1000.0 / clockTicksPerSecond),
+			System = TimeSpan.FromMilliseconds(cpuTimes.system * 1000.0 / clockTicksPerSecond),
+			Idle = TimeSpan.FromMilliseconds(cpuTimes.idle * 1000.0 / clockTicksPerSecond),
+		};
+	}
 	
 	/// <inheritdoc/>
 	public MemoryMetrics? GetMemory()
@@ -330,7 +427,11 @@ public sealed class LinuxSystemMetrics : ISystemMetrics, IDisposable
 			.Skip(1) // Skip "cpu" label
 			.Select(Int64.Parse)
 			.ToArray();
-		return (user: values[UserIndex - 1], system: values[SystemIndex - 1], idle: values[IdleIndex - 1]);
+		long user = values[UserIndex - 1];
+		long nice = values[NiceIndex - 1];
+		long system = values[SystemIndex - 1];
+		long idle = values[IdleIndex - 1];
+		return (user + nice, system, idle);
 	}
 	
 	[StructLayout(LayoutKind.Sequential)]
diff --git a/Engine/Source/Programs/Horde/HordeAgent/Services/WindowsCapabilities.cs b/Engine/Source/Programs/Horde/HordeAgent/Services/WindowsCapabilities.cs
--- a/Engine/Source/Programs/Horde/HordeAgent/Services/WindowsCapabilities.cs
+++ b/Engine/Source/Programs/Horde/HordeAgent/Services/WindowsCapabilities.cs
@@ -65,9 +65,7 @@ public class WindowsCapabilities(IOptionsMonitor<AgentSettings> settings, ILogge
 	}
 	
 	private static void SetProp(RpcAgentCapabilities caps, string key, string value) => caps.Properties.Add($"{key}={value}");
-	private static void SetProp(RpcAgentCapabilities caps, string key, int value) => caps.Properties.Add($"{key}={value}");
 	private static void SetProp(RpcAgentCapabilities caps, string key, bool value) => caps.Properties.Add($"{key}={value}");
-	private static void SetResource(RpcAgentCapabilities caps, string key, int value) => caps.Resources.Add(key, value);
 	
 	/// <summary>
 	/// Gets details of current system's physical memory
@@ -245,18 +243,18 @@ public class WindowsCapabilities(IOptionsMonitor<AgentSettings> settings, ILogge
 		public ulong ullAvailExtendedVirtual;
 	}
 
-    [DllImport("kernel32.dll", SetLastError = true)]
-    [return: MarshalAs(UnmanagedType.Bool)]
-    static extern bool GetSystemTimes(out FILETIME idle, out FILETIME system, out FILETIME user);
+	[DllImport("kernel32.dll", SetLastError = true)]
+	[return: MarshalAs(UnmanagedType.Bool)]
+	internal static extern bool GetSystemTimes(out FILETIME idle, out FILETIME system, out FILETIME user);
 
-    [StructLayout(LayoutKind.Sequential)]
-    internal struct FILETIME
-    {
-        public uint dwLowDateTime;
-        public uint dwHighDateTime;
-        public ulong ToUInt64() => ((ulong)dwHighDateTime << 32) | dwLowDateTime;
-        public bool IsZero => dwLowDateTime == 0 && dwHighDateTime == 0;
-    }
+	[StructLayout(LayoutKind.Sequential)]
+	internal struct FILETIME
+	{
+		public uint dwLowDateTime;
+		public uint dwHighDateTime;
+		public ulong ToUInt64() => ((ulong)dwHighDateTime << 32) | dwLowDateTime;
+		public bool IsZero => dwLowDateTime == 0 && dwHighDateTime == 0;
+	}
 	#pragma warning enable
 	
 	#endregion P/Invoke
