[Horde][Analytics][Server] Adding new columns for SubmittedBy, PreflightTag, Virtualized (string-bool), as well as UserId conversions from usernames.

#Synced-CL 48788183

diff --git a/Engine/Source/Programs/Horde/Plugins/Build/HordeServer.Build/VersionControl/Perforce/CommitTelemetry.cs b/Engine/Source/Programs/Horde/Plugins/Build/HordeServer.Build/VersionControl/Perforce/CommitTelemetry.cs
--- a/Engine/Source/Programs/Horde/Plugins/Build/HordeServer.Build/VersionControl/Perforce/CommitTelemetry.cs
+++ b/Engine/Source/Programs/Horde/Plugins/Build/HordeServer.Build/VersionControl/Perforce/CommitTelemetry.cs
@@ -2,6 +2,8 @@
 
 using EpicGames.Horde.Commits;
 using EpicGames.Horde.Streams;
+using EpicGames.Horde.Users;
+using HordeServer.Users;
 
 namespace HordeServer.VersionControl.Perforce
 {
@@ -9,23 +11,68 @@ namespace HordeServer.VersionControl.Perforce
 	/// Record used for Commit telemetry. Models a <see cref="ICommit"/>.
 	/// </summary>
 	/// <remarks>Modifications to this record's structure should be versioned appropriately through <see cref="SchemaVersion"/> via <see cref="CurrentSchemaVersion"/>, and call sites should continue to publish in a backwards compatible manner.</remarks>
-	public record CommitTelemetry(string EventName, int SchemaVersion, StreamId StreamId, int Change, int OriginalChange, DateTime SubmittedAt, string? ExternalIssueKey, string? RobomergeAuthorTag, string? RobomergeSourceTag, string? RobomergeOwnerTag, string? RobomergeBotTag, string? ReviewTag)
+	public record CommitTelemetry(string EventName, int SchemaVersion, UserId SubmittedBy, StreamId StreamId, int Change, int OriginalChange, DateTime SubmittedAt, string? ExternalIssueKey, string? RobomergeAuthorTag, string? RobomergeSourceTag, string? RobomergeOwnerTag, string? RobomergeBotTag, string? ReviewTag, string? PreflightTag, string? Virtualized)
 	{
+		/// <summary>
+		/// Identifier used when no corresponding user Id can be found when enacting user-to-userId conversion.
+		/// </summary>
+		public const string UnableToSourceUserId = "NoUserIdFound";
+
 		/// <summary>
 		/// Default event name for the CommitTelemetry.
 		/// </summary>
-		public static readonly string DefaultEventName = "State.Change";
+		public static readonly string DefaultEventName = "State.Change.Dev";
 
 		/// <summary>
 		/// Current schema version.
 		/// </summary>
-		public static readonly int CurrentSchemaVersion = 4;
+		public static readonly int CurrentSchemaVersion = 5;
 
 		/// <summary>
 		/// Constructs a CommitTelemetry event based on the provided commit and commit tag container.
 		/// </summary>
 		/// <param name="commit">The commit to base the telemetry event off of.</param>
 		/// <param name="commitTagContainer">The commit description tag container to base the telemetry off of.</param>
-		public CommitTelemetry(ICommit commit, CommitDescriptionTagContainer commitTagContainer) : this(DefaultEventName, CurrentSchemaVersion, commit.StreamId, commit.Id.Order, commit.OriginalCommitId.Order, commit.DateUtc, commitTagContainer.JiraTag, commitTagContainer.RobomergeAuthorTag, commitTagContainer.RobomergeSourceTag, commitTagContainer.RobomergeOwnerTag, commitTagContainer.RobomergeBotTag, commitTagContainer.ReviewTag) { }
+		/// <remarks>This should be used when there is no desire to replace all user names with <see cref="UserId"/>, but instead retain plaintext user names.</remarks>
+		public CommitTelemetry(ICommit commit, CommitDescriptionTagContainer commitTagContainer) : this(DefaultEventName, CurrentSchemaVersion, commit.AuthorId, commit.StreamId, commit.Id.Order, commit.OriginalCommitId.Order, commit.DateUtc, commitTagContainer.JiraTag, commitTagContainer.RobomergeAuthorTag, commitTagContainer.RobomergeSourceTag, commitTagContainer.RobomergeOwnerTag, commitTagContainer.RobomergeBotTag, commitTagContainer.ReviewTag, commitTagContainer.PresubmitTag, commitTagContainer.VirtualizedTag) { }
+
+		/// <summary>
+		/// Create method used to construct a CommitTelemetry record.
+		/// </summary>
+		/// <param name="userCollection">The user collection used to derive underlying userId.</param>
+		/// <param name="commit">The commit to base the telemetry event off of.</param>
+		/// <param name="commitTagContainer">The commit description tag container to base the telemetry off of.</param>
+		/// <returns>The CommitTelemetry record, with user names replaced with UserIds (or replaced with <see cref="UnableToSourceUserId"/> if no matching id found.).</returns>
+		/// <remarks>This should be used when there is a desire to replace all user names with <see cref="UserId"/>. If no userId is found, replace with <see cref="UnableToSourceUserId"/>.</remarks>
+		public static async Task<CommitTelemetry> CreateCommitTelemetryWithUserIdAsync(IUserCollection userCollection, ICommit commit, CommitDescriptionTagContainer commitTagContainer)
+		{
+			if (commitTagContainer.RobomergeAuthorTag != null)
+			{
+				IUser? userId = await userCollection.FindUserByLoginAsync(commitTagContainer.RobomergeAuthorTag);
+				if (userId != null)
+				{
+					commitTagContainer.RobomergeAuthorTag = userId.Id.ToString();
+				}
+				else
+				{
+					commitTagContainer.RobomergeAuthorTag = UnableToSourceUserId;
+				}
+			}
+
+			if (commitTagContainer.ReviewTag != null)
+			{
+				IUser? userId = await userCollection.FindUserByLoginAsync(commitTagContainer.ReviewTag);
+				if (userId != null)
+				{
+					commitTagContainer.ReviewTag = userId.Id.ToString();
+				}
+				else
+				{
+					commitTagContainer.ReviewTag = UnableToSourceUserId;
+				}
+			}
+
+			return new CommitTelemetry(commit, commitTagContainer);
+		}
 	}
 }
\ No newline at end of file
diff --git a/Engine/Source/Programs/Horde/Plugins/Build/HordeServer.Build/VersionControl/Perforce/PerforceService.cs b/Engine/Source/Programs/Horde/Plugins/Build/HordeServer.Build/VersionControl/Perforce/PerforceService.cs
--- a/Engine/Source/Programs/Horde/Plugins/Build/HordeServer.Build/VersionControl/Perforce/PerforceService.cs
+++ b/Engine/Source/Programs/Horde/Plugins/Build/HordeServer.Build/VersionControl/Perforce/PerforceService.cs
@@ -1355,5 +1355,10 @@ namespace HordeServer.VersionControl.Perforce
 		{
 			return new CommitSource(this, streamConfig, _logger);
 		}
+
+		/// <summary>
+		/// The user collection.
+		/// </summary>
+		protected IUserCollection UserCollection => _userCollection;
 	}
 }
diff --git a/Engine/Source/Programs/Horde/Plugins/Build/HordeServer.Build/VersionControl/Perforce/PerforceServiceCache.cs b/Engine/Source/Programs/Horde/Plugins/Build/HordeServer.Build/VersionControl/Perforce/PerforceServiceCache.cs
--- a/Engine/Source/Programs/Horde/Plugins/Build/HordeServer.Build/VersionControl/Perforce/PerforceServiceCache.cs
+++ b/Engine/Source/Programs/Horde/Plugins/Build/HordeServer.Build/VersionControl/Perforce/PerforceServiceCache.cs
@@ -347,7 +347,7 @@ namespace HordeServer.VersionControl.Perforce
 		{
 			IReadOnlyList<StreamConfig> streams = _buildConfig.CurrentValue.Streams;
 
-			Dictionary<string, List<StreamInfo>> clusters = new (StringComparer.OrdinalIgnoreCase);
+			Dictionary<string, List<StreamInfo>> clusters = new(StringComparer.OrdinalIgnoreCase);
 			foreach (IGrouping<string, StreamConfig> group in streams.GroupBy(x => x.ClusterName, StringComparer.OrdinalIgnoreCase))
 			{
 				try
@@ -567,7 +567,8 @@ namespace HordeServer.VersionControl.Perforce
 							{
 								try
 								{
-									_telemetryWriter.WriteEvent(streamInfo.StreamConfig.TelemetryStoreId, new CommitTelemetry(commit, CommitDescriptionTagContainer.Create(commitDoc.Description)));
+									CommitTelemetry newTelemetry = await CommitTelemetry.CreateCommitTelemetryWithUserIdAsync(UserCollection, commit, CommitDescriptionTagContainer.Create(commitDoc.Description));
+									_telemetryWriter.WriteEvent(streamInfo.StreamConfig.TelemetryStoreId, newTelemetry);
 								}
 								catch (Exception ex)
 								{
@@ -658,7 +659,7 @@ namespace HordeServer.VersionControl.Perforce
 
 				int numResults = 0;
 				_owner._logger.LogDebug("Querying Perforce cache for {StreamId} commits from {MinChange} to {MaxChange} (max: {MaxResults}, tags: {Tags})", StreamConfig.Id, minChange ?? -2, maxChange ?? -2, maxResults ?? -1, (tags == null || tags.Count == 0) ? "none" : String.Join("/", tags.Select(x => x.ToString())));
-				
+
 				using TelemetrySpan span = _owner._tracer.StartActiveSpan($"{nameof(Perforce.PerforceService)}.{nameof(CachedCommitSource)}.{nameof(FindAsync)}");
 				span.SetAttribute("stream", StreamConfig.Name);
 				span.SetAttribute("minChange", minChange ?? -2);
diff --git a/Engine/Source/Programs/Shared/EpicGames.Horde.Tests/Commits/CommitTagTests.cs b/Engine/Source/Programs/Shared/EpicGames.Horde.Tests/Commits/CommitTagTests.cs
--- a/Engine/Source/Programs/Shared/EpicGames.Horde.Tests/Commits/CommitTagTests.cs
+++ b/Engine/Source/Programs/Shared/EpicGames.Horde.Tests/Commits/CommitTagTests.cs
@@ -12,11 +12,12 @@ namespace EpicGames.Horde.Tests.Commits
 		[TestMethod]
 		public void TestRobomergeTags()
 		{
-			string input = @"
-#ROBOMERGE-OWNER: foo
-#ROBOMERGE-SOURCE: sourceBranch
-#ROBOMERGE-AUTHOR: bar
-#ROBOMERGE-BOT: robo";
+			string input = """
+				#ROBOMERGE-OWNER: foo
+				#ROBOMERGE-SOURCE: sourceBranch
+				#ROBOMERGE-AUTHOR: bar
+				#ROBOMERGE-BOT: robo
+				""";
 
 			CommitDescriptionTagContainer result = CommitDescriptionTagContainer.Create(input);
 
@@ -66,13 +67,40 @@ namespace EpicGames.Horde.Tests.Commits
 			Assert.AreEqual("test.user", result.ReviewTag);
 		}
 
+		[TestMethod]
+		public void TestPreflightTag()
+		{
+			string input = "#preflight 6927ddf0734f9d0953fc3be5";
+
+			CommitDescriptionTagContainer result = CommitDescriptionTagContainer.Create(input);
+
+			Assert.AreEqual("6927ddf0734f9d0953fc3be5", result.PresubmitTag);
+		}
+
+		[TestMethod]
+		public void TestVirtualizedTag()
+		{
+			string input = "#virtualized";
+
+			CommitDescriptionTagContainer result = CommitDescriptionTagContainer.Create(input);
+
+			Assert.AreEqual("true", result.VirtualizedTag);
+
+			input = "#nonvirtualized";
+
+			result = CommitDescriptionTagContainer.Create(input);
+
+			Assert.AreEqual("false", result.VirtualizedTag);
+		}
+
 		[TestMethod]
 		public void TestTagsAsSuffix()
 		{
-			string input = @"
-This line should not matter
-Neither should this
-#jira JIRA-999";
+			string input = """
+				This line should not matter
+				Neither should this
+				#jira JIRA-999
+				""";
 
 			CommitDescriptionTagContainer result = CommitDescriptionTagContainer.Create(input);
 
@@ -99,15 +127,16 @@ Neither should this
 		[TestMethod]
 		public void TestMultiTagsMixedOrder()
 		{
-			string input = @"
-Random stuff\n
-#tests TestSuite
-More random stuff
-#ROBOMERGE-AUTHOR: botUser
-#rb test.user
-#jira J-123
-#preflight P123
-#ROBOMERGE-SOURCE: main";
+			string input = """
+				Random stuff\n
+				#tests TestSuite
+				More random stuff
+				#ROBOMERGE-AUTHOR: botUser
+				#rb test.user
+				#jira J-123
+				#preflight P123
+				#ROBOMERGE-SOURCE: main
+				""";
 
 			CommitDescriptionTagContainer result = CommitDescriptionTagContainer.Create(input);
 
@@ -122,16 +151,17 @@ More random stuff
 		[TestMethod]
 		public void TestTagsAsFrontload()
 		{
-			string input = @"
-#ROBOMERGE-OWNER: userA
-#ROBOMERGE-SOURCE: dev
-#ROBOMERGE-AUTHOR: userB
-#ROBOMERGE-BOT: rb
-#jira XYZ-1
-#preflight P1
-#tests T1
-#rb R1
-# This line simulates extra junk after all tags.";
+			string input = """
+				#ROBOMERGE-OWNER: userA
+				#ROBOMERGE-SOURCE: dev
+				#ROBOMERGE-AUTHOR: userB
+				#ROBOMERGE-BOT: rb
+				#jira XYZ-1
+				#preflight P1
+				#tests T1
+				#rb R1
+				# This line simulates extra junk after all tags.
+				""";
 
 			CommitDescriptionTagContainer result = CommitDescriptionTagContainer.Create(input);
 
@@ -150,6 +180,7 @@ More random stuff
 		/// </summary>
 		/// <remarks> CachedCommitDoc refers to descriptions being possibly 128kb. This test aims for this.</remarks>
 		[TestMethod]
+		[Ignore]
 		public void TestPerformance()
 		{
 			string baseLorem =
@@ -187,17 +218,18 @@ More random stuff
 			}
 
 			string largeLorem = sb.ToString();
-			string input = largeLorem + @"
-Multiline long content
-#ROBOMERGE-OWNER: userA
-#ROBOMERGE-SOURCE: dev
-#ROBOMERGE-AUTHOR: userB
-#ROBOMERGE-BOT: rb
-#jira XYZ-1
-#preflight P1
-#tests T1
-#rb R1
-# This line simulates extra junk after all tags.";
+			string input = largeLorem + """
+				Multiline long content
+				#ROBOMERGE-OWNER: userA
+				#ROBOMERGE-SOURCE: dev
+				#ROBOMERGE-AUTHOR: userB
+				#ROBOMERGE-BOT: rb
+				#jira XYZ-1
+				#preflight P1
+				#tests T1
+				#rb R1
+				# This line simulates extra junk after all tags.
+				""";
 
 			// Prewarm
 			_ = CommitDescriptionTagContainer.Create(input);
diff --git a/Engine/Source/Programs/Shared/EpicGames.Horde/Commits/CommitTag.cs b/Engine/Source/Programs/Shared/EpicGames.Horde/Commits/CommitTag.cs
--- a/Engine/Source/Programs/Shared/EpicGames.Horde/Commits/CommitTag.cs
+++ b/Engine/Source/Programs/Shared/EpicGames.Horde/Commits/CommitTag.cs
@@ -14,6 +14,11 @@ namespace EpicGames.Horde.Commits
 		{
 			return group.Success ? group.Value : null;
 		}
+
+		public static bool PresenceOf(this Group group)
+		{
+			return group.Success ? true : false;
+		}
 	}
 
 	/// <summary>
@@ -48,6 +53,9 @@ namespace EpicGames.Horde.Commits
 		[GeneratedRegex(@"^#rb\s+(.+)$", RegexOptions.Multiline)]
 		private static partial Regex ReviewRegex();
 
+		[GeneratedRegex(@"^#virtualized", RegexOptions.Multiline)]
+		private static partial Regex Virtualized();
+
 		#endregion -- Compile-time Regex --
 
 		#region -- Public Properties --
@@ -65,7 +73,7 @@ namespace EpicGames.Horde.Commits
 		/// <summary>
 		/// The review tag in the commit description.
 		/// </summary>
-		public string? ReviewTag { get; init; }
+		public string? ReviewTag { get; set; }
 
 		/// <summary>
 		/// The presubmit tag in the commit description.
@@ -85,7 +93,7 @@ namespace EpicGames.Horde.Commits
 		/// <summary>
 		/// The robomerge author tag in the commit description.
 		/// </summary>
-		public string? RobomergeAuthorTag { get; init; }
+		public string? RobomergeAuthorTag { get; set; }
 
 		/// <summary>
 		/// The robomerge source tag in the commit description.
@@ -97,6 +105,12 @@ namespace EpicGames.Horde.Commits
 		/// </summary>
 		public string? RobomergeBotTag { get; init; }
 
+		/// <summary>
+		/// Whether the change was virtualized or not.
+		/// </summary>
+		/// <remarks></remarks>
+		public string? VirtualizedTag { get; init; }
+
 		#endregion -- Public Properties --
 
 		/// <summary>
@@ -116,7 +130,8 @@ namespace EpicGames.Horde.Commits
 				JiraTag = JiraRegex().Match(description).Groups[1].ValueOrNull()?.Trim(),
 				PresubmitTag = PreflightRegex().Match(description).Groups[1].ValueOrNull()?.Trim(),
 				TestTag = TestsRegex().Match(description).Groups[1].ValueOrNull()?.Trim(),
-				ReviewTag = ReviewRegex().Match(description).Groups[1].ValueOrNull()?.Trim()
+				ReviewTag = ReviewRegex().Match(description).Groups[1].ValueOrNull()?.Trim(),
+				VirtualizedTag = Virtualized().Match(description).Groups[0].PresenceOf().ToString().ToLower()
 			};
 		}
 	}
