Horde: Improve logging when saving lease outcomes

diff --git a/Engine/Source/Programs/Horde/Plugins/Compute/HordeServer.Compute/Agents/AgentCollection.cs b/Engine/Source/Programs/Horde/Plugins/Compute/HordeServer.Compute/Agents/AgentCollection.cs
--- a/Engine/Source/Programs/Horde/Plugins/Compute/HordeServer.Compute/Agents/AgentCollection.cs
+++ b/Engine/Source/Programs/Horde/Plugins/Compute/HordeServer.Compute/Agents/AgentCollection.cs
@@ -963,11 +963,13 @@ namespace HordeServer.Agents
 					RpcLease? rpcLease = options.Leases?.FirstOrDefault(x => x.Id == oldLease.Id);
 					if (rpcLease is { State: RpcLeaseState.Completed })
 					{
+						_logger.LogInformation("Setting outcome {Outcome} for completed lease {LeaseId}", oldLease.Id, rpcLease.Outcome);
 						await _leaseCollection.TrySetOutcomeAsync(oldLease.Id, _clock.UtcNow, (LeaseOutcome)rpcLease.Outcome, LeaseOutcomeReason.None, rpcLease.Output.ToArray(), CancellationToken.None);
 					}
 					else
 					{
 						LeaseOutcomeReason reason = rpcLease?.ToNativeOutcomeReason() ?? LeaseOutcomeReason.None;
+						_logger.LogInformation("Setting outcome {Outcome} with reason {Reason} for lease {LeaseId} ({State})", rpcLease?.Outcome, reason, oldLease.Id, rpcLease?.State);
 						await _leaseCollection.TrySetOutcomeAsync(oldLease.Id, _clock.UtcNow, LeaseOutcome.Cancelled, reason, null, CancellationToken.None);
 					}
 				}
