Fix issue when meta filter is not updated when using browser history

#rnx

#Synced-CL 48846087

diff --git a/Engine/Source/Programs/Horde/HordeDashboard/plugins/test-automation/testAutomationView.tsx b/Engine/Source/Programs/Horde/HordeDashboard/plugins/test-automation/testAutomationView.tsx
--- a/Engine/Source/Programs/Horde/HordeDashboard/plugins/test-automation/testAutomationView.tsx
+++ b/Engine/Source/Programs/Horde/HordeDashboard/plugins/test-automation/testAutomationView.tsx
@@ -228,6 +228,7 @@ const MetaKeyChooser: React.FC<{ onUpdateKey: (key: string) => void }> = observe
 const MetaValueChooser: React.FC<{ metakey: string, onRemove: (key: string) => void }> = observer(({ metakey, onRemove }) => {
 
     handler.subscribeToFilter();
+    handler.subscribeToSearch();
 
     const metaType = handler.allStreamTestMetaKeys.includes(metakey)? MetaType.Test : MetaType.Session;
     const state = metaType === MetaType.Test? handler.filterState.testmeta : handler.filterState.metadata;
