[Horde][BuildHealth][Server] BuildHealthFilter collection + API endpoints

#Synced-CL 48788183

diff --git a/Engine/Source/Programs/Horde/Plugins/Build/HordeServer.Build.Tests/BuildHealth/BuildHealthFilterCollectionTests.cs b/Engine/Source/Programs/Horde/Plugins/Build/HordeServer.Build.Tests/BuildHealth/BuildHealthFilterCollectionTests.cs
new file mode 100644
--- /dev/null
+++ b/Engine/Source/Programs/Horde/Plugins/Build/HordeServer.Build.Tests/BuildHealth/BuildHealthFilterCollectionTests.cs
@@ -0,0 +1,307 @@
+// Copyright Epic Games, Inc. All Rights Reserved.
+
+using EpicGames.Horde.BuildHealth;
+using EpicGames.Horde.Projects;
+using EpicGames.Horde.Streams;
+using EpicGames.Horde.Users;
+using HordeServer.Projects;
+using HordeServer.Streams;
+using HordeServer.Utilities;
+using Microsoft.AspNetCore.Mvc;
+
+namespace HordeServer.Tests.BuildHealth
+{
+	[TestClass]
+	public class BuildHealthFilterCollectionTests : BuildTestSetup
+	{
+		private readonly StreamId _streamId = new("ue5-main");
+
+		[TestInitialize]
+		public async Task UpdateConfigAsync()
+		{
+			await UpdateConfigAsync(globalConfig =>
+			{
+				ProjectConfig projectConfig = new() { Id = new ProjectId("ue5") };
+				projectConfig.Streams.Add(new StreamConfig { Id = _streamId });
+				globalConfig.Plugins.GetBuildConfig().Projects.Add(projectConfig);
+			});
+		}
+
+		[TestMethod]
+		public async Task TestAddBuildHealthFilterAsync()
+		{
+			BuildHealthFilterAddRequest request = new BuildHealthFilterAddRequest();
+			ProjectId expectedProject = new ProjectId("ue5");
+			string expectedFilterQuery = "SomeCompletePath?isDebug=true";
+			string expectedFilterName = "NewFilter";
+
+			request.FilterProject = expectedProject;
+			request.FilterQuery = expectedFilterQuery;
+			request.FilterName = expectedFilterName;
+
+			ActionResult<BuildHealthFilterResponse>? response = await BuildHealthFilterController.AddBuildHealthFilterAsync(request);
+
+			Assert.IsNotNull(response);
+			Assert.IsNotNull(response.Value);
+			Assert.AreEqual(expectedProject, response.Value.FilterProject);
+			Assert.AreEqual(expectedFilterQuery, response.Value.FilterQuery);
+			Assert.AreEqual(expectedFilterName, response.Value.FilterName);
+		}
+
+		[TestMethod]
+		public async Task TestNonExistentBuildHealthFilterFailureAsync()
+		{
+			BuildHealthFilterAddRequest request = new BuildHealthFilterAddRequest();
+			ProjectId expectedProject = new ProjectId("ue6");
+			string expectedFilterQuery = GenerateFailureString(BuildHealthFilterRestrictions.MaxFilterQueryLength + 1);
+			string expectedFilterName = GenerateFailureString(BuildHealthFilterRestrictions.MaxFilterNameLength + 1);
+
+			request.FilterProject = expectedProject;
+			request.FilterQuery = expectedFilterQuery;
+			request.FilterName = expectedFilterName;
+
+			ActionResult<BuildHealthFilterResponse>? response = await BuildHealthFilterController.AddBuildHealthFilterAsync(request);
+
+			Assert.IsNotNull(response);
+			Assert.IsNull(response.Value);
+
+			Assert.IsInstanceOfType(response.Result, typeof(BadRequestObjectResult));
+		}
+
+		[TestMethod]
+		public async Task TestAddBuildHealthFilterFailureAsync()
+		{
+			BuildHealthFilterAddRequest request = new BuildHealthFilterAddRequest();
+			ProjectId expectedProject = new ProjectId("ue5");
+			string expectedFilterQuery = GenerateFailureString(BuildHealthFilterRestrictions.MaxFilterQueryLength + 1);
+			string expectedFilterName = GenerateFailureString(BuildHealthFilterRestrictions.MaxFilterNameLength + 1);
+
+			request.FilterProject = expectedProject;
+			request.FilterQuery = expectedFilterQuery;
+			request.FilterName = expectedFilterName;
+
+			ActionResult<BuildHealthFilterResponse>? response = await BuildHealthFilterController.AddBuildHealthFilterAsync(request);
+
+			Assert.IsNotNull(response);
+			Assert.IsNull(response.Value);
+
+			Assert.IsInstanceOfType(response.Result, typeof(BadRequestObjectResult));
+		}
+
+		[TestMethod]
+		public async Task TestUpdatedBuildHealthFilterAsync()
+		{
+			ProjectId project = new ProjectId("ue5");
+
+			// Create initial filter
+			BuildHealthFilterAddRequest request = new BuildHealthFilterAddRequest();
+			string initialFilterQuery = "SomeCompletePath?isDebug=true";
+			string initialFilterName = "NewFilter";
+
+			request.FilterProject = project;
+			request.FilterQuery = initialFilterQuery;
+			request.FilterName = initialFilterName;
+
+			ActionResult<BuildHealthFilterResponse>? response = await BuildHealthFilterController.AddBuildHealthFilterAsync(request);
+			Assert.IsNotNull(response);
+			Assert.IsNotNull(response.Value);
+
+			// Update previous filter
+			BuildHealthFilterUpdateRequest updateRequest = new BuildHealthFilterUpdateRequest();
+			string updatedFilterQuery = "SomeCompletePath?isDebug=false";
+			string updatedFilterName = "SomeNewFilter";
+
+			updateRequest.FilterQuery = updatedFilterQuery;
+			updateRequest.FilterName = updatedFilterName;
+
+			ActionResult<BuildHealthFilterResponse>? updateResponse = await BuildHealthFilterController.UpdateBuildHealthFilterAsync(response.Value.Id, updateRequest);
+
+			Assert.IsNotNull(updateResponse);
+			Assert.IsNotNull(updateResponse.Value);
+			Assert.AreEqual(updatedFilterQuery, updateResponse.Value.FilterQuery);
+			Assert.AreEqual(updatedFilterName, updateResponse.Value.FilterName);
+			Assert.AreEqual(response.Value.Id, updateResponse.Value.Id);
+
+			// Update just filter query
+			string updatedFilterQuery2 = "SomeCompletePath?isDebug=true&stream=false";
+			updateRequest.FilterName = String.Empty;
+			updateRequest.FilterQuery = updatedFilterQuery2;
+
+			ActionResult<BuildHealthFilterResponse>? updateResponse2 = await BuildHealthFilterController.UpdateBuildHealthFilterAsync(response.Value.Id, updateRequest);
+
+			Assert.IsNotNull(updateResponse2);
+			Assert.IsNotNull(updateResponse2.Value);
+			Assert.AreEqual(updatedFilterQuery2, updateResponse2.Value.FilterQuery);
+			Assert.AreEqual(updatedFilterName, updateResponse2.Value.FilterName);
+
+			// Update just filter name
+			string updatedFilterName2 = "SomeNewFilter2";
+			updateRequest.FilterName = updatedFilterName2;
+			updateRequest.FilterQuery = String.Empty;
+
+			ActionResult<BuildHealthFilterResponse>? updateResponse3 = await BuildHealthFilterController.UpdateBuildHealthFilterAsync(updateResponse2.Value.Id, updateRequest);
+
+			Assert.IsNotNull(updateResponse3);
+			Assert.IsNotNull(updateResponse3.Value);
+			Assert.AreEqual(updatedFilterQuery2, updateResponse3.Value.FilterQuery);
+			Assert.AreEqual(updatedFilterName2, updateResponse3.Value.FilterName);
+		}
+
+		[TestMethod]
+		public async Task TestUpdatedBuildHealthFilterFailureAsync()
+		{
+			ProjectId project = new ProjectId("ue5");
+
+			// Create initial filter
+			BuildHealthFilterAddRequest request = new BuildHealthFilterAddRequest();
+			string initialFilterQuery = "SomeCompletePath?isDebug=true";
+			string initialFilterName = "NewFilter";
+
+			request.FilterProject = project;
+			request.FilterQuery = initialFilterQuery;
+			request.FilterName = initialFilterName;
+
+			ActionResult<BuildHealthFilterResponse>? response = await BuildHealthFilterController.AddBuildHealthFilterAsync(request);
+			Assert.IsNotNull(response);
+			Assert.IsNotNull(response.Value);
+
+			// Update previous filter
+			BuildHealthFilterUpdateRequest updateRequest = new BuildHealthFilterUpdateRequest();
+			string updatedFilterQuery = GenerateFailureString(BuildHealthFilterRestrictions.MaxFilterQueryLength + 1);
+			string updatedFilterName = GenerateFailureString(BuildHealthFilterRestrictions.MaxFilterNameLength + 1);
+
+			updateRequest.FilterQuery = updatedFilterQuery;
+			updateRequest.FilterName = updatedFilterName;
+
+			ActionResult<BuildHealthFilterResponse>? updateResponse = await BuildHealthFilterController.UpdateBuildHealthFilterAsync(response.Value.Id, updateRequest);
+			Assert.IsNotNull(updateResponse);
+			Assert.IsNull(updateResponse.Value);
+
+			Assert.IsInstanceOfType(updateResponse.Result, typeof(BadRequestObjectResult));
+
+			// Empty update should obtain nothing
+			updateRequest.FilterQuery = String.Empty;
+			updateRequest.FilterName = String.Empty;
+			ActionResult<BuildHealthFilterResponse>? emptyUpdateResponse = await BuildHealthFilterController.UpdateBuildHealthFilterAsync(response.Value.Id, updateRequest);
+			Assert.IsNotNull(emptyUpdateResponse);
+			Assert.IsNull(emptyUpdateResponse.Value);
+		}
+
+		[TestMethod]
+		public async Task TestGetBuildHealthFilterAsync()
+		{
+			BuildHealthFilterAddRequest request = new BuildHealthFilterAddRequest();
+			ProjectId expectedProject = new ProjectId("ue5");
+			string expectedFilterQuery = "SomeFilter";
+			string expectedFilterName = "SomeName";
+
+			request.FilterProject = expectedProject;
+			request.FilterQuery = expectedFilterQuery;
+			request.FilterName = expectedFilterName;
+
+			// Test inidvidiaul get
+			ActionResult<BuildHealthFilterResponse>? response = await BuildHealthFilterController.AddBuildHealthFilterAsync(request);
+			Assert.IsNotNull(response);
+			Assert.IsNotNull(response.Value);
+
+			ActionResult<BuildHealthFilterResponse>? getResponse = await BuildHealthFilterController.GetBuildHealthFilterAsync(response!.Value!.Id);
+
+			Assert.IsNotNull(getResponse);
+			Assert.IsNotNull(getResponse.Value);
+			Assert.AreEqual(expectedFilterName, getResponse.Value.FilterName);
+			Assert.AreEqual(expectedFilterQuery, getResponse.Value.FilterQuery);
+			Assert.AreEqual(response.Value.Id, getResponse.Value.Id);
+
+			// Test project get
+			for (int i = 0; i < 4; i++)
+			{
+				_ = await BuildHealthFilterController.AddBuildHealthFilterAsync(request);
+			}
+
+			ActionResult<List<BuildHealthFilterResponse>> results = await BuildHealthFilterController.GetBuildHealthFiltersAsync(expectedProject);
+
+			// Verify all the results are present
+			Assert.IsNotNull(results);
+			Assert.IsNotNull(results.Value);
+			Assert.AreEqual(5, results.Value.Count);
+		}
+
+		[TestMethod]
+		public async Task TestGetDeleteBuildHealthFilterAsync()
+		{
+			ProjectId project = new ProjectId("ue5");
+
+			// Create initial filter
+			BuildHealthFilterAddRequest request = new BuildHealthFilterAddRequest();
+			string initialFilterQuery = "SomeCompletePath?isDebug=true";
+			string initialFilterName = "NewFilter";
+
+			request.FilterProject = project;
+			request.FilterQuery = initialFilterQuery;
+			request.FilterName = initialFilterName;
+
+			ActionResult<BuildHealthFilterResponse>? response = await BuildHealthFilterController.AddBuildHealthFilterAsync(request);
+			Assert.IsNotNull(response);
+			Assert.IsNotNull(response.Value);
+
+			// Issue delete
+			ActionResult deleteResponse = await BuildHealthFilterController.DeleteBuildHealthFiltersAsync(response.Value.Id);
+			Assert.IsNotNull(deleteResponse);
+
+			// Verify delete
+			ActionResult<BuildHealthFilterResponse>? reentrantGetResponse = await BuildHealthFilterController.GetBuildHealthFilterAsync(response.Value.Id);
+			Assert.IsNotNull(reentrantGetResponse);
+			Assert.IsNull(reentrantGetResponse.Value);
+		}
+
+		[TestMethod]
+		public async Task TestOwnershipSemanticsAsync()
+		{
+			ProjectId project = new ProjectId("ue5");
+
+			// Create initial filter
+			BuildHealthFilterAddRequest request = new BuildHealthFilterAddRequest();
+			string initialFilterQuery = "SomeCompletePath?isDebug=true";
+			string initialFilterName = "NewFilter";
+
+			request.FilterProject = project;
+			request.FilterQuery = initialFilterQuery;
+			request.FilterName = initialFilterName;
+
+			UserId originalUser = new UserId(BinaryIdUtils.CreateNew());
+
+			IBuildHealthFilter? result = await BuildHealthFilterCollection.AddBuildHealthFilterAsync(originalUser, request, CancellationToken.None);
+			Assert.IsNotNull(result);
+			Assert.AreEqual(originalUser, result.Owner);
+
+			// Modify filter of other user
+			// Update previous filter
+			UserId newUser = new UserId(BinaryIdUtils.CreateNew());
+			BuildHealthFilterUpdateRequest updateRequest = new BuildHealthFilterUpdateRequest();
+			string updatedFilterQuery = GenerateFailureString(BuildHealthFilterRestrictions.MaxFilterQueryLength + 1);
+			string updatedFilterName = GenerateFailureString(BuildHealthFilterRestrictions.MaxFilterNameLength + 1);
+
+			updateRequest.FilterQuery = updatedFilterQuery;
+			updateRequest.FilterName = updatedFilterName;
+
+			IBuildHealthFilter? failedUpdateResult = await BuildHealthFilterCollection.UpdateBuildHealthFilterAsync(result.Id, newUser, updateRequest, CancellationToken.None);
+			Assert.IsNull(failedUpdateResult);
+
+			// Delete filter of other user
+			bool failedDeleteResult = await BuildHealthFilterCollection.DeleteBuildHealthFilterAsync(result.Id, newUser, CancellationToken.None);
+			Assert.IsFalse(failedDeleteResult);
+		}
+
+		private static string GenerateFailureString(int length)
+		{
+			const string Chunk = "ABCDE";
+
+			if (length <= 0)
+			{
+				return String.Empty;
+			}
+
+			return String.Concat(Enumerable.Repeat(Chunk, (length / Chunk.Length) + 1)).Substring(0, length);
+		}
+	}
+}
diff --git a/Engine/Source/Programs/Horde/Plugins/Build/HordeServer.Build.Tests/BuildTestSetup.cs b/Engine/Source/Programs/Horde/Plugins/Build/HordeServer.Build.Tests/BuildTestSetup.cs
--- a/Engine/Source/Programs/Horde/Plugins/Build/HordeServer.Build.Tests/BuildTestSetup.cs
+++ b/Engine/Source/Programs/Horde/Plugins/Build/HordeServer.Build.Tests/BuildTestSetup.cs
@@ -7,6 +7,7 @@ using Amazon.EC2;
 using EpicGames.Horde.Agents;
 using EpicGames.Horde.Agents.Pools;
 using EpicGames.Horde.Artifacts;
+using EpicGames.Horde.BuildHealth;
 using EpicGames.Horde.Jobs.Templates;
 using EpicGames.Horde.Streams;
 using EpicGames.Horde.Tools;
@@ -21,6 +22,7 @@ using HordeServer.Agents.Relay;
 using HordeServer.Agents.Telemetry;
 using HordeServer.Agents.Utilization;
 using HordeServer.Artifacts;
+using HordeServer.Build.BuildHealth;
 using HordeServer.Commits;
 using HordeServer.Compute;
 using HordeServer.Dashboard;
@@ -35,7 +37,6 @@ using HordeServer.Jobs.TestData;
 using HordeServer.Jobs.Timing;
 using HordeServer.Logs;
 using HordeServer.Notifications;
-using HordeServer.VersionControl.Perforce;
 using HordeServer.Server;
 using HordeServer.Storage;
 using HordeServer.Streams;
@@ -45,6 +46,7 @@ using HordeServer.Tools;
 using HordeServer.Ugs;
 using HordeServer.Users;
 using HordeServer.Utilities;
+using HordeServer.VersionControl.Perforce;
 using Microsoft.AspNetCore.Http;
 using Microsoft.AspNetCore.Mvc;
 using Microsoft.Extensions.Configuration;
@@ -78,6 +80,7 @@ namespace HordeServer.Tests
 		public IDeviceCollection DeviceCollection => ServiceProvider.GetRequiredService<IDeviceCollection>();
 		public IDashboardPreviewCollection DashboardPreviewCollection => ServiceProvider.GetRequiredService<IDashboardPreviewCollection>();
 		public IBisectTaskCollection BisectTaskCollection => ServiceProvider.GetRequiredService<IBisectTaskCollection>();
+		public IBuildHealthFilterCollection BuildHealthFilterCollection => ServiceProvider.GetRequiredService<IBuildHealthFilterCollection>();
 
 		public FleetService FleetService => ServiceProvider.GetRequiredService<FleetService>();
 		public AgentService AgentService => ServiceProvider.GetRequiredService<AgentService>();
@@ -108,6 +111,7 @@ namespace HordeServer.Tests
 		public DashboardController DashboardController => GetDashboardController();
 		public TestDataController TestDataController => GetTestDataController();
 		public BisectTasksController BisectTasksController => GetBisectTasksController();
+		public BuildHealthFilterController BuildHealthFilterController => GetBuildHealthFilterController();
 
 		public BuildTestSetup()
 		{
@@ -241,6 +245,13 @@ namespace HordeServer.Tests
 			return bisectCtrl;
 		}
 
+		private BuildHealthFilterController GetBuildHealthFilterController()
+		{
+			BuildHealthFilterController buildHealthFilterCtrl = ActivatorUtilities.CreateInstance<BuildHealthFilterController>(ServiceProvider);
+			buildHealthFilterCtrl.ControllerContext = GetControllerContext();
+			return buildHealthFilterCtrl;
+		}
+
 		private AgentsController GetAgentsController()
 		{
 			AgentsController agentCtrl = ActivatorUtilities.CreateInstance<AgentsController>(ServiceProvider);
@@ -314,18 +325,18 @@ namespace HordeServer.Tests
 			{
 				tempProps.Add(KnownPropertyNames.AwsInstanceId + "=" + awsInstanceId);
 			}
-			
+
 			IAgent? agent = await AgentService.CreateAgentAsync(new CreateAgentOptions(new AgentId("TestAgent" + s_agentIdCounter++), AgentMode.Dedicated, ephemeral, ""));
 			Assert.IsNotNull(agent);
 
 			agent = await agent.TryUpdateAsync(new UpdateAgentOptions { Enabled = enabled, ExplicitPools = poolId != null ? [poolId.Value] : [] });
 			Assert.IsNotNull(agent);
-			
+
 			if (status != null)
 			{
 				agent = await AgentService.CreateSessionAsync(agent, new RpcAgentCapabilities(tempProps), null);
 				Assert.IsNotNull(agent);
-				
+
 				if (status != AgentStatus.Ok)
 				{
 					agent = await agent.TryUpdateSessionAsync(new UpdateSessionOptions { Status = status });
@@ -351,18 +362,18 @@ namespace HordeServer.Tests
 			Clock.UtcNow = now;
 			return agent;
 		}
-		
+
 		protected Task<IPool> CreatePoolAsync(PoolConfig poolConfig, bool viaConfig = true)
 		{
 			return viaConfig ? CreatePoolViaConfigAsync(poolConfig) : CreatePoolViaDatabaseAsync(poolConfig);
 		}
-		
+
 		protected async Task<IPool> CreatePoolViaConfigAsync(PoolConfig poolConfig)
 		{
 			await UpdateConfigAsync(config => config.Plugins.GetComputeConfig().Pools.Add(poolConfig));
 			return await PoolCollection.GetAsync(poolConfig.Id) ?? throw new Exception($"Unable to get pool ID {poolConfig.Id}");
 		}
-		
+
 		protected async Task<IPool> CreatePoolViaDatabaseAsync(PoolConfig poolConfig)
 		{
 #pragma warning disable CS0618 // Type or member is obsolete
diff --git a/Engine/Source/Programs/Horde/Plugins/Build/HordeServer.Build/BuildHealth/BuildHealthFilterAclAction.cs b/Engine/Source/Programs/Horde/Plugins/Build/HordeServer.Build/BuildHealth/BuildHealthFilterAclAction.cs
new file mode 100644
--- /dev/null
+++ b/Engine/Source/Programs/Horde/Plugins/Build/HordeServer.Build/BuildHealth/BuildHealthFilterAclAction.cs
@@ -0,0 +1,17 @@
+// Copyright Epic Games, Inc. All Rights Reserved.
+
+using EpicGames.Horde.Acls;
+
+namespace HordeServer.Build.BuildHealth
+{
+	/// <summary>
+	/// ACL Actions used for build health filters.
+	/// </summary>
+	public static class BuildHealthFilterAclAction
+	{
+		/// <summary>
+		/// Ability to add, update, and delete build health filters.
+		/// </summary>
+		public static AclAction AddUpdateDeleteBuildHealthFilter { get; } = new AclAction("AddUpdateDeleteBuildHealthFilter");
+	}
+}
diff --git a/Engine/Source/Programs/Horde/Plugins/Build/HordeServer.Build/BuildHealth/BuildHealthFilterCollection.cs b/Engine/Source/Programs/Horde/Plugins/Build/HordeServer.Build/BuildHealth/BuildHealthFilterCollection.cs
new file mode 100644
--- /dev/null
+++ b/Engine/Source/Programs/Horde/Plugins/Build/HordeServer.Build/BuildHealth/BuildHealthFilterCollection.cs
@@ -0,0 +1,167 @@
+// Copyright Epic Games, Inc. All Rights Reserved.
+
+using EpicGames.Horde.BuildHealth;
+using EpicGames.Horde.Projects;
+using EpicGames.Horde.Users;
+using HordeServer.Server;
+using HordeServer.Utilities;
+using Microsoft.Extensions.Logging;
+using MongoDB.Bson;
+using MongoDB.Bson.Serialization.Attributes;
+using MongoDB.Driver;
+
+namespace HordeServer.Build.BuildHealth
+{
+	/// <summary>
+	/// Build Health Filter Collection implementation.
+	/// </summary>
+	class BuildHealthFilterCollection : IBuildHealthFilterCollection
+	{
+		#region -- Private Members --
+
+		readonly IMongoCollection<BuildHealthFilterDocument> _buildHealthFilterDocuments;
+		readonly ILogger<BuildHealthFilterCollection> _logger;
+
+		#endregion -- Private Members --
+
+		#region -- Constructor --
+
+		/// <summary>
+		/// Constructor.
+		/// </summary>
+		/// <param name="mongoService">The underlying mongo service to use.</param>
+		/// <param name="logger">The logger.</param>
+		public BuildHealthFilterCollection(IMongoService mongoService, ILogger<BuildHealthFilterCollection> logger)
+		{
+			_buildHealthFilterDocuments = mongoService.GetCollection<BuildHealthFilterDocument>("BuildHealthFilters");
+			_logger = logger;
+		}
+
+		#endregion -- Constructor --
+
+		#region -- Interface API --
+
+		/// <inheritdoc/>
+		public async Task<IBuildHealthFilter?> AddBuildHealthFilterAsync(UserId? owner, BuildHealthFilterAddRequest filterRequest, CancellationToken cancellationToken)
+		{
+			if (!filterRequest.IsValidAddRequest())
+			{
+				_logger.LogWarning("Invalid add build health filter request. IsValidFilterQuery: {IsValidFilterQuery}; IsValidFilterName: {IsValidFilterQuery}", filterRequest.IsValidFilterQuery(), filterRequest.IsValidFilterName());
+				return null;
+			}
+
+			BuildHealthFilterDocument newBuildHealthFilter = new BuildHealthFilterDocument() { Id = new BuildHealthFilterId(BinaryIdUtils.CreateNew()), Owner = owner, FilterName = filterRequest.FilterName, FilterProject = filterRequest.FilterProject, FilterQuery = filterRequest.FilterQuery, UpdateTimeUtc = DateTime.UtcNow };
+
+			await _buildHealthFilterDocuments.InsertOneAsync(newBuildHealthFilter, null, cancellationToken);
+			return newBuildHealthFilter;
+		}
+
+		/// <inheritdoc/>
+		public async Task<IBuildHealthFilter?> UpdateBuildHealthFilterAsync(BuildHealthFilterId filterId, UserId? owner, BuildHealthFilterUpdateRequest filterRequest, CancellationToken cancellationToken)
+		{
+			if (!filterRequest.IsValidUpdateRequest())
+			{
+				_logger.LogWarning("Invalid update build health filter request. IsValidFilterQuery: {IsValidFilterQuery}; IsValidFilterName: {IsValidFilterQuery}", filterRequest.IsValidFilterQuery(), filterRequest.IsValidFilterName());
+				return null;
+			}
+
+			FilterDefinition<BuildHealthFilterDocument> filter = Builders<BuildHealthFilterDocument>.Filter.Expr(x => x.Id == filterId) & Builders<BuildHealthFilterDocument>.Filter.Expr(x => (x.Owner == null || x.Owner == owner));
+			List<UpdateDefinition<BuildHealthFilterDocument>> updates = new List<UpdateDefinition<BuildHealthFilterDocument>>();
+			if (!String.IsNullOrEmpty(filterRequest.FilterName))
+			{
+				updates.Add(Builders<BuildHealthFilterDocument>.Update.Set(x => x.FilterName, filterRequest.FilterName));
+			}
+			if (!String.IsNullOrEmpty(filterRequest.FilterQuery))
+			{
+				updates.Add
+					(Builders<BuildHealthFilterDocument>.Update.Set(x => x.FilterQuery, filterRequest.FilterQuery));
+			}
+
+			// We have no updates, just fetch the item for the requestor.
+			if (updates.Count == 0)
+			{
+				_logger.LogInformation("No updates to apply for filterId ({FilterId}). Returning.", filterId);
+				return null;
+			}
+
+			updates.Add(Builders<BuildHealthFilterDocument>.Update.Set(x => x.UpdateTimeUtc, DateTime.UtcNow));
+
+			UpdateDefinition<BuildHealthFilterDocument> update = updates.Count > 0 ? Builders<BuildHealthFilterDocument>.Update.Combine(updates) : Builders<BuildHealthFilterDocument>.Update.Combine();
+			IBuildHealthFilter? buildHealthFilter = await _buildHealthFilterDocuments.FindOneAndUpdateAsync(filter, update, new FindOneAndUpdateOptions<BuildHealthFilterDocument, BuildHealthFilterDocument> { ReturnDocument = ReturnDocument.After }, cancellationToken);
+
+			if (buildHealthFilter == null)
+			{
+				_logger.LogInformation("Unable to update Build Health Filter with Id: {BuildHealthFilterId} with Owner: {OwnerId}.", filterId, owner);
+			}
+
+			return buildHealthFilter;
+		}
+
+		/// <inheritdoc/>
+		public async Task<IBuildHealthFilter?> GetBuildHealthFilterAsync(BuildHealthFilterId filterId, CancellationToken cancellationToken)
+		{
+			FilterDefinition<BuildHealthFilterDocument> filter = Builders<BuildHealthFilterDocument>.Filter.Expr(x => x.Id == filterId);
+			IBuildHealthFilter document = await _buildHealthFilterDocuments.Find(filter).FirstOrDefaultAsync<BuildHealthFilterDocument>(cancellationToken);
+
+			return document;
+		}
+
+		/// <inheritdoc/>
+		public async Task<IEnumerable<IBuildHealthFilter>> GetBuildHealthFiltersAsync(ProjectId projectId, CancellationToken cancellationToken)
+		{
+			FilterDefinitionBuilder<BuildHealthFilterDocument> filterBuilder = Builders<BuildHealthFilterDocument>.Filter;
+			FilterDefinition<BuildHealthFilterDocument> filter = filterBuilder.Empty;
+
+			filter &= filterBuilder.Eq(x => x.FilterProject, projectId);
+
+			List<BuildHealthFilterDocument> documents = await _buildHealthFilterDocuments.Find(filter).ToListAsync(cancellationToken);
+
+			return documents;
+		}
+
+		/// <inheritdoc/>
+		public async Task<bool> DeleteBuildHealthFilterAsync(BuildHealthFilterId filterId, UserId? owner, CancellationToken cancellationToken)
+		{
+			FilterDefinition<BuildHealthFilterDocument> filter = Builders<BuildHealthFilterDocument>.Filter.Expr(x => x.Id == filterId && (x.Owner == null || x.Owner == owner));
+			DeleteResult result = await _buildHealthFilterDocuments.DeleteOneAsync(filter, cancellationToken);
+
+			return result.DeletedCount == 1;
+		}
+
+		#endregion -- Interface API --
+
+		#region -- Internal Types --
+
+		/// <summary>
+		/// Build health filter document.
+		/// </summary>
+		class BuildHealthFilterDocument : IBuildHealthFilter
+		{
+			/// <inheritdoc/>
+			[BsonRequired, BsonId]
+			public BuildHealthFilterId Id { get; set; }
+
+			/// <inheritdoc/>
+			[BsonElement("fn")]
+			public string FilterName { get; set; } = default!;
+
+			/// <inheritdoc/>
+			[BsonElement("fp")]
+			public ProjectId FilterProject { get; set; }
+
+			/// <inheritdoc/>
+			[BsonElement("fq")]
+			public string FilterQuery { get; set; } = default!;
+
+			/// <inheritdoc/>
+			[BsonElement("o")]
+			public UserId? Owner { get; set; }
+
+			/// <inheritdoc/>
+			[BsonElement("ut")]
+			public DateTime UpdateTimeUtc { get; set; }
+		}
+
+		#endregion -- Internal Types --
+	}
+}
diff --git a/Engine/Source/Programs/Horde/Plugins/Build/HordeServer.Build/BuildHealth/BuildHealthFilterController.cs b/Engine/Source/Programs/Horde/Plugins/Build/HordeServer.Build/BuildHealth/BuildHealthFilterController.cs
new file mode 100644
--- /dev/null
+++ b/Engine/Source/Programs/Horde/Plugins/Build/HordeServer.Build/BuildHealth/BuildHealthFilterController.cs
@@ -0,0 +1,294 @@
+// Copyright Epic Games, Inc. All Rights Reserved.
+
+using EpicGames.Horde.BuildHealth;
+using EpicGames.Horde.Projects;
+using EpicGames.Horde.Users;
+using HordeServer.Projects;
+using HordeServer.Utilities;
+using Microsoft.AspNetCore.Authorization;
+using Microsoft.AspNetCore.Http;
+using Microsoft.AspNetCore.Mvc;
+using Microsoft.Extensions.Options;
+
+namespace HordeServer.Build.BuildHealth
+{
+	/// <summary>
+	/// Controller for build health filters.
+	/// </summary>
+	[ApiController]
+	[Authorize]
+	[Route("[controller]")]
+	public class BuildHealthFilterController : HordeControllerBase
+	{
+
+		#region -- Private Members --
+
+		readonly IBuildHealthFilterCollection _buildHealthFilterCollection;
+		private readonly IOptionsSnapshot<BuildConfig> _buildConfig;
+
+		#endregion -- Private Members --
+
+		#region -- Constructor --
+
+		/// <summary>
+		/// Constructor.
+		/// </summary>
+		/// <param name="buildHealthFilterCollection">The underyling build health filter collection.</param>
+		/// <param name="buildConfig">Build config used in accessing relevant ACL configurations.</param>
+		public BuildHealthFilterController(IBuildHealthFilterCollection buildHealthFilterCollection, IOptionsSnapshot<BuildConfig> buildConfig)
+		{
+			_buildHealthFilterCollection = buildHealthFilterCollection;
+			_buildConfig = buildConfig;
+		}
+
+		#endregion -- Constructor --
+
+		#region -- Public API --
+
+		/// <summary>
+		/// Adds a new build health filter given the request.
+		/// </summary>
+		/// <param name="filterUpdateRequest">The add request.</param>
+		/// <param name="cancellationToken">The cancellation token to use throughout the operation.</param>
+		/// <returns>The newly created filter if applicable.</returns>
+		/// <remarks>The requesting user will become the implicit owner of the new filter.</remarks>
+		[HttpPost]
+		[Route("/api/v1/buildhealth/filters")]
+		[ProducesResponseType(typeof(BuildHealthFilterResponse), 200)]
+		public async Task<ActionResult<BuildHealthFilterResponse>?> AddBuildHealthFilterAsync([FromBody] BuildHealthFilterAddRequest filterUpdateRequest, CancellationToken cancellationToken = default)
+		{
+			// Validate inputs
+			if (String.IsNullOrEmpty(filterUpdateRequest.FilterQuery))
+			{
+				return BadRequest("No filter query provided. Filter query is a required parameter.");
+			}
+
+			if (String.IsNullOrEmpty(filterUpdateRequest.FilterName))
+			{
+				return BadRequest("No filter name provided. Filter name is a required parameter.");
+			}
+
+			if (!filterUpdateRequest.IsValidAddRequest())
+			{
+				return BadRequest($"UpdateFilterRequest was not valid. Check filter name length (Max:{BuildHealthFilterRestrictions.MaxFilterNameLength}) & filter query length (Max:{BuildHealthFilterRestrictions.MaxFilterQueryLength}).");
+			}
+
+			// Validate user has access for add operation.
+			if (!_buildConfig.Value.Authorize(BuildHealthFilterAclAction.AddUpdateDeleteBuildHealthFilter, User))
+			{
+				return Forbid(BuildHealthFilterAclAction.AddUpdateDeleteBuildHealthFilter);
+			}
+
+			// Validate project exists, and user has access for project
+			(bool isValidationSuccessful, ActionResult actionResult) = VerifyAndAuthorizeProject(filterUpdateRequest.FilterProject);
+			if (!isValidationSuccessful)
+			{
+				return actionResult;
+			}
+
+			UserId? requestingUser = User.GetUserId();
+			IBuildHealthFilter? filter = await _buildHealthFilterCollection.AddBuildHealthFilterAsync(requestingUser, filterUpdateRequest, cancellationToken);
+
+			if (filter == null)
+			{
+				return StatusCode(StatusCodes.Status500InternalServerError);
+
+			}
+
+			return new BuildHealthFilterResponse(filter);
+		}
+
+		/// <summary>
+		/// Updates the filter based off the request.
+		/// </summary>
+		/// <param name="filterId">The id of the filter to update.</param>
+		/// <param name="filterUpdateRequest">The update request.</param>
+		/// <param name="cancellationToken">The cancellation token to use throughout the operation.</param>
+		/// <returns>The updated filter if applicable.</returns>
+		/// <remarks>Users can only update filters that don't have an owner, or ones they are an owner of.</remarks>
+		[HttpPut]
+		[Route("/api/v1/buildhealth/filters/{filterId}")]
+		[ProducesResponseType(typeof(BuildHealthFilterResponse), 200)]
+		public async Task<ActionResult<BuildHealthFilterResponse>?> UpdateBuildHealthFilterAsync(BuildHealthFilterId filterId, [FromBody] BuildHealthFilterUpdateRequest filterUpdateRequest, CancellationToken cancellationToken = default)
+		{
+			// Validate inputs
+			if (String.IsNullOrEmpty(filterUpdateRequest.FilterName) && String.IsNullOrEmpty(filterUpdateRequest.FilterQuery))
+			{
+				return BadRequest("Niether filter name or filter query were provided. Please provide at least one in order to update the filter.");
+			}
+
+			if (!filterUpdateRequest.IsValidUpdateRequest())
+			{
+				return BadRequest($"UpdateFilterRequest was not valid. Check filter name length (Max:{BuildHealthFilterRestrictions.MaxFilterNameLength}) & filter query length (Max:{BuildHealthFilterRestrictions.MaxFilterQueryLength}).");
+			}
+
+			// Validate user has access for update operation.
+			if (!_buildConfig.Value.Authorize(BuildHealthFilterAclAction.AddUpdateDeleteBuildHealthFilter, User))
+			{
+				return Forbid(BuildHealthFilterAclAction.AddUpdateDeleteBuildHealthFilter);
+			}
+
+			// Verify filter exists.
+			IBuildHealthFilter? existingFilter = await _buildHealthFilterCollection.GetBuildHealthFilterAsync(filterId, cancellationToken);
+			if (existingFilter == null)
+			{
+				return NotFound($"Could not find BuildHealthFilter with Id: {filterId}");
+			}
+
+			// Validate project exists, and user has access for project
+			(bool isValidationSuccessful, ActionResult actionResult) = VerifyAndAuthorizeProject(existingFilter.FilterProject);
+			if (!isValidationSuccessful)
+			{
+				return actionResult;
+			}
+
+			// Verify the user is able to edit this filter.
+			UserId? requestingUser = User.GetUserId();
+			if (existingFilter.Owner != null && existingFilter.Owner != requestingUser)
+			{
+				return Forbid($"User is not the owner of BuildHealthFilter Id: {filterId}.");
+			}
+
+			IBuildHealthFilter? filter = await _buildHealthFilterCollection.UpdateBuildHealthFilterAsync(filterId, requestingUser, filterUpdateRequest, cancellationToken);
+
+			if (filter == null)
+			{
+				return StatusCode(StatusCodes.Status500InternalServerError);
+
+			}
+
+			return new BuildHealthFilterResponse(filter);
+		}
+
+		/// <summary>
+		/// Gets filter with specific filter id.
+		/// </summary>
+		/// <param name="filterId">The filter id to filter on.</param>
+		/// <param name="cancellationToken">The cancellation token to use throughout the operation.</param>
+		/// <returns>The filter if it was found.</returns>
+		[HttpGet]
+		[Route("/api/v1/buildhealth/filters/{filterId}")]
+		[ProducesResponseType(typeof(BuildHealthFilterResponse), 200)]
+		public async Task<ActionResult<BuildHealthFilterResponse>?> GetBuildHealthFilterAsync(BuildHealthFilterId filterId, CancellationToken cancellationToken = default)
+		{
+			IBuildHealthFilter? filter = await _buildHealthFilterCollection.GetBuildHealthFilterAsync(filterId, cancellationToken);
+			if (filter == null)
+			{
+				return NotFound(filterId);
+			}
+
+			// Validate project exists, and user has access for project
+			(bool isValidationSuccessful, ActionResult actionResult) = VerifyAndAuthorizeProject(filter.FilterProject);
+			if (!isValidationSuccessful)
+			{
+				return actionResult;
+			}
+
+			return new BuildHealthFilterResponse(filter);
+		}
+
+		/// <summary>
+		/// Gets all applicable filters given the filter criteria.
+		/// </summary>
+		/// <param name="projectId">Project id to filter on.</param>
+		/// <param name="cancellationToken">The cancellation token to use throughout the operation.</param>
+		/// <returns>The list of filters that match.</returns>
+		[HttpGet]
+		[Route("/api/v1/buildhealth/filters")]
+		[ProducesResponseType(typeof(List<BuildHealthFilterResponse>), 200)]
+		public async Task<ActionResult<List<BuildHealthFilterResponse>>> GetBuildHealthFiltersAsync([FromQuery] ProjectId projectId, CancellationToken cancellationToken = default)
+		{
+			// Validate project exists, and user has access for project
+			(bool isValidationSuccessful, ActionResult actionResult) = VerifyAndAuthorizeProject(projectId);
+			if (!isValidationSuccessful)
+			{
+				return actionResult;
+			}
+
+			IEnumerable<IBuildHealthFilter> filters = await _buildHealthFilterCollection.GetBuildHealthFiltersAsync(projectId, cancellationToken);
+			List<BuildHealthFilterResponse> returnObjects = new List<BuildHealthFilterResponse>(filters.Count());
+
+			foreach (IBuildHealthFilter filter in filters)
+			{
+				returnObjects.Add(new BuildHealthFilterResponse(filter));
+			}
+
+			return returnObjects;
+		}
+
+		/// <summary>
+		/// Deletes a filter.
+		/// </summary>
+		/// <param name="filterId">The id of the filter to delete.</param>
+		/// <param name="cancellationToken">The cancellation token to use throughout the operation.</param>
+		/// <returns></returns>
+		[HttpDelete]
+		[Route("/api/v1/buildhealth/filters/{filterId}")]
+		public async Task<ActionResult> DeleteBuildHealthFiltersAsync(BuildHealthFilterId filterId, CancellationToken cancellationToken = default)
+		{
+			// Validate user has access for delete operation.
+			if (!_buildConfig.Value.Authorize(BuildHealthFilterAclAction.AddUpdateDeleteBuildHealthFilter, User))
+			{
+				return Forbid(BuildHealthFilterAclAction.AddUpdateDeleteBuildHealthFilter);
+			}
+
+			// Verify filter exists
+			IBuildHealthFilter? existingFilter = await _buildHealthFilterCollection.GetBuildHealthFilterAsync(filterId, cancellationToken);
+
+			if (existingFilter == null)
+			{
+				return NotFound($"Could not find BuildHealthFilter with Id: {filterId}");
+			}
+
+			// Validate project exists, and user has access for project
+			(bool isValidationSuccessful, ActionResult actionResult) = VerifyAndAuthorizeProject(existingFilter.FilterProject);
+			if (!isValidationSuccessful)
+			{
+				return actionResult;
+			}
+
+			// Verify user is able to delete this filter.
+			UserId? requestingUser = User.GetUserId();
+			if (existingFilter.Owner != null && existingFilter.Owner != requestingUser)
+			{
+				return Forbid("User is not the owner of BuildHealthFilter Id: {BuildHealthFilterId}.", existingFilter.Id);
+			}
+
+			if (!await _buildHealthFilterCollection.DeleteBuildHealthFilterAsync(filterId, requestingUser, cancellationToken))
+			{
+				return BadRequest($"Unable to delete filter ({filterId}).");
+			}
+
+			return Ok();
+		}
+
+		#endregion -- Public API --
+
+		#region -- Private API --
+
+		/// <summary>
+		/// Helper method to verify and authorize project.
+		/// </summary>
+		/// <param name="projectId">The projectId to verify the user against.</param>
+		/// <returns>Tuple representing whether validation was success, and the action result. Action result is Ok() on success.</returns>
+		private (bool isValidationSuccessful, ActionResult actionResult) VerifyAndAuthorizeProject(ProjectId projectId)
+		{
+			// Verify project.
+			ProjectConfig? projectConfig;
+			if (!_buildConfig.Value.TryGetProject(projectId, out projectConfig))
+			{
+				return (isValidationSuccessful: false, actionResult: NotFound(projectId));
+			}
+
+			// Verify user can view details around this project.
+			if (!projectConfig.Authorize(ProjectAclAction.ViewProject, User))
+			{
+				return (isValidationSuccessful: false, actionResult: NotFound(ProjectAclAction.ViewProject));
+			}
+
+			return (isValidationSuccessful: true, actionResult: Ok());
+		}
+
+		#endregion
+	}
+}
diff --git a/Engine/Source/Programs/Horde/Plugins/Build/HordeServer.Build/BuildPlugin.cs b/Engine/Source/Programs/Horde/Plugins/Build/HordeServer.Build/BuildPlugin.cs
--- a/Engine/Source/Programs/Horde/Plugins/Build/HordeServer.Build/BuildPlugin.cs
+++ b/Engine/Source/Programs/Horde/Plugins/Build/HordeServer.Build/BuildPlugin.cs
@@ -3,6 +3,7 @@
 using System.Diagnostics.CodeAnalysis;
 using EpicGames.Core;
 using EpicGames.Horde.Artifacts;
+using EpicGames.Horde.BuildHealth;
 using EpicGames.Horde.Devices;
 using EpicGames.Horde.Jobs;
 using EpicGames.Horde.Jobs.Templates;
@@ -14,6 +15,7 @@ using HordeServer.Agents.Pools;
 using HordeServer.Agents.Utilization;
 using HordeServer.Artifacts;
 using HordeServer.Auditing;
+using HordeServer.Build.BuildHealth;
 using HordeServer.Commits;
 using HordeServer.Configuration;
 using HordeServer.Devices;
@@ -104,6 +106,7 @@ namespace HordeServer
 			services.AddSingleton<ITestDataCollectionV2, TestDataCollectionV2>();
 			services.AddSingleton<IUtilizationDataCollection, UtilizationDataCollection>();
 			services.AddSingleton<ITemplateCollection, TemplateCollection>();
+			services.AddSingleton<IBuildHealthFilterCollection, BuildHealthFilterCollection>();
 
 			services.AddSingleton<IPoolSizeStrategyFactory, JobQueueStrategyFactory>();
 			services.AddHostedService(sp => sp.GetRequiredService<ArtifactCollection>());
diff --git a/Engine/Source/Programs/Shared/EpicGames.Horde/BuildHealth/BuildHealthFilterId.cs b/Engine/Source/Programs/Shared/EpicGames.Horde/BuildHealth/BuildHealthFilterId.cs
new file mode 100644
--- /dev/null
+++ b/Engine/Source/Programs/Shared/EpicGames.Horde/BuildHealth/BuildHealthFilterId.cs
@@ -0,0 +1,40 @@
+// Copyright Epic Games, Inc. All Rights Reserved.
+
+using System.ComponentModel;
+using EpicGames.Core;
+
+namespace EpicGames.Horde.BuildHealth
+{
+	/// <summary>
+	/// Identifier for a build health filter.
+	/// </summary>
+	/// <param name="Id">Id to construct from</param>
+	[JsonSchemaString]
+	[TypeConverter(typeof(BinaryIdTypeConverter<BuildHealthFilterId, FilterIdConverter>))]
+	[BinaryIdConverter(typeof(FilterIdConverter))]
+	public readonly record struct BuildHealthFilterId(BinaryId Id)
+	{
+		/// <summary>
+		/// Constant value for an empty build health filter id.
+		/// </summary>
+		public static BuildHealthFilterId Empty { get; } = default;
+
+		/// <inheritdoc cref="BinaryId.Parse(System.String)"/>
+		public static BuildHealthFilterId Parse(string text) => new BuildHealthFilterId(BinaryId.Parse(text));
+
+		/// <inheritdoc/>
+		public override string ToString() => Id.ToString();
+	}
+
+	/// <summary>
+	/// Converter to and from <see cref="BinaryId"/> instances.
+	/// </summary>
+	class FilterIdConverter : BinaryIdConverter<BuildHealthFilterId>
+	{
+		/// <inheritdoc/>
+		public override BuildHealthFilterId FromBinaryId(BinaryId id) => new BuildHealthFilterId(id);
+
+		/// <inheritdoc/>
+		public override BinaryId ToBinaryId(BuildHealthFilterId value) => value.Id;
+	}
+}
diff --git a/Engine/Source/Programs/Shared/EpicGames.Horde/BuildHealth/BuildHealthFilterMessages.cs b/Engine/Source/Programs/Shared/EpicGames.Horde/BuildHealth/BuildHealthFilterMessages.cs
new file mode 100644
--- /dev/null
+++ b/Engine/Source/Programs/Shared/EpicGames.Horde/BuildHealth/BuildHealthFilterMessages.cs
@@ -0,0 +1,148 @@
+// Copyright Epic Games, Inc. All Rights Reserved.
+
+using System;
+using EpicGames.Horde.Projects;
+using EpicGames.Horde.Users;
+
+namespace EpicGames.Horde.BuildHealth
+{
+	/// <summary>
+	/// Static class that describes build health filter restrictions.
+	/// </summary>
+	public static class BuildHealthFilterRestrictions
+	{
+		/// <summary>
+		/// The maximum filter name length.
+		/// </summary>
+		public const int MaxFilterNameLength = 250;
+
+		/// <summary>
+		/// The maximum filter query length.
+		/// </summary>
+		public const int MaxFilterQueryLength = 4500;
+	}
+
+	/// <summary>
+	/// DTO used for requests used to add a build health filter.
+	/// </summary>
+	public record BuildHealthFilterAddRequest
+	{
+		/// <summary>
+		/// The project the filter belongs to.
+		/// </summary>
+		public ProjectId FilterProject { get; set; } = default!;
+
+		/// <summary>
+		/// The name of the filter.
+		/// </summary>
+		public string FilterName { get; set; } = default!;
+
+		/// <summary>
+		/// The filter query.
+		/// </summary>
+		public string FilterQuery { get; set; } = default!;
+
+		/// <summary>
+		/// Returns whether the filter name request is valid or not.
+		/// </summary>
+		/// <returns>True if valid, false otherwise.</returns>
+		public bool IsValidFilterName() => (FilterName != null && FilterName.Length <= BuildHealthFilterRestrictions.MaxFilterNameLength);
+
+		/// <summary>
+		/// Returns whether the filter query request is valid or not.
+		/// </summary>
+		/// <returns>True if valid, false otherwise.</returns>
+		public bool IsValidFilterQuery() => (FilterQuery != null && FilterQuery.Length <= BuildHealthFilterRestrictions.MaxFilterQueryLength);
+
+		/// <summary>
+		/// Returns whether the filter add request is valid.
+		/// </summary>
+		/// <returns>True if valid, false otherwise.</returns>
+		public bool IsValidAddRequest() => IsValidFilterQuery() && IsValidFilterName();
+	}
+
+	/// <summary>
+	/// DTO used for requests used to update a build health filter.
+	/// </summary>
+	public record BuildHealthFilterUpdateRequest
+	{
+		/// <summary>
+		/// The name of the filter.
+		/// </summary>
+		public string? FilterName { get; set; } = default!;
+
+		/// <summary>
+		/// The filter query.
+		/// </summary>
+		public string? FilterQuery { get; set; } = default!;
+
+		/// <summary>
+		/// Returns whether the filter name request is valid or not.
+		/// </summary>
+		/// <returns>True if valid, false otherwise.</returns>
+		public bool IsValidFilterName() => (FilterName == null || FilterName.Length <= BuildHealthFilterRestrictions.MaxFilterNameLength);
+
+		/// <summary>
+		/// Returns whether the filter query request is valid or not.
+		/// </summary>
+		/// <returns>True if valid, false otherwise.</returns>
+		public bool IsValidFilterQuery() => (FilterQuery == null || FilterQuery.Length <= BuildHealthFilterRestrictions.MaxFilterQueryLength);
+
+		/// <summary>
+		/// Returns whether the filter update request is valid.
+		/// </summary>
+		/// <returns>True if valid, false otherwise.</returns>
+		public bool IsValidUpdateRequest() => IsValidFilterQuery() && IsValidFilterName();
+	}
+
+	/// <summary>
+	/// Response object for build health filters.
+	/// </summary>
+	public class BuildHealthFilterResponse
+	{
+		/// <summary>
+		/// The id of the filter.
+		/// </summary>
+		public BuildHealthFilterId Id { get; init; } = default;
+
+		/// <summary>
+		/// The owner of the filter.
+		/// </summary>
+		public UserId? Owner { get; init; } = default;
+
+		/// <summary>
+		/// The project the filter belongs to.
+		/// </summary>
+		public ProjectId FilterProject { get; init; }
+
+		/// <summary>
+		/// The filter name.
+		/// </summary>
+		public string FilterName { get; init; }
+
+		/// <summary>
+		/// The filter query.
+		/// </summary>
+		public string FilterQuery { get; init; }
+
+		/// <summary>
+		/// The last updated time of the filter.
+		/// </summary>
+		public DateTime UpdateTimeUtc { get; init; }
+
+		/// <summary>
+		/// Creates a build health filter response from a provided build health filter.
+		/// </summary>
+		/// <param name="filter">The source filter.</param>
+		/// <returns>The response object.</returns>
+		public BuildHealthFilterResponse(IBuildHealthFilter filter)
+		{
+			Id = filter.Id;
+			Owner = filter.Owner;
+			FilterProject = filter.FilterProject;
+			FilterName = filter.FilterName;
+			FilterQuery = filter.FilterQuery;
+			UpdateTimeUtc = filter.UpdateTimeUtc;
+		}
+	}
+}
diff --git a/Engine/Source/Programs/Shared/EpicGames.Horde/BuildHealth/IBuildHealthFilter.cs b/Engine/Source/Programs/Shared/EpicGames.Horde/BuildHealth/IBuildHealthFilter.cs
new file mode 100644
--- /dev/null
+++ b/Engine/Source/Programs/Shared/EpicGames.Horde/BuildHealth/IBuildHealthFilter.cs
@@ -0,0 +1,44 @@
+// Copyright Epic Games, Inc. All Rights Reserved.
+
+using System;
+using EpicGames.Horde.Projects;
+using EpicGames.Horde.Users;
+
+namespace EpicGames.Horde.BuildHealth
+{
+	/// <summary>
+	/// Describes a build health filter.
+	/// </summary>
+	public interface IBuildHealthFilter
+	{
+		/// <summary>
+		/// The id of the filter.
+		/// </summary>
+		BuildHealthFilterId Id { get; }
+
+		/// <summary>
+		/// Owner of the filter.
+		/// </summary>
+		UserId? Owner { get; }
+
+		/// <summary>
+		/// The underyling query of the filter.
+		/// </summary>
+		string FilterQuery { get; }
+
+		/// <summary>
+		/// The name of the filter.
+		/// </summary>
+		string FilterName { get; }
+
+		/// <summary>
+		/// The project the filter belongs to.
+		/// </summary>
+		ProjectId FilterProject { get; }
+
+		/// <summary>
+		/// The last time this filter was updated (UTC).
+		/// </summary>
+		DateTime UpdateTimeUtc { get; }
+	}
+}
diff --git a/Engine/Source/Programs/Shared/EpicGames.Horde/BuildHealth/IBuildHealthFilterCollection.cs b/Engine/Source/Programs/Shared/EpicGames.Horde/BuildHealth/IBuildHealthFilterCollection.cs
new file mode 100644
--- /dev/null
+++ b/Engine/Source/Programs/Shared/EpicGames.Horde/BuildHealth/IBuildHealthFilterCollection.cs
@@ -0,0 +1,62 @@
+// Copyright Epic Games, Inc. All Rights Reserved.
+
+using System.Collections.Generic;
+using System.Threading;
+using System.Threading.Tasks;
+using EpicGames.Horde.Projects;
+using EpicGames.Horde.Users;
+
+namespace EpicGames.Horde.BuildHealth
+{
+	/// <summary>
+	/// Collection interace for managing build heatlh filters.
+	/// </summary>
+	public interface IBuildHealthFilterCollection
+	{
+		/// <summary>
+		/// Creates a new filter.
+		/// </summary>
+		/// <param name="owner">The owner of the new filter.</param>
+		/// <param name="filterRequest">The add request.</param>
+		/// <param name="cancellationToken">The cancellation token to use throughout the operation.</param>
+		/// <returns>The new filter, if applicable.</returns>
+		Task<IBuildHealthFilter?> AddBuildHealthFilterAsync(UserId? owner, BuildHealthFilterAddRequest filterRequest, CancellationToken cancellationToken);
+
+		/// <summary>
+		/// Updates the build health filter with the provided updates.
+		/// </summary>
+		/// <param name="filterId">The id of the filter to update.</param>
+		/// <param name="owner">The owner of the filter.</param>
+		/// <param name="filterRequest">The update request.</param>
+		/// <param name="cancellationToken">The cancellation token to use throughout the operation.</param>
+		/// <returns>The updated filter, if applicable, null if no updates were enacted.</returns>
+		/// <remarks>The provided owner must be making the request.</remarks>
+		Task<IBuildHealthFilter?> UpdateBuildHealthFilterAsync(BuildHealthFilterId filterId, UserId? owner, BuildHealthFilterUpdateRequest filterRequest, CancellationToken cancellationToken);
+
+		/// <summary>
+		/// Gets a filter.
+		/// </summary>
+		/// <param name="filterId">THe id of the filter to retrieve.</param>
+		/// <param name="cancellationToken">The cancellation token to use throughout the operation.</param>
+		/// <returns>The filter, if found.</returns>
+		Task<IBuildHealthFilter?> GetBuildHealthFilterAsync(BuildHealthFilterId filterId, CancellationToken cancellationToken);
+
+		/// <summary>
+		/// Gets filters for a given project.
+		/// </summary>
+		/// <param name="projectId">The project id to search for filters.</param>
+		/// <param name="cancellationToken">The cancellation token to use throughout the operation.</param>
+		/// <returns>A list of filters, if applicable.</returns>
+		Task<IEnumerable<IBuildHealthFilter>> GetBuildHealthFiltersAsync(ProjectId projectId, CancellationToken cancellationToken);
+
+		/// <summary>
+		/// Deletes a filter.
+		/// </summary>
+		/// <param name="filterId">The id of the filter to delete.</param>
+		/// <param name="owner">The owner of the filter.</param>
+		/// <param name="cancellationToken">The cancellation token to use throughout the operation.</param>
+		/// <returns>True if the filter was deleted successfully, false otherwise.</returns>
+		/// <remarks>Only the owner of an owned filter can delete it. Unowned filters can be deleted by any user.</remarks>
+		Task<bool> DeleteBuildHealthFilterAsync(BuildHealthFilterId filterId, UserId? owner, CancellationToken cancellationToken);
+	}
+}
