[Horde][BuildHealth][Dashboard] BuildHealthFilter support

#Synced-CL 48788183

diff --git a/Engine/Source/Programs/Horde/HordeDashboard/src/backend/Api.ts b/Engine/Source/Programs/Horde/HordeDashboard/src/backend/Api.ts
--- a/Engine/Source/Programs/Horde/HordeDashboard/src/backend/Api.ts
+++ b/Engine/Source/Programs/Horde/HordeDashboard/src/backend/Api.ts
@@ -5406,5 +5406,60 @@ export type GetServiceAccountResponse = {
 	enabled: boolean;
 }
 
+// #region -- Build Health Filters --
 
+/// Build Health Filters
 
+/**
+ * Datastructure that describes an add request for Build Health Filters.
+ */
+export type AddBuildHealthFiltersRequest = {
+	/// The project id.
+	projectId: string;
+
+	/// The filter name.
+	filterName: string;
+
+	/// The filter query.
+	filterQuery: string;
+}
+
+/**
+ * Datastructure that describes an update request for Build Health Filters.
+ */
+export type UpdateBuildHealthFiltersRequest = {
+	/// The filter name.
+	filterName: string;
+
+	/// The filter query.
+	filterQuery: string;
+}
+
+/** Parameters to request project Build Health Filters */
+export type GetBuildHealthFiltersRequest = {
+	/// The project id to obtain filters for.
+	projectId: string;
+}
+
+/** Response object representing a Build Health Filter. */
+export type GetBuildHealthFilterResponse = {
+	/// The filter id.
+	id: string;
+
+	/// The owner of the filter.
+	owner: string;
+
+	/// The project the filter belongs to.
+	filterProject: string;
+
+	/// The name of the filter.
+	filterName: string;
+
+	/// The query of the filter.
+	filterQuery: string;
+
+	/// The last updated time of the filter.
+	updateTimeUtc: Date | string;
+}
+
+// #endregion -- Build Health Filters --
\ No newline at end of file
diff --git a/Engine/Source/Programs/Horde/HordeDashboard/src/backend/Backend.ts b/Engine/Source/Programs/Horde/HordeDashboard/src/backend/Backend.ts
--- a/Engine/Source/Programs/Horde/HordeDashboard/src/backend/Backend.ts
+++ b/Engine/Source/Programs/Horde/HordeDashboard/src/backend/Backend.ts
@@ -1,7 +1,7 @@
 // Copyright Epic Games, Inc. All Rights Reserved.
 
 import templateCache from '../backend/TemplateCache';
-import { AccountClaimMessage, AgentData, AgentQuery, ApproveAgentsRequest, AuditLogEntry, AuditLogQuery, BatchUpdatePoolRequest, ChangeSummaryData, CreateAccountRequest, CreateBisectTaskRequest, CreateBisectTaskResponse, CreateDeviceRequest, CreateDeviceResponse, CreateExternalIssueRequest, CreateExternalIssueResponse, CreateJobRequest, CreateJobResponse, CreateNoticeRequest, CreatePoolRequest, CreateServiceAccountRequest, CreateServiceAccountResponse, CreateSoftwareResponse, CreateSubscriptionRequest, CreateSubscriptionResponse, CreateZipRequest, DashboardPreference, DevicePoolTelemetryQuery, DeviceTelemetryQuery, EventData, FindArtifactsResponse, FindIssueResponse, FindJobTimingsResponse, GetAccountResponse, GetAgentLeaseResponse, GetAgentTelemetrySampleResponse, GetArtifactDirectoryResponse, GetArtifactResponse, GetBisectTaskResponse, GetDashboardConfigResponse, GetDevicePlatformResponse, GetDevicePoolResponse, GetDevicePoolTelemetryResponse, GetDeviceReservationResponse, GetDeviceResponse, GetDeviceTelemetryResponse, GetExternalIssueProjectResponse, GetExternalIssueResponse, GetIssueStreamResponse, GetJobsTabResponse, GetJobStepRefResponse, GetJobStepTraceResponse, GetJobTimingResponse, GetLogEventResponse, GetLogFileResponse, GetNoticeResponse, GetNotificationResponse, GetPendingAgentsResponse, GetPerforceServerStatusResponse, GetPoolResponse, GetPoolSummaryResponse, GetServerInfoResponse, GetServerSettingsResponse, GetServiceAccountResponse, GetSoftwareResponse, GetSubscriptionResponse, GetTestDataDetailsResponse, GetTestDataRefResponse, GetTestMetaResponse, GetTestResponse, GetTestsRequest, GetTestStreamResponse, GetToolResponse, GetToolSummaryResponse, GetUserResponse, GetUtilizationTelemetryResponse, GlobalConfig, IssueData, IssueQuery, IssueQueryV2, JobData, JobQuery, JobsTabColumnType, JobStepOutcome, JobStreamQuery, JobTimingsQuery, LeaseData, LogData, LogEventQuery, LogLineData, PoolData, PoolQuery, PreflightConfigResponse, ProjectData, ScheduleData, ScheduleQuery, SearchLogFileResponse, ServerPluginInfoResponse, ServerStatusResponse, ServerUpdateResponse, SessionData, StreamData, TabType, TestData, UpdateAccountRequest, UpdateAgentRequest, UpdateBisectTaskRequest, UpdateCurrentAccountRequest, UpdateDeploymentRequest, UpdateDeviceRequest, UpdateGlobalConfigRequest, UpdateIssueRequest, UpdateJobRequest, UpdateLeaseRequest, UpdateNoticeRequest, UpdateNotificationsRequest, UpdatePoolRequest, UpdateServerSettingsRequest, UpdateServiceAccountRequest, UpdateServiceAccountResponse, UpdateStepRequest, UpdateStepResponse, UpdateTemplateRefRequest, UpdateUserRequest, UsersQuery } from './Api';
+import { AccountClaimMessage, AddBuildHealthFiltersRequest, AgentData, AgentQuery, ApproveAgentsRequest, AuditLogEntry, AuditLogQuery, BatchUpdatePoolRequest, ChangeSummaryData, CreateAccountRequest, CreateBisectTaskRequest, CreateBisectTaskResponse, CreateDeviceRequest, CreateDeviceResponse, CreateExternalIssueRequest, CreateExternalIssueResponse, CreateJobRequest, CreateJobResponse, CreateNoticeRequest, CreatePoolRequest, CreateServiceAccountRequest, CreateServiceAccountResponse, CreateSoftwareResponse, CreateSubscriptionRequest, CreateSubscriptionResponse, CreateZipRequest, DashboardPreference, DevicePoolTelemetryQuery, DeviceTelemetryQuery, EventData, FindArtifactsResponse, FindIssueResponse, FindJobTimingsResponse, GetAccountResponse, GetAgentLeaseResponse, GetAgentTelemetrySampleResponse, GetArtifactDirectoryResponse, GetArtifactResponse, GetBisectTaskResponse, GetBuildHealthFilterResponse, GetBuildHealthFiltersRequest, GetDashboardConfigResponse, GetDevicePlatformResponse, GetDevicePoolResponse, GetDevicePoolTelemetryResponse, GetDeviceReservationResponse, GetDeviceResponse, GetDeviceTelemetryResponse, GetExternalIssueProjectResponse, GetExternalIssueResponse, GetIssueStreamResponse, GetJobsTabResponse, GetJobStepRefResponse, GetJobStepTraceResponse, GetJobTimingResponse, GetLogEventResponse, GetLogFileResponse, GetNoticeResponse, GetNotificationResponse, GetPendingAgentsResponse, GetPerforceServerStatusResponse, GetPoolResponse, GetPoolSummaryResponse, GetServerInfoResponse, GetServerSettingsResponse, GetServiceAccountResponse, GetSoftwareResponse, GetSubscriptionResponse, GetTestDataDetailsResponse, GetTestDataRefResponse, GetTestMetaResponse, GetTestResponse, GetTestsRequest, GetTestStreamResponse, GetToolResponse, GetToolSummaryResponse, GetUserResponse, GetUtilizationTelemetryResponse, GlobalConfig, IssueData, IssueQuery, IssueQueryV2, JobData, JobQuery, JobsTabColumnType, JobStepOutcome, JobStreamQuery, JobTimingsQuery, LeaseData, LogData, LogEventQuery, LogLineData, PoolData, PoolQuery, PreflightConfigResponse, ProjectData, ScheduleData, ScheduleQuery, SearchLogFileResponse, ServerPluginInfoResponse, ServerStatusResponse, ServerUpdateResponse, SessionData, StreamData, TabType, TestData, UpdateAccountRequest, UpdateAgentRequest, UpdateBisectTaskRequest, UpdateBuildHealthFiltersRequest, UpdateCurrentAccountRequest, UpdateDeploymentRequest, UpdateDeviceRequest, UpdateGlobalConfigRequest, UpdateIssueRequest, UpdateJobRequest, UpdateLeaseRequest, UpdateNoticeRequest, UpdateNotificationsRequest, UpdatePoolRequest, UpdateServerSettingsRequest, UpdateServiceAccountRequest, UpdateServiceAccountResponse, UpdateStepRequest, UpdateStepResponse, UpdateTemplateRefRequest, UpdateUserRequest, UsersQuery } from './Api';
 import dashboard, { Dashboard } from './Dashboard';
 import { ChallengeStatus, Fetch } from './Fetch';
 import { projectStore } from './ProjectStore';
@@ -1784,6 +1784,17 @@ export class Backend {
         });
     }
 
+    // get account access for a specific scope and action
+    getAccountAccess(scope: string, action: string): Promise<any> {
+        return new Promise<any>((resolve, reject) => {
+            this._fetch.get(`/account/access`, { params: { scope: scope, action: action } }).then((value) => {
+                resolve(value.data as any);
+            }).catch(reason => {
+                reject(reason);
+            });
+        });
+    }
+
     // update current account 
     updateCurrentAccount(request: UpdateCurrentAccountRequest): Promise<boolean> {
         return new Promise<boolean>((resolve, reject) => {
@@ -1893,6 +1904,54 @@ export class Backend {
         });
     }
 
+    //#region -- Build Health Filters --
+
+    // add a build health project filter
+    addBuildHealthProjectFilter(request: AddBuildHealthFiltersRequest): Promise<GetBuildHealthFilterResponse> {
+        return new Promise<GetBuildHealthFilterResponse>((resolve, reject) => {
+            this._fetch.post(`/api/v1/buildhealth/filters`, request).then((value) => {
+                resolve(value.data as GetBuildHealthFilterResponse);
+            }).catch(reason => {
+                reject(reason);
+            });
+        });
+    }
+
+    // update a build health project filter
+    updateBuildHealthProjectFilter(filterId: string, request: UpdateBuildHealthFiltersRequest): Promise<GetBuildHealthFilterResponse> {
+        return new Promise<GetBuildHealthFilterResponse>((resolve, reject) => {
+            this._fetch.put(`/api/v1/buildhealth/filters/${filterId}`, request).then((value) => {
+                resolve(value.data as GetBuildHealthFilterResponse);
+            }).catch(reason => {
+                reject(reason);
+            });
+        });
+    }
+
+    // delete a build health project filter
+    deleteBuildHealthProjectFilter(filterId: string): Promise<boolean> {
+        return new Promise<boolean>((resolve, reject) => {
+            this._fetch.delete(`/api/v1/buildhealth/filters/${filterId}`).then(() => {
+                resolve(true);
+            }).catch(() => {
+                reject(false);
+            });
+        });
+    }
+
+    // get build health project filters
+    getBuildHealthProjectFilters(query: GetBuildHealthFiltersRequest): Promise<GetBuildHealthFilterResponse[]> {
+        return new Promise<GetBuildHealthFilterResponse[]>((resolve, reject) => {
+            this._fetch.get(`/api/v1/buildhealth/filters`, { params: query }).then((value) => {
+                resolve(value.data as GetBuildHealthFilterResponse[]);
+            }).catch(reason => {
+                reject(reason);
+            });
+        });
+    }
+
+    // #endregion -- Build Health Filters --
+
     private async update() {
 
         if (this.updateID === "updating") {
diff --git a/Engine/Source/Programs/Horde/HordeDashboard/src/components/buildhealth/BuildHealthDataHandler.ts b/Engine/Source/Programs/Horde/HordeDashboard/src/components/buildhealth/BuildHealthDataHandler.ts
--- a/Engine/Source/Programs/Horde/HordeDashboard/src/components/buildhealth/BuildHealthDataHandler.ts
+++ b/Engine/Source/Programs/Horde/HordeDashboard/src/components/buildhealth/BuildHealthDataHandler.ts
@@ -1,7 +1,7 @@
 // Copyright Epic Games, Inc. All Rights Reserved.
 
 import backend from "horde/backend";
-import { JobData, JobStreamQuery, ProjectData, StreamData } from "horde/backend/Api";
+import { GetBuildHealthFilterResponse, JobData, JobStreamQuery, ProjectData, StreamData } from "horde/backend/Api";
 import { action, makeObservable, observable, runInAction } from "mobx";
 import { BuildHealthOptionsState } from "./BuildHealthOptions";
 import { LabelRefData, StepRefData, TemplateRefData } from "./BuildHealthDataTypes";
@@ -71,6 +71,8 @@ export class BuildHealthDataHandler implements IHierarchyValidator {
     private readonly STREAMS_FILTER = "id,state,streamId,name,labels";
     private readonly STREAM_JOBS_BATCH_STEP_FILTER = "batches.steps.id,batches.steps.name,batches.steps.state,batches.steps.outcome,batches.steps.finishTime"
     private readonly STEP_DATA_SAMPLE_SIZE = 5;
+    private readonly ACL_SCOPE = "horde";
+    private readonly ACL_ACTION = "AddUpdateDeleteBuildHealthFilter";
 
     // #region -- Private Members --
 
@@ -89,6 +91,12 @@ export class BuildHealthDataHandler implements IHierarchyValidator {
     @observable.shallow
     projectsData: ProjectData[] = [];
 
+    @observable.shallow
+    projectFilterData: Map<string, GetBuildHealthFilterResponse[]> = new Map<string, GetBuildHealthFilterResponse[]>();
+
+    @observable
+    hasBuildHealthFilterWriteAccess: boolean = false;
+
     @observable.shallow
     streamsData: StreamData[] = [];
 
@@ -129,10 +137,10 @@ export class BuildHealthDataHandler implements IHierarchyValidator {
      */
     @action
     reset() {
-        this.streamsData =[];
+        this.streamsData = [];
         this.templatesData = [];
         this.stepData = [];
-        this.labelData =[];
+        this.labelData = [];
     }
 
     /**
@@ -160,11 +168,36 @@ export class BuildHealthDataHandler implements IHierarchyValidator {
             runInAction(() => {
                 this.projectsData = projectDataResult;
             });
-        }
 
+            const accessResults: { access: boolean, scopes: any[] } = await backend.getAccountAccess(this.ACL_SCOPE, this.ACL_ACTION);
+
+            runInAction(() => {
+                this.hasBuildHealthFilterWriteAccess = accessResults.access;
+            })
+        }
+        await this.getProjectFiltersData();
         await this.getStreamData(request);
     }
 
+    /**
+     * Get the project filters data.
+     * @param projectId the project to get filter data for. If undefined, obtain for all known projects.
+     */
+    async getProjectFiltersData(projectId?: string) {
+        let buildHealthFilterData: Map<string, GetBuildHealthFilterResponse[]> = new Map<string, GetBuildHealthFilterResponse[]>;
+
+        for (let idx: number = 0; idx < this.projectsData.length; ++idx) {
+            if (projectId === undefined || projectId === this.projectsData[idx].id) {
+                const projectFilterDataResult: GetBuildHealthFilterResponse[] = (await backend.getBuildHealthProjectFilters({ projectId: this.projectsData[idx].id }));
+                buildHealthFilterData.set(this.projectsData[idx].id, projectFilterDataResult);
+            }
+        }
+
+        runInAction(() => {
+            this.projectFilterData = buildHealthFilterData;
+        })
+    }
+
     /**
      * Obtains the stream data given the current buildHealthOptions.
      * @param request The data refresh request message to use in order to control refresh depth.
@@ -311,7 +344,7 @@ export class BuildHealthDataHandler implements IHierarchyValidator {
         let returnArray: StreamDescriptor[] = [];
 
         if (this.cachedProjectToStreamData.has(projectId)) {
-            returnArray = this.cachedProjectToStreamData.get(projectId)!.map(x =>  ({ id: x.id, name: x.name }));
+            returnArray = this.cachedProjectToStreamData.get(projectId)!.map(x => ({ id: x.id, name: x.name }));
         }
 
         return returnArray;
diff --git a/Engine/Source/Programs/Horde/HordeDashboard/src/components/buildhealth/BuildHealthFilterComponent.tsx b/Engine/Source/Programs/Horde/HordeDashboard/src/components/buildhealth/BuildHealthFilterComponent.tsx
new file mode 100644
--- /dev/null
+++ b/Engine/Source/Programs/Horde/HordeDashboard/src/components/buildhealth/BuildHealthFilterComponent.tsx
@@ -0,0 +1,316 @@
+import { DirectionalHint, Dropdown, IconButton, IDropdownOption, ISelectableOption, Label, Stack, TooltipHost } from "@fluentui/react";
+import { GetBuildHealthFilterResponse } from "horde/backend/Api";
+import { BuildHealthDataHandler } from "./BuildHealthDataHandler";
+import { BuildHealthOptionsController } from "./BuildHealthOptionsController";
+import { useState } from "react";
+import { observer } from "mobx-react-lite";
+import { BuildHealthFilterModal, BuildHealthFilterModalProps } from "./BuildHealthFilterModal";
+import { parseBuildHealthQueryParams, upgradeBuildHealthQueryParams, loadBuildHealthOptionsFromParams, toNavigationQuery } from "./BuildHealthOptions";
+import { getHordeStyling } from "horde/styles/Styles";
+
+// #region -- Helper Functions --
+
+function constructEditableBuildHealthFilter(buildHealthFilter: GetBuildHealthFilterResponse | null, buildHealthOptions: BuildHealthOptionsController, handler: BuildHealthDataHandler): EditableBuildHealthFilter {
+    if (buildHealthFilter !== null) {
+        let currentOptionsString = toNavigationQuery(buildHealthOptions.state, buildHealthOptions.uiState, buildHealthOptions.querySchemaVersion, handler).toString();
+        let currentHumanReadableQuery = parseBuildHealthQueryParams(new URLSearchParams(currentOptionsString));
+        let humanReadableQuery = parseBuildHealthQueryParams(new URLSearchParams(buildHealthFilter.filterQuery));
+        return EditableBuildHealthFilter.produceFromExisting(buildHealthFilter, currentOptionsString, JSON.stringify(humanReadableQuery, null, 2), JSON.stringify(currentHumanReadableQuery, null, 2))
+    }
+    else {
+        let currentProject: string | null = buildHealthOptions.getCurrentProject()!;
+        let currentOptionsString = toNavigationQuery(buildHealthOptions.state, buildHealthOptions.uiState, buildHealthOptions.querySchemaVersion, handler).toString();
+        let humanReadableQuery = parseBuildHealthQueryParams(new URLSearchParams(currentOptionsString));
+        return EditableBuildHealthFilter.produceNew(currentProject, currentOptionsString, JSON.stringify(humanReadableQuery, null, 2))
+    }
+}
+
+function requestFilterApplication(filterQuery: string, buildHealthOptions: BuildHealthOptionsController, handler: BuildHealthDataHandler) {
+    let searchParams: URLSearchParams = new URLSearchParams(filterQuery);
+    const params = parseBuildHealthQueryParams(searchParams);
+    const upgradedParams = upgradeBuildHealthQueryParams(params.buildHealthQueryParams, buildHealthOptions.querySchemaVersion);
+    loadBuildHealthOptionsFromParams(upgradedParams, params.buildHealthUIParams, buildHealthOptions, handler);
+}
+
+// #endregion -- Helper Functions --
+
+// #region -- Public Types --
+
+/**
+ * Data structure that represents an editable Build Health Filter.
+ */
+export class EditableBuildHealthFilter {
+    /// The filter id.
+    filterId: string | null;
+
+    /// The project the filter belongs to.
+    filterProjectId: string;
+
+    /// The name of the filter.
+    filterName: string;
+
+    /// The query of the filter. If null, there is no prior state (new filter).
+    oldFilterQuery: string;
+
+    /// The query of the filter.
+    newFilterQuery: string;
+
+    /// The last updated time of the filter. If null, there is no prior state (new filter).
+    lastUpdated: Date | null;
+
+    /// The filter owner. If null, there is no prior state (new filter).
+    filterOwner: string | null;
+
+    /// The old query, in human readable form. If null, there is no prior state (new filter).
+    oldHumanReadableQuery: string | null;
+
+    /// The new query, in human readable form.
+    newHumanReadableQuery: string;
+
+    private constructor(filterName: string, filterQuery: string, newFilterQuery: string, oldHumanReadableQuery: string | null, newHumanReadableQuery: string, filterId: string | null, filterProject: string, ownerId: string | null, lastUpdated: Date | string | null) {
+        this.filterName = filterName;
+        this.oldFilterQuery = filterQuery;
+        this.newFilterQuery = newFilterQuery;
+        this.filterId = filterId;
+        this.oldHumanReadableQuery = oldHumanReadableQuery;
+        this.newHumanReadableQuery = newHumanReadableQuery;
+        this.filterProjectId = filterProject;
+        this.filterOwner = ownerId;
+        this.lastUpdated = lastUpdated === null ? null : new Date(lastUpdated);
+    }
+
+    /**
+     * Creates an editable Build Health Filter from an existing filter. 
+     * @param existingFilter The existing filter.
+     * @param newFilterQuery The new filter query to apply.
+     * @param oldHumanReadableQuery The old human readable query filter.
+     * @param newHumanReadableQuery The new human readable query filter.
+     * @returns The editable form of the existing filter.
+     */
+    static produceFromExisting(existingFilter: GetBuildHealthFilterResponse, newFilterQuery: string, oldHumanReadableQuery: string, newHumanReadableQuery: string): EditableBuildHealthFilter {
+        return new EditableBuildHealthFilter(existingFilter.filterName, existingFilter.filterQuery, newFilterQuery, oldHumanReadableQuery, newHumanReadableQuery, existingFilter.id, existingFilter.filterProject, existingFilter.owner, existingFilter.updateTimeUtc);
+    }
+
+    /**
+     * Creates an editable Build Health Filter from new data.
+     * @param projectId The project the Build Health Filter belongs to.
+     * @param newFilterQuery The new filter query to apply.
+     * @param newHumanReadableQuery The new human readable query filter.
+     * @returns The editable form of the new filter.
+     */
+    static produceNew(projectId: string, newFilterQuery: string, newHumanReadableQuery: string) {
+        return new EditableBuildHealthFilter("New Filter", newFilterQuery, newFilterQuery, null, newHumanReadableQuery, null, projectId, null, null);
+    }
+}
+
+// #endregion -- Public Types --
+
+// #region -- Private Types --
+
+type BuildHealthFilterOptionData = {
+    id: string,
+    text: string,
+    query: string,
+    owner?: string,
+    underylingObject: GetBuildHealthFilterResponse | undefined
+}
+
+// #endregion -- Private Types --
+
+// #region -- Private Components --
+
+const ProjectFilterDropdownSingle: React.FC<{ handler: BuildHealthDataHandler, buildHealthOptions: BuildHealthOptionsController, onUploadBuildFilter: (fitler: GetBuildHealthFilterResponse | null) => void }> = observer(function ConstructProjectDropdownSingle({ handler, buildHealthOptions, onUploadBuildFilter }) {
+    let selectedFilterId = buildHealthOptions.getBuildHealthFilterId();
+    let selectedFilter: GetBuildHealthFilterResponse | null = null;
+    let selectedProject: string | null = buildHealthOptions.getCurrentProject();
+    if (selectedProject === null) {
+        return (null);
+    }
+
+    let buildHealthFilters: GetBuildHealthFilterResponse[] | undefined = handler.projectFilterData.get(selectedProject);
+    if (buildHealthFilters === undefined) {
+        return (null)
+    }
+
+    const doptions: IDropdownOption<BuildHealthFilterOptionData>[] = [];
+    doptions.push({
+        key: "empty",
+        text: "(New)",
+        data: undefined
+    });
+
+    buildHealthFilters.map((projectFilter: GetBuildHealthFilterResponse) => {
+        if (projectFilter.id === selectedFilterId) {
+            selectedFilter = projectFilter;
+        }
+
+        let listParam: BuildHealthFilterOptionData = {
+            id: `${projectFilter.id}`,
+            text: projectFilter.filterName,
+            query: projectFilter.filterQuery,
+            underylingObject: projectFilter
+        };
+
+        doptions.push({
+            key: listParam.id,
+            text: listParam.text,
+            data: listParam
+        });
+    });
+
+    let iconName: string = selectedFilterId === null ? "Upload" : "Edit";
+    let computedSelectedKey: string = selectedFilterId === null ? "empty" : selectedFilterId;
+
+    return (
+        <Stack>
+            <TooltipHost
+                content="Saved filters for the selected project."
+                calloutProps={{ directionalHint: DirectionalHint.rightCenter }}>
+                <Label>Project Filters</Label>
+            </TooltipHost>
+            <Stack horizontal tokens={{ childrenGap: 8 }}>
+                <Dropdown
+                    options={doptions}
+                    styles={{ dropdown: { width: 300 } }}
+                    selectedKey={computedSelectedKey}
+                    placeholder={"(New)"}
+                    calloutProps={{
+                        directionalHint: DirectionalHint.rightTopEdge,
+                        alignTargetEdge: true,
+                    }}
+                    onRenderOption={(option?: ISelectableOption) => {
+                        if (!option) {
+                            return null;
+                        }
+                        let castedOption: IDropdownOption<BuildHealthFilterOptionData> = option as IDropdownOption<BuildHealthFilterOptionData>;
+                        if (!castedOption) {
+                            return null;
+                        }
+
+                        return (
+                            <TooltipHost content={"Last Updated: " + (castedOption.data?.underylingObject !== undefined ? castedOption.data?.underylingObject?.updateTimeUtc : "N/A")}
+                                calloutProps={{ directionalHint: DirectionalHint.rightCenter }}>
+                                <div
+                                    style={{
+                                        display: 'block',
+                                        overflow: 'hidden',
+                                        textOverflow: 'ellipsis',
+                                        whiteSpace: 'nowrap',
+                                        maxWidth: '280px',
+                                    }}
+                                >
+                                    {castedOption.text}
+                                </div>
+                            </TooltipHost>
+                        )
+                    }}
+                    onChange={(_ev, option: IDropdownOption<BuildHealthFilterOptionData>) => {
+                        if (!option.data) {
+
+                            buildHealthOptions.setBuildHealthFilterId(null);
+                            return;
+                        }
+
+                        // Simply re-issue our parse on the stored query.
+                        const query = option.data!.query;
+                        buildHealthOptions.setBuildHealthFilterId(option.data!.underylingObject!.id);
+                        requestFilterApplication(query, buildHealthOptions, handler);
+                    }}>
+                </Dropdown>
+                {/* Only show the button if the user has write access for Build Health Filters. */}
+                {
+                    handler.hasBuildHealthFilterWriteAccess
+                    && (
+                        <TooltipHost
+                            content="Promotes or updates the current filter selection to the server."
+                            directionalHint={DirectionalHint.rightTopEdge}
+                        >
+                            <IconButton
+                                iconProps={{ iconName: iconName }}
+                                styles={{
+                                    root: {
+                                        width: 30,
+                                        height: 30,
+                                        minWidth: 0,
+                                        padding: 0,
+                                    },
+                                }}
+                                onClick={() => {
+                                    // get build filter from data structs
+                                    onUploadBuildFilter(selectedFilter);
+                                }}
+                            />
+                        </TooltipHost>
+                    )
+                }
+            </Stack>
+        </Stack>
+    );
+});
+
+// #endregion -- Private Components --
+
+// #region -- Public Components --
+
+/**
+ * Toplevel Build Health Filter component.
+ * @returns The component.
+ */
+export const BuildHealthFilterComponent: React.FC<{ handler: BuildHealthDataHandler, buildHealthOptions: BuildHealthOptionsController }> = observer(function ConstructBuildHealthFilterComponent({ handler, buildHealthOptions }) {
+    const [selectedFilter, setSelectedFilter] = useState<EditableBuildHealthFilter | undefined>();
+    const [modalVisible, setModalVisible] = useState(false);
+    const { hordeClasses, modeColors } = getHordeStyling();
+
+    let modalProperties: BuildHealthFilterModalProps = {
+        buildHealthFilter: selectedFilter,
+        onDismiss: () => {
+            setModalVisible(false);
+        },
+        onSaveSuccess: (savedFilter: GetBuildHealthFilterResponse) => {
+            setModalVisible(false);
+            buildHealthOptions.setBuildHealthFilterId(savedFilter.id);
+            handler.getProjectFiltersData(savedFilter.filterProject);
+        },
+        onDeleteSuccess: () => {
+            setModalVisible(false);
+            buildHealthOptions.setBuildHealthFilterId(null);
+            handler.getProjectFiltersData();
+        },
+        onRevert: () => {
+            setModalVisible(false);
+
+            if (selectedFilter === undefined) {
+                return;
+            }
+
+            requestFilterApplication(selectedFilter.oldFilterQuery, buildHealthOptions, handler);
+        }
+    }
+    return (
+        <>
+            {buildHealthOptions.getCurrentProject() !== null && (
+                <Stack
+                    styles={{
+                        root: {
+                            border: `1px solid ${modeColors.content}`,
+                            borderRadius: 4,
+                            padding: 8,
+                        },
+                    }}>
+                    <ProjectFilterDropdownSingle
+                        handler={handler}
+                        buildHealthOptions={buildHealthOptions}
+                        onUploadBuildFilter={(selectedItem: GetBuildHealthFilterResponse | null) => {
+                            setModalVisible(true);
+                            let editableBuildHealthFilter: EditableBuildHealthFilter = constructEditableBuildHealthFilter(selectedItem, buildHealthOptions, handler);
+                            setSelectedFilter(editableBuildHealthFilter);
+                        }} />
+                    {modalVisible && <BuildHealthFilterModal {...modalProperties} />}
+                </Stack>
+            )
+            }
+        </>
+    );
+});
+
+// #endregion -- Public Components --
\ No newline at end of file
diff --git a/Engine/Source/Programs/Horde/HordeDashboard/src/components/buildhealth/BuildHealthFilterModal.tsx b/Engine/Source/Programs/Horde/HordeDashboard/src/components/buildhealth/BuildHealthFilterModal.tsx
new file mode 100644
--- /dev/null
+++ b/Engine/Source/Programs/Horde/HordeDashboard/src/components/buildhealth/BuildHealthFilterModal.tsx
@@ -0,0 +1,245 @@
+// Copyright Epic Games, Inc. All Rights Reserved.
+
+import { Dialog, DialogType, Label, PrimaryButton, Stack, TextField } from "@fluentui/react";
+import backend from "horde/backend";
+import { AddBuildHealthFiltersRequest, GetBuildHealthFilterResponse, UpdateBuildHealthFiltersRequest } from "horde/backend/Api";
+import { useState } from "react";
+import dashboard from "horde/backend/Dashboard";
+import { EditableBuildHealthFilter } from "./BuildHealthFilterComponent";
+import { getHordeStyling } from "horde/styles/Styles";
+import { getHordeTheme } from "horde/styles/theme";
+
+// #region -- Private Helpers --
+
+function addLineNumbers(text: string): string {
+    return text
+        .split("\n")
+        .map((line, i) => `${i + 1}: ${line}`)
+        .join("\n");
+}
+
+function firstDiffLine(a: string, b: string): string {
+    const aLines = a.split("\n");
+    const bLines = b.split("\n");
+    const max = Math.max(aLines.length, bLines.length);
+
+    for (let i = 0; i < max; i++) {
+        if (aLines[i] !== bLines[i]) {
+            return [
+                `First difference at line ${i + 1}:`,
+                `OLD: ${aLines[i] ?? ""}`,
+                `NEW: ${bLines[i] ?? ""}`
+            ].join("\n");
+        }
+    }
+
+    return "No differences found.";
+}
+
+// #endregion -- Private Helpers --
+
+// #region -- Private Types --
+
+/**
+ * Datastructure for build health modal props.
+ */
+export type BuildHealthFilterModalProps = {
+    // The filter to use for the modal.
+    buildHealthFilter: EditableBuildHealthFilter | undefined;
+
+    // The delegate to execute when the modal is dismissed.
+    onDismiss?: (_ev?: React.MouseEvent<HTMLButtonElement>) => any;
+
+    // The delegate to execute when the Build Health Filter has been saved.
+    onSaveSuccess?: (savedFilter: GetBuildHealthFilterResponse) => any;
+
+    // The delegate to execute when changes to the current Build Health Filter has been reverted.
+    onRevert?: () => any;
+
+    // The delegate to execute when the Build Health Filter has been deleted.
+    onDeleteSuccess?: () => any;
+};
+
+// #endregion -- Private Types --
+
+// #region -- Components --
+
+/**
+ * Build Health Filter modal, that is responsible for the visualization and editing of a Build Health Filter
+ * @param Props that describe the Build Health Filter to use for the modal, and all delegates to use for operations. 
+ * @returns The component.
+ */
+export const BuildHealthFilterModal: React.FC<BuildHealthFilterModalProps> = ({ buildHealthFilter, onDismiss, onSaveSuccess, onDeleteSuccess, onRevert }) => {
+    const theme = getHordeTheme();
+
+    if (buildHealthFilter === undefined) {
+        return;
+    }
+
+    const [state, setState] = useState<{ filterName: string }>(buildHealthFilter !== null && buildHealthFilter.filterName !== undefined ? { filterName: buildHealthFilter.filterName } : { filterName: "" });
+    const isNewFilter: boolean = buildHealthFilter.filterId === null;
+    const isEditorOwner: boolean = buildHealthFilter.filterOwner !== null && buildHealthFilter.filterOwner == dashboard.userId;
+    const isOldAndNewDiff: boolean = buildHealthFilter.oldHumanReadableQuery !== buildHealthFilter.newHumanReadableQuery || buildHealthFilter.filterName !== state.filterName;
+    const hasHumanReadanbleOldQuery: boolean = buildHealthFilter.oldHumanReadableQuery !== null;
+
+    // #region -- Action Handlers --
+
+    async function Save(): Promise<GetBuildHealthFilterResponse | null> {
+        if (buildHealthFilter?.filterId === null) {
+            let request: AddBuildHealthFiltersRequest = { filterName: state.filterName, projectId: buildHealthFilter.filterProjectId, filterQuery: buildHealthFilter.oldFilterQuery };
+            let response: GetBuildHealthFilterResponse = await backend.addBuildHealthProjectFilter(request);
+
+            return response;
+        } else {
+            let request: UpdateBuildHealthFiltersRequest = { filterName: state.filterName, filterQuery: buildHealthFilter!.newFilterQuery };
+            let response: GetBuildHealthFilterResponse = await backend.updateBuildHealthProjectFilter(buildHealthFilter?.filterId!, request);
+            return response;
+        }
+    };
+
+    async function SaveAsCopy(): Promise<GetBuildHealthFilterResponse | null> {
+        if (buildHealthFilter?.filterId === null || buildHealthFilter === undefined) {
+            return null;
+        }
+
+        let request: AddBuildHealthFiltersRequest = { filterName: state.filterName, projectId: buildHealthFilter.filterProjectId, filterQuery: buildHealthFilter.newFilterQuery };
+        let response: GetBuildHealthFilterResponse = await backend.addBuildHealthProjectFilter(request);
+
+        return response;
+    };
+
+    async function Delete(): Promise<boolean> {
+        if (buildHealthFilter?.filterId === null || buildHealthFilter === undefined) {
+            return false;
+        }
+
+        let response: boolean = await backend.deleteBuildHealthProjectFilter(buildHealthFilter?.filterId);
+
+        return response;
+    };
+
+    // #endregion -- Action Handlers --
+
+    // #region -- Button Delegates --
+
+    const handleSaveClick = async () => {
+        const savedFilter = await Save();
+        if (savedFilter) {
+            onSaveSuccess?.(savedFilter);
+        }
+    };
+
+    const handleSaveAsCopyClick = async () => {
+        const savedFilter = await SaveAsCopy();
+        if (savedFilter) {
+            onSaveSuccess?.(savedFilter);
+        }
+    };
+
+    const handleDeleteClick = async () => {
+        const savedFilter = await Delete();
+        if (savedFilter) {
+            onDeleteSuccess?.();
+        }
+    };
+
+    const handleRevertClick = () => {
+        onRevert?.();
+    }
+
+    //#endregion -- Button Delegates --
+
+    return (
+        <Dialog
+            modalProps={{
+                isBlocking: false,
+                topOffsetFixed: true,
+                styles: {
+                    root: {
+                        selectors: {
+                            ".ms-Dialog-title": {
+                                paddingTop: '24px',
+                                paddingLeft: '32px'
+                            }
+                        }
+                    }
+                }
+            }}
+            onDismiss={onDismiss}
+            hidden={false}
+            minWidth={1400}
+            dialogContentProps={{
+                type: DialogType.close,
+                onDismiss: onDismiss,
+                title: "Upload Build Health Filter Preset",
+            }}>
+            <TextField
+                label="Project ID"
+                value={buildHealthFilter.filterProjectId}
+                readOnly
+            />
+            <TextField
+                label="Filter Name"
+                value={state.filterName}
+                onChange={(_, newValue) => {
+                    setState({ ...state, filterName: newValue ?? "" })
+                }}
+            />
+            <Stack horizontal tokens={{ childrenGap: 16 }}>
+                {hasHumanReadanbleOldQuery && <Stack.Item grow>
+                    <TextField
+                        label="Old Query"
+                        value={addLineNumbers(buildHealthFilter.oldHumanReadableQuery!)}
+                        multiline
+                        rows={15}
+                        readOnly
+                        resizable={false}
+                        styles={{ root: { width: "100%" } }}
+                    />
+                </Stack.Item>}
+                {(!hasHumanReadanbleOldQuery || isOldAndNewDiff) && <Stack.Item grow>
+                    <TextField
+                        label="New Query"
+                        value={addLineNumbers(buildHealthFilter.newHumanReadableQuery)}
+                        multiline
+                        rows={15}
+                        readOnly
+                        resizable={false}
+                        styles={{ root: { width: "100%" } }}
+                    />
+                </Stack.Item>}
+            </Stack>
+            {
+                (hasHumanReadanbleOldQuery && isOldAndNewDiff) && <TextField
+                    label="Differences"
+                    value={firstDiffLine(buildHealthFilter.oldHumanReadableQuery!, buildHealthFilter.newHumanReadableQuery)}
+                    multiline
+                    rows={3}
+                    readOnly
+                    width={"50px"}
+                    resizable={false}
+                />
+            }
+            {!isNewFilter && (
+                <Label styles={{
+                    root: {
+                        width: '100%',
+                        textAlign: 'right',
+                        fontStyle: 'italic',
+                        color: theme.semanticColors.bodyBackgroundHovered,
+                    }
+                }}>
+                    Filter Last Updated: {buildHealthFilter.lastUpdated?.toUTCString()}
+                </Label>
+            )
+            }
+
+            {(isNewFilter || isEditorOwner) && <PrimaryButton disabled={(!isNewFilter && !isOldAndNewDiff)} onClick={() => { handleSaveClick(); }} text="Save" />}
+            {!isNewFilter && <PrimaryButton onClick={() => { handleSaveAsCopyClick(); }} text="Save As Copy" />}
+            {!isNewFilter && isOldAndNewDiff && <PrimaryButton onClick={() => { handleRevertClick(); }} text="Revert" />}
+            {!isNewFilter && <PrimaryButton disabled={buildHealthFilter.filterOwner !== null && buildHealthFilter.filterOwner !== dashboard.userId} onClick={() => { handleDeleteClick(); }} text="Delete" />}
+        </Dialog >
+    );
+};
+
+// #endregion -- Components --
\ No newline at end of file
diff --git a/Engine/Source/Programs/Horde/HordeDashboard/src/components/buildhealth/BuildHealthOptions.ts b/Engine/Source/Programs/Horde/HordeDashboard/src/components/buildhealth/BuildHealthOptions.ts
--- a/Engine/Source/Programs/Horde/HordeDashboard/src/components/buildhealth/BuildHealthOptions.ts
+++ b/Engine/Source/Programs/Horde/HordeDashboard/src/components/buildhealth/BuildHealthOptions.ts
@@ -158,7 +158,7 @@ export function toNavigationQuery(optionsState: BuildHealthOptionsState, uiOptio
             },
             (prefix: string, suffixes: string[]) => {
                 let { streamId, templateId } = decodeTemplateKey(prefix);
-                let summarized = buildHealthDataHandler.getCachedTemplateStepData(streamId, templateId)?.length == suffixes.length ? STEPS_URL_PARAM_ALL_STEPS : undefined;
+                let summarized = buildHealthDataHandler.getCachedTemplateStepData(streamId, templateId)?.length === suffixes.length ? STEPS_URL_PARAM_ALL_STEPS : undefined;
 
                 return summarized;
             });
@@ -179,7 +179,7 @@ export function toNavigationQuery(optionsState: BuildHealthOptionsState, uiOptio
             },
             (prefix: string, suffixes: string[]) => {
                 let { streamId, templateId } = decodeTemplateKey(prefix);
-                let summarized = buildHealthDataHandler.getCachedTemplateLabelNames(streamId, templateId)?.length == suffixes.length ? LABELS_URL_PARAM_ALL_LABELS : undefined;
+                let summarized = buildHealthDataHandler.getCachedTemplateLabelNames(streamId, templateId)?.length === suffixes.length ? LABELS_URL_PARAM_ALL_LABELS : undefined;
 
                 return summarized;
             });
@@ -214,7 +214,7 @@ export function toNavigationQuery(optionsState: BuildHealthOptionsState, uiOptio
  * @param location Location object.
  * @returns Parsed Build Health Query Params.
  */
-export function parseBuildHealthQueryParams(location: Location): { buildHealthQueryParams: BuildHealthQueryParams, buildHealthUIParams: BuildHealthUIParams } {
+export function parseBuildHealthQueryParams(searchParams: URLSearchParams): { buildHealthQueryParams: BuildHealthQueryParams, buildHealthUIParams: BuildHealthUIParams } {
 
     // Internal helper method to expand keys from the comrpessed representation.
     // e.g. stream-id::template-id(stepname-1;stepname-2) -> stream-id::template-id::stepname-1,stream-id::template-id::stepname-2
@@ -241,7 +241,6 @@ export function parseBuildHealthQueryParams(location: Location): { buildHealthQu
         return result;
     }
 
-    const searchParams = new URLSearchParams(location.search);
     let projectIds: string[] = [];
     let streamIds: string[] = [];
     let templateIds: string[] = [];
@@ -419,15 +418,15 @@ export async function loadBuildHealthOptionsFromParams(
 
     runInAction(() => {
         optionsWriter.setDebugMode(params.debugMode);
-    })
+    });
 
     runInAction(() => {
         optionsWriter.setIncludePreflight(params.includePreflight);
-    })
+    });
 
     runInAction(() => {
         optionsWriter.setIncludeCancelledJobs(params.includeCancelledJobs);
-    })
+    });
 
     if (params.timeSpanKey) {
         runInAction(() => {
@@ -437,7 +436,9 @@ export async function loadBuildHealthOptionsFromParams(
 
     // Set States
     if (params.stateNames) {
-        optionsWriter.clearStates();
+        runInAction(() => {
+            optionsWriter.clearStates();
+        });
 
         if (params.stateNames.length === 1 && params.stateNames[0].includes(STATES_URL_PARAM_ALL_STATES)) {
             optionsWriter.toggleAllStates(ALL_STATES_TUPLE_LIST.map(entry => entry.key), true);
@@ -451,7 +452,10 @@ export async function loadBuildHealthOptionsFromParams(
 
     // Set Job Starts
     if (params.jobStartNames) {
-        optionsWriter.clearJobStartMethods();
+        runInAction(() => {
+            optionsWriter.clearJobStartMethods();
+        });
+
         if (params.jobStartNames.length === 1 && params.jobStartNames[0].includes(JOB_START_URL_PARAM_ALL_NAMES)) {
             optionsWriter.toggleAllJobStartMethods(ALL_JOB_START_METHODS_LIST, true);
         } else if (params.jobStartNames.length > 0) {
@@ -725,8 +729,10 @@ export class BuildHealthOptionsState {
     // #endregion -- Published Options --
 
     // #region -- User Selected Options --
-
+    
+    @observable buildHealthFilterId : string | null; 
     @observable enabledProjects: Record<string, string> = {};
+    @observable enabledProjectFilter: string;
     @observable enabledStreams: Record<string, string> = {};
     @observable enabledTemplates: Record<string, string> = {};
     @observable enabledLabels: Record<string, string> = {};
diff --git a/Engine/Source/Programs/Horde/HordeDashboard/src/components/buildhealth/BuildHealthOptionsController.ts b/Engine/Source/Programs/Horde/HordeDashboard/src/components/buildhealth/BuildHealthOptionsController.ts
--- a/Engine/Source/Programs/Horde/HordeDashboard/src/components/buildhealth/BuildHealthOptionsController.ts
+++ b/Engine/Source/Programs/Horde/HordeDashboard/src/components/buildhealth/BuildHealthOptionsController.ts
@@ -816,6 +816,14 @@ export class BuildHealthOptionsController {
 
     // #region -- Public API --
 
+    /**
+     * Gets the currently selected project for the entire options set.
+     */
+    getCurrentProject(): string | null {
+        const keys = Object.keys(this.state.enabledProjects);
+        return keys.length > 0 ? keys[0] : null;
+    }
+
     /**
      * Sets the hierarchy validator for the controller.
      * @param validator The hierarchy valildator to us.
@@ -824,6 +832,23 @@ export class BuildHealthOptionsController {
         this.hierarchyValidator = validator;
     }
 
+    /**
+     * Sets the current selected Build Health Filter Id.
+     * @param buildHealthFilterId The Build Health Filter Id to set to. Null if no active filter is set.
+     */
+    @action
+    setBuildHealthFilterId(buildHealthFilterId: string | null) {
+        this.state.buildHealthFilterId = buildHealthFilterId;
+    }
+
+    /**
+     * Obtains the currently selected Build Health Filter Id.
+     * @returns The currently selected Build Health Filter id.
+     */
+    getBuildHealthFilterId() {
+        return this.state.buildHealthFilterId;
+    }
+
     /**
      * Sets whether the linked selection mode is on when editing.
      * @param linkedSelectionModeEnabled True to enable linked selection, false otherwise.
diff --git a/Engine/Source/Programs/Horde/HordeDashboard/src/components/buildhealth/BuildhealthView.tsx b/Engine/Source/Programs/Horde/HordeDashboard/src/components/buildhealth/BuildhealthView.tsx
--- a/Engine/Source/Programs/Horde/HordeDashboard/src/components/buildhealth/BuildhealthView.tsx
+++ b/Engine/Source/Programs/Horde/HordeDashboard/src/components/buildhealth/BuildhealthView.tsx
@@ -1,7 +1,7 @@
 // Copyright Epic Games, Inc. All Rights Reserved.
 
-import { Checkbox, ComboBox, DefaultButton, DirectionalHint, Dropdown, DropdownMenuItemType, IconButton, IDropdownOption, IDropdownProps, ISelectableOption, Label, mergeStyleSets, PrimaryButton, SearchBox, SelectableOptionMenuItemType, Stack, TooltipHost } from "@fluentui/react";
-import { GetProjectResponse, GetStreamResponse } from "horde/backend/Api";
+import { Checkbox, ComboBox, DefaultButton, DirectionalHint, Dropdown, DropdownMenuItemType, IconButton, IDropdownOption, IDropdownProps, ISelectableOption, Label, mergeStyleSets, PrimaryButton, SearchBox, SelectableOptionMenuItemType, Stack, TextField, TooltipHost } from "@fluentui/react";
+import { GetBuildHealthFilterResponse, GetProjectResponse, GetStreamResponse } from "horde/backend/Api";
 import { observer } from "mobx-react-lite";
 import { useCallback, useEffect, useLayoutEffect, useMemo, useRef, useState } from "react";
 import { useLocation, useNavigate } from "react-router";
@@ -19,6 +19,7 @@ import { makeAutoObservable, observable, runInAction } from "mobx";
 import { BuildHealthOptionsController, StreamDescriptor, TemplateDescriptor } from "./BuildHealthOptionsController";
 import { StepOutcomeLegend } from "./stepoutcome/StepOutcomeLegend";
 import { StepOutcomeDataHandler, StepOutcomeFilters } from "./stepoutcome/StepOutcomeDataHandler";
+import { BuildHealthFilterComponent } from "./BuildHealthFilterComponent";
 
 // #region -- Dropdown Components --
 
@@ -162,14 +163,20 @@ const SearchableDropdown: React.FC<ISearchableDropdownProps> = (props) => {
             ref={dropdownRef}
             {...dropdownProps}
             calloutProps={{
-                directionalHint: DirectionalHint.rightBottomEdge,
+                calloutWidth: 510,
+                directionalHint: DirectionalHint.rightTopEdge,
                 alignTargetEdge: true,
                 styles: {
+                    root: {
+                        border: `1px solid ${modeColors.content}`,
+                        borderRadius: 4,
+                        boxSizing: "border-box",
+                    },
                     calloutMain: {
                         width: 500,
                         minWidth: 500,
-                        height: 1000,
-                        maxHeight: 1000,
+                        height: 500,
+                        maxHeight: 1000
                     }
                 },
             }}
@@ -250,13 +257,13 @@ const BasicListParameter: React.FC<{ handler: BuildHealthDataHandler, buildHealt
     return (<Dropdown id={changeVersion.toString() + key}
         key={key}
         label={param.label}
-        placeholder={"Select option"}
+        placeholder={"Select Project"}
         options={doptions}
         selectedKey={selectedKey}
         dropdownWidth={200}
         styles={{ dropdown: { width: 300 } }}
         calloutProps={{
-            directionalHint: DirectionalHint.rightBottomEdge,
+            directionalHint: DirectionalHint.rightTopEdge,
             alignTargetEdge: true,
         }}
         onChange={(_ev, option) => {
@@ -458,7 +465,9 @@ const MultiListParameter: React.FC<{ handler: BuildHealthDataHandler, buildHealt
             }}
             onRenderOption={onRenderOption}
             styles={{
-                dropdown: { width: 300, maxWidth: 900 },
+                dropdown: {
+                    width: 300, maxWidth: 900,
+                },
             }}
             disabled={disabled}
         />
@@ -551,7 +560,6 @@ const StreamDropdown: React.FC<{ handler: BuildHealthDataHandler, buildHealthOpt
 
     let selectedStreamKeys: string[] = Object.keys(buildHealthOptions.state.enabledStreams);
     let request: DataHandlerRefreshRequest = expandRefreshRequest({ streams: true });
-    let linkString = buildHealthLinkedModeNotepad.releaseDirty ? "Current State: Unlinked - Some templates, steps, and labels may require linking." : "Current State: Linked";
 
     return streamDropdownParams.items.length > 0 && (
         <Stack>
@@ -694,7 +702,6 @@ const TemplateDropdown: React.FC<{ handler: BuildHealthDataHandler, buildHealthO
 const StepDropdown: React.FC<{ handler: BuildHealthDataHandler, buildHealthOptions: BuildHealthOptionsController, buildHealthLinkedModeNotepad: BuildHealthLinkNotepad }> = observer(function ConstructProjectDropdown({ handler, buildHealthOptions, buildHealthLinkedModeNotepad }) {
     let selectedStepKeys: string[] = Object.keys(buildHealthOptions.state.enabledSteps);
     let request: DataHandlerRefreshRequest = { steps: false };
-
     // #region -- Data Memos --
 
     const items = useMemo<BuildHealthOptionData[]>(
@@ -1126,9 +1133,11 @@ const BuildHealthUrlSync: React.FC<{ buildHealthController: BuildHealthOptionsCo
     const location = useLocation();
     const navigate = useNavigate();
 
+    // We only ever receive from location search once.
     useEffect(() => {
         const initializeFromUrl = async () => {
-            const params = parseBuildHealthQueryParams(location);
+            let searchParams: URLSearchParams = new URLSearchParams(location.search)
+            const params = parseBuildHealthQueryParams(searchParams);
             const upgradedParams = upgradeBuildHealthQueryParams(params.buildHealthQueryParams, buildHealthController.querySchemaVersion);
             await loadBuildHealthOptionsFromParams(upgradedParams, params.buildHealthUIParams, buildHealthController, handler);
         };
@@ -1338,6 +1347,7 @@ const BuildHealthSidebar: React.FC<{ handler: BuildHealthDataHandler, buildHealt
                         />
                     </Stack>
                 </Stack>
+                <BuildHealthFilterComponent handler={handler} buildHealthOptions={buildHealthOptions} />
                 <Stack
                     styles={{
                         root: {
@@ -1457,6 +1467,7 @@ const BuildHealthSidebar: React.FC<{ handler: BuildHealthDataHandler, buildHealt
 export const BuildHealthView: React.FC = observer(function ConstructBuildHealthView() {
     const windowSize = useWindowSize();
     const { hordeClasses, modeColors } = getHordeStyling();
+    const [selectedItem, setSelectedItem] = useState<GetBuildHealthFilterResponse | null>(null);
 
     const jobHistorySpan = useMemo(() => ({
         start: buildHealthState.startDate,
@@ -1465,17 +1476,21 @@ export const BuildHealthView: React.FC = observer(function ConstructBuildHealthV
         // eslint-disable-next-line react-hooks/exhaustive-deps -- All necessary dependencies included for MobX observables
         [buildHealthState.startDate, buildHealthState.endDate]);
 
-    const filter: StepOutcomeFilters = useMemo(() => ({
-        streams: buildHealthState.stepOutcomeEnabledStreamKeys,
-        jobs: buildHealthState.stepOutcomeEnabledJobKeys,
-        steps: buildHealthState.stepOutcomeEnabledStepKeys,
-        labels: buildHealthState.stepOutcomeEnabledLabelKeys,
-        jobHistorySpan,
-        includePreflights: buildHealthState.includePreflight,
-        includeCancelledJobs: buildHealthState.includeCancelledJobs,
-        debugMode: buildHealthState.debugMode,
-        jobStartFilter: { supportsManualStarts: buildHealthState.stepOutcomeEnabledJobStartMethods.includes(MANUAL_START_METHOD), supportsScheduledStarts: buildHealthState.stepOutcomeEnabledJobStartMethods.includes(SCHEDULED_START_METHOD) }
-    }),
+    const filter: StepOutcomeFilters = useMemo(() => (
+        {
+            streams: buildHealthState.stepOutcomeEnabledStreamKeys,
+            jobs: buildHealthState.stepOutcomeEnabledJobKeys,
+            steps: buildHealthState.stepOutcomeEnabledStepKeys,
+            labels: buildHealthState.stepOutcomeEnabledLabelKeys,
+            jobHistorySpan,
+            includePreflights: buildHealthState.includePreflight,
+            includeCancelledJobs: buildHealthState.includeCancelledJobs,
+            debugMode: buildHealthState.debugMode,
+            jobStartFilter: {
+                supportsManualStarts: buildHealthState.stepOutcomeEnabledJobStartMethods.includes(MANUAL_START_METHOD), supportsScheduledStarts: buildHealthState.stepOutcomeEnabledJobStartMethods.includes(SCHEDULED_START_METHOD)
+            }
+        }
+    ),
         // eslint-disable-next-line react-hooks/exhaustive-deps -- All necessary dependencies included for MobX observables
         [
             buildHealthState.stepOutcomeEnabledStreamKeys,
diff --git a/Engine/Source/Programs/Horde/HordeDashboard/tests/BuildHealthOptions.test.ts b/Engine/Source/Programs/Horde/HordeDashboard/tests/BuildHealthOptions.test.ts
--- a/Engine/Source/Programs/Horde/HordeDashboard/tests/BuildHealthOptions.test.ts
+++ b/Engine/Source/Programs/Horde/HordeDashboard/tests/BuildHealthOptions.test.ts
@@ -25,10 +25,8 @@ describe('parseBuildHealthQueryParams', () => {
 		mockParams.set('includePreflight', compressedPreflight);
 		mockParams.set('schema', '1');
 
-		const location = { search: `?${mockParams.toString()}` } as Location;
-
 		// Act
-		const result = parseBuildHealthQueryParams(location);
+		const result = parseBuildHealthQueryParams(new URLSearchParams(mockParams.toString()));
 
 		// Assert
 		expect(result.buildHealthQueryParams.projectIds).toEqual(['projA', 'projB']);
@@ -63,10 +61,8 @@ describe('upgradeBuildHealthQueryParams', () => {
 		mockParams.set('includePreflight', compressedPreflight);
 		mockParams.set('schema', '1');
 
-		const location = { search: `?${mockParams.toString()}` } as Location;
-
 		// Act
-		const result = parseBuildHealthQueryParams(location);
+		const result = parseBuildHealthQueryParams(new URLSearchParams(mockParams.toString()));
 		const upgradedParams = upgradeBuildHealthQueryParams(result.buildHealthQueryParams, 2);
 
 		// Assert
