[Horde][BuildHealth][Dashboard] Export Summary column to CSV report, and sorting step rows by capabilities.

#Synced-CL 48788183

diff --git a/Engine/Source/Programs/Horde/HordeDashboard/src/components/buildhealth/stepoutcome/StepOutcome.tsx b/Engine/Source/Programs/Horde/HordeDashboard/src/components/buildhealth/stepoutcome/StepOutcome.tsx
--- a/Engine/Source/Programs/Horde/HordeDashboard/src/components/buildhealth/stepoutcome/StepOutcome.tsx
+++ b/Engine/Source/Programs/Horde/HordeDashboard/src/components/buildhealth/stepoutcome/StepOutcome.tsx
@@ -9,6 +9,7 @@ import { prettyPrintStepOutcomeFilters, prettyPrintStepOutcomeFiltersSummary, St
 import { AbortedStr, OutcomeTableEntry, isLabelOutcome, isStepOutcome, StepOutcomeTableEntry } from "./StepOutcomeDataTypes";
 import { NOT_APPLICABLE_LABEL, sharedBorderStyle } from "./StepOutcomeSharedUIComponents";
 import { isTerminalState, toDate } from "./StepOutcomeUtilities";
+import { getHordeTheme } from "horde/styles/theme";
 
 // #region -- View Components --
 
@@ -216,6 +217,7 @@ const StepOutcomeTableComponent: React.FC<StepOutcomeTableProps> = observer(({ h
     handler.updated;
     uiOptions;
     const [dataCount, setDataCount] = useState(() => handler.StepOutcomeTableData.dataCount);
+    const theme = getHordeTheme();
 
     useEffect(() => {
         setDataCount(handler.StepOutcomeTableData.dataCount);
@@ -234,6 +236,7 @@ const StepOutcomeTableComponent: React.FC<StepOutcomeTableProps> = observer(({ h
                             <Stack horizontalAlign="center" styles={{ root: { paddingTop: 8, paddingBottom: 8 } }}>
                                 <Text styles={{ root: { fontSize: "large", fontWeight: 'bold' } }}>{header}</Text>
                                 <Text styles={{ root: { fontSize: "small", fontWeight: 'bold' } }} >{subHeader} - {subSubHeader}</Text>
+                                <Text styles={{ root: { color: theme.semanticColors.bodyBackgroundHovered } }}>Last refresh: {handler.lastRefreshDate?.toLocaleTimeString([], { hour: '2-digit', minute: '2-digit', second: '2-digit', timeZoneName: 'short' })}</Text>
                             </Stack>
                             <StepOutcomeD3Table
                                 stepOutcomeTable={handler.StepOutcomeTableData}
diff --git a/Engine/Source/Programs/Horde/HordeDashboard/src/components/buildhealth/stepoutcome/StepOutcomeD3Component.tsx b/Engine/Source/Programs/Horde/HordeDashboard/src/components/buildhealth/stepoutcome/StepOutcomeD3Component.tsx
--- a/Engine/Source/Programs/Horde/HordeDashboard/src/components/buildhealth/stepoutcome/StepOutcomeD3Component.tsx
+++ b/Engine/Source/Programs/Horde/HordeDashboard/src/components/buildhealth/stepoutcome/StepOutcomeD3Component.tsx
@@ -5,8 +5,8 @@ import * as d3 from "d3";
 import { getHordeStyling, preloadFonts } from "horde/styles/Styles";
 import { useEffect, useRef } from "react";
 import { StepOutcomeDataHandler } from "./StepOutcomeDataHandler";
-import { AbortedStr, ChangeHeader, ColumnHeader, encodeStepNameFromStepNameHeader, isChangeHeader, isDateHeader, isLabelOutcome, isStepOutcome, isSummary, isSummaryHeader, StepNameHeader, StepOutcomeTable, SummaryTableEntry, TableEntry } from "./StepOutcomeDataTypes";
-import { getCellAccents, getCellClasses, getCellStyle, getColorRecords, IncrementalRefreshIndicatorPanel, NOT_APPLICABLE_LABEL, resetCellClasses, sharedBorderStyle } from "./StepOutcomeSharedUIComponents";
+import { AbortedStr, ChangeHeader, ColumnHeader, encodeStepNameFromStepNameHeader, isChangeHeader, isDateHeader, isLabelOutcome, isStepOutcome, isSummary, isSummaryHeader, StepNameHeader, StepOutcomeTable, TableEntry } from "./StepOutcomeDataTypes";
+import { getCellAccents, getCellClasses, getCellSortClasses, getCellStyle, getColorRecords, IncrementalRefreshIndicatorPanel, NOT_APPLICABLE_LABEL, resetCellClasses, sharedBorderStyle, headerStyles } from "./StepOutcomeSharedUIComponents";
 import { getHordeTheme } from "horde/styles/theme";
 import { observer } from "mobx-react-lite";
 import { isTerminalState, toLocalISOString } from "./StepOutcomeUtilities";
@@ -67,6 +67,10 @@ function exportStepOutcomeDataToCSV(stepOutcomeTable: StepOutcomeTable, handler:
 
     saveFile(labelOrderedCsvContent, `${filename}-LabelOrderedData-${defaultName}.csv`);
     saveFile(labelchangeOrderedCsvContent, `${filename}-LabelChangeOrderedData-${defaultName}.csv`);
+
+    const summaryOrderedCsvContent = stepOutcomeTable.summaryToRowOrderedCSV();
+
+    saveFile(summaryOrderedCsvContent, `${filename}-SummaryData-${defaultName}.csv`);
 }
 
 // #endregion -- Step Outcome Table Download Handlers --
@@ -189,7 +193,6 @@ function GenerateD3Table(stepOutcomeTable: StepOutcomeTable, handler: StepOutcom
 
     let generateTableTooltipRow = (tooltipContent: HTMLDivElement, label: string, value: string) => {
         const tooltipRowElement = document.createElement("div");
-
         const tooltipLabelElement = document.createElement("strong");
         tooltipLabelElement.textContent = `${label}:`;
 
@@ -202,52 +205,6 @@ function GenerateD3Table(stepOutcomeTable: StepOutcomeTable, handler: StepOutcom
         tooltipContent.appendChild(tooltipRowElement);
     }
 
-    let computePassRatio = (entry: SummaryTableEntry) => {
-        let ratio: number = -1;
-
-        if (entry.steps > 0 && entry.stepsIncluded > 0) {
-            let numerator: number = entry.stepsPass;
-
-            if (!stepOutcomeTable.viewModelOptions.warningsAsSummaryFailure) {
-                numerator += entry.stepsWarning;
-            }
-
-            ratio = numerator / entry.stepsIncluded;
-        }
-
-        let ratioStr: string = ratio === -1 ? NOT_APPLICABLE_LABEL : `${(ratio * 100).toFixed(1)}%`;
-        return { ratio, ratioStr };
-    }
-
-    let computeSkipRatio = (entry: SummaryTableEntry) => {
-        let skippedRatio: number = -1;
-
-        if (entry.steps > 0 && entry.stepsIncluded > 0) {
-            let numerator: number = entry.stepsSkipped;
-            skippedRatio = numerator / entry.stepsIncluded;
-        }
-
-        let skippedRatioStr: string = skippedRatio === -1 ? NOT_APPLICABLE_LABEL : `${(skippedRatio * 100).toFixed(1)}%`;
-        return { skippedRatio, skippedRatioStr };
-    }
-
-    let computeCompletedPassRatio = (entry: SummaryTableEntry) => {
-        let truePassRatio: number = -1;
-
-        if (entry.steps > 0 && entry.stepsIncluded > 0) {
-            let numerator: number = entry.stepsPass;
-
-            if (!stepOutcomeTable.viewModelOptions.warningsAsSummaryFailure) {
-                numerator += entry.stepsWarning;
-            }
-
-            truePassRatio = numerator / (entry.stepsIncluded - entry.stepsSkipped);
-        }
-
-        let truePassRatioStr: string = truePassRatio === -1 ? NOT_APPLICABLE_LABEL : `${(truePassRatio * 100).toFixed(1)}%`;
-        return { truePassRatio, truePassRatioStr };
-    }
-
     // #endregion -- Tooltip Helpers --
 
     const containerRef = useRef<HTMLDivElement>(null);
@@ -308,7 +265,7 @@ function GenerateD3Table(stepOutcomeTable: StepOutcomeTable, handler: StepOutcom
             .style("padding", "6px 12px")
             .style("background", "rgba(0,0,0,0.95)")
             .style("color", "#fff")
-            .style("z-index", "5")
+            .style("z-index", "100")
             .style("min-width", "250px")
             .style("font-size", "12px")
             .style("border-radius", "2px")
@@ -317,17 +274,6 @@ function GenerateD3Table(stepOutcomeTable: StepOutcomeTable, handler: StepOutcom
 
         const theaderrow = thead.append("tr").style("height", "100px");
 
-        // First column header; (0,0) - we just want to hide this; sits on top of group; guarantee the width of groups TH.
-        theaderrow.append("th")
-            .style("position", "sticky")
-            .style("z-index", "6")
-            .style("top", "0")
-            .style("left", "0")
-            .style("background-color", modeColors.background)
-            .style("width", "30px")
-            .style("min-width", "30px")
-            .style("max-width", "30px");
-
         scaffold.current = { wrapper, table, thead, theaderrow, tbody, tooltip };
 
         return () => {
@@ -339,41 +285,93 @@ function GenerateD3Table(stepOutcomeTable: StepOutcomeTable, handler: StepOutcom
     useEffect(() => {
         if (!containerRef.current || !scaffold.current) return;
 
+        // First column header; (0,0) - this will be responsible for resetting to step insertion order
+        scaffold.current.theaderrow
+            .selectAll("th.start-header")
+            .data(["start"])
+            .join(
+                enter => enter.append("th")
+                    .attr("class", d => `start-header ${getCellSortClasses("start", stepOutcomeTable.sortField, stepOutcomeTable.sortDirectionAsc)}`)
+                    .call(th => th
+                        .style("position", "sticky")
+                        .style("z-index", "7")
+                        .style("top", "0")
+                        .style("background-color", modeColors.content)
+                        .style("width", "30px")
+                        .style("min-width", "30px")
+                        .style("max-width", "30px")
+                        .style("cursor", "pointer")
+                    ),
+                update => update
+                    .attr("class", d => `start-header ${getCellSortClasses("start", stepOutcomeTable.sortField, stepOutcomeTable.sortDirectionAsc)}`)
+                    .call(th => th
+                        .style("background-color", modeColors.content)
+                    ),
+                exit => exit.remove()
+            )
+            .on("mouseover", function (_event, _d) {
+                d3.select(this).style("background-color", theme.semanticColors.bodyBackgroundHovered);
+
+                const tooltipContent = document.createElement("div");
+                tooltipContent.style.paddingLeft = "4px";
+                let tooltipLabel: string = "";
+                let tooltipValue: string = "";
+
+                if (stepOutcomeTable.sortField !== "start") {
+                    tooltipLabel = "Click to sort by";
+                    tooltipValue = "Start";
+                } else {
+                    tooltipLabel = "Sort by Start";
+                    tooltipValue = stepOutcomeTable.sortDirectionAsc ? "Asc" : "Desc"
+                }
+
+                generateTableTooltipRow(tooltipContent, tooltipLabel, tooltipValue);
+                attachTargettedToolTip(this as HTMLElement, scaffold.current!.tooltip, tooltipContent);
+            })
+            .on("mouseout", function () {
+                scaffold.current!.tooltip.style("opacity", 0);
+                d3.select(this).style("background-color", modeColors.content);
+            })
+            .on("click", (_event, _d) => {
+                scaffold.current!.tooltip.style("opacity", 0);
+                stepOutcomeTable.changeOrderingBy("start");
+            });
+
+
         // Second Column Header; (0, 1) - used to display refresh time - this sits on top of step name; guarantee the width of step name TH.
         scaffold.current.theaderrow
             .selectAll("th.refresh-header")
-            .data([handler.lastRefreshDate])
-            .join("th")
-            .attr("class", "refresh-header")
+            .data(d3.range(1))
+            .join(enter => enter.append("th")
+                .attr("class", d => `refresh-header ${getCellSortClasses("name", stepOutcomeTable.sortField, stepOutcomeTable.sortDirectionAsc)}`),
+                update => update.attr("class", d => `refresh-header ${getCellSortClasses("name", stepOutcomeTable.sortField, stepOutcomeTable.sortDirectionAsc)}`),
+                exit => exit.remove())
             .style("position", "sticky")
             .style("z-index", "6")
             .style("top", "0")
             .style("left", "30px")
-            .style("background-color", modeColors.background)
+            .style("background-color", modeColors.content)
             .style("width", "320px")
             .style("min-width", "320px")
             .style("max-width", "320px")
             .style("cursor", "pointer")
-            .text(d =>
-                d
-                    ? `Last Refresh: ${d.toLocaleTimeString([], {
-                        hour: '2-digit',
-                        minute: '2-digit',
-                        second: '2-digit',
-                        timeZoneName: 'short'
-                    })}`
-                    : ""
-            )
+            .text("Step Name")
             .on("mouseover", function (_event, _d) {
                 d3.select(this).style("background-color", theme.semanticColors.bodyBackgroundHovered);
-                const totalSeconds = Math.floor(handler.refreshTime / 1000);
-                const minutes = Math.floor(totalSeconds / 60);
-                const seconds = totalSeconds % 60;
-                const formatted = `${minutes}m ${seconds}s`;
-
                 const tooltipContent = document.createElement("div");
                 tooltipContent.style.paddingLeft = "4px";
-                generateTableTooltipRow(tooltipContent, "Refresh Rate", formatted);
+                let tooltipLabel: string = "";
+                let tooltipValue: string = "";
+
+                if (stepOutcomeTable.sortField !== "name") {
+                    tooltipLabel = "Click to sort by";
+                    tooltipValue = "Step Name";
+                } else {
+                    tooltipLabel = "Sort by Step Name";
+                    tooltipValue = stepOutcomeTable.sortDirectionAsc ? "Asc" : "Desc"
+                }
+
+                generateTableTooltipRow(tooltipContent, tooltipLabel, tooltipValue);
                 attachTargettedToolTip(this as HTMLElement, scaffold.current!.tooltip, tooltipContent);
             })
             .on("mouseout", function () {
@@ -384,36 +382,64 @@ function GenerateD3Table(stepOutcomeTable: StepOutcomeTable, handler: StepOutcom
                 if (handler.isInFullDataRefresh || handler.isInIncrementalDataRefresh) {
                     return;
                 }
-
-                handler.poll();
+                scaffold.current!.tooltip.style("opacity", 0);
+                stepOutcomeTable.changeOrderingBy("name");
             });
 
         // Add all columns to the header row; these will be the changelists & the summary; these control downstream widths
         scaffold.current.theaderrow.selectAll("th.col")
             .data([...stepOutcomeTable.getColumnHeaders()] as ColumnHeader[])
             .join("th")
-            .attr("class", "col")
+            .attr("class", d => {
+                if (!isSummaryHeader(d)) {
+                    return "col";
+                }
+                const arrowClass = getCellSortClasses("summary", stepOutcomeTable.sortField, stepOutcomeTable.sortDirectionAsc);
+                return `col ${arrowClass}`;
+            })
             .style("position", "sticky")
             .style("top", "0")
             .style("z-index", "3")
             .style("background-color", modeColors.content)
-            .style("width", d => isSummaryHeader(d) ? "70px" : "30px")
-            .style("font-size", d => (isDateHeader(d) || isSummaryHeader(d) ? "large" : "small"))
-            .style("right", d => isSummaryHeader(d) ? "0" : null)
-            .style("z-index", d => isSummaryHeader(d) ? 4 : 3)
+            .style("width", d => (isSummaryHeader(d) ? "70px" : "30px"))
+            .style("font-size", d =>
+                isDateHeader(d) || isSummaryHeader(d) ? "large" : "small"
+            )
+            .style("right", d => (isSummaryHeader(d) ? "0" : null))
+            .style("z-index", d => (isSummaryHeader(d) ? 4 : 3))
             .style("box-sizing", "border-box")
-            .style("writing-mode", "vertical-rl")
-            .style("transform", "rotate(180deg)")
             .style("padding-bottom", "5px")
-            .style("text-align", "right")
             .style("border", `1px solid ${modeColors.background}`)
-            .style("border", d => isSummaryHeader(d) ? `2px solid ${modeColors.content}` : `1px solid ${modeColors.background}`)
-
-            .text(entry => isSummaryHeader(entry) ? entry.name : (isChangeHeader(entry) ? entry.change : (entry.date ?? "")))
+            .style("border", d =>
+                isSummaryHeader(d)
+                    ? `2px solid ${modeColors.content}`
+                    : `1px solid ${modeColors.background}`
+            )
+            .style("cursor", entry => (isSummaryHeader(entry) ? "pointer" : ""))
+            .html(entry => `<span class="${headerStyles.verticalText}">${isSummaryHeader(entry) ? entry.name : isChangeHeader(entry) ? entry.change : entry.date ?? ""}</span>`)
             .on("mouseover", function (_event, d) {
                 const entry = d as unknown as ChangeHeader;
 
-                if (isDateHeader(entry) || isSummaryHeader(entry)) {
+                if (isDateHeader(entry)) {
+                    return;
+                }
+
+                if (isSummaryHeader(entry)) {
+                    const tooltipContent = document.createElement("div");
+                    tooltipContent.style.paddingLeft = "4px";
+                    let tooltipLabel: string = "";
+                    let tooltipValue: string = "";
+
+                    if (stepOutcomeTable.sortField !== "summary") {
+                        tooltipLabel = "Click to sort by";
+                        tooltipValue = "Success Ratio";
+                    } else {
+                        tooltipLabel = "Sort by Success Ratio";
+                        tooltipValue = stepOutcomeTable.sortDirectionAsc ? "Asc" : "Desc"
+                    }
+
+                    generateTableTooltipRow(tooltipContent, tooltipLabel, tooltipValue);
+                    attachTargettedToolTip(this as HTMLElement, scaffold.current!.tooltip, tooltipContent);
                     return;
                 }
 
@@ -441,6 +467,13 @@ function GenerateD3Table(stepOutcomeTable: StepOutcomeTable, handler: StepOutcom
             .on("mouseout", function () {
                 scaffold.current!.tooltip.style("opacity", 0);
                 d3.select(this).style("background-color", modeColors.content);
+            })
+            .on("click", (_event, d) => {
+                const entry = d as unknown as ChangeHeader;
+                if (isSummaryHeader(entry)) {
+                    scaffold.current!.tooltip.style("opacity", 0);
+                    stepOutcomeTable.changeOrderingBy("summary");
+                }
             });
 
         const rows = scaffold.current.tbody.selectAll("tr")
@@ -656,9 +689,9 @@ function GenerateD3Table(stepOutcomeTable: StepOutcomeTable, handler: StepOutcom
                                 })
                         getCellClasses(entry, handler.getJobResponseData(entry.jobId)).forEach(cls => div.classed(cls, true));
                     } else if (isSummary(entry)) { // SummaryTableEntry
-                        let { ratio, ratioStr }: { ratio: number; ratioStr: string; } = computePassRatio(entry);
-                        let { skippedRatio, skippedRatioStr }: { skippedRatio: number; skippedRatioStr: string; } = computeSkipRatio(entry);
-                        let { truePassRatio, truePassRatioStr }: { truePassRatio: number; truePassRatioStr: string; } = computeCompletedPassRatio(entry);
+                        let { ratio, ratioStr }: { ratio: number; ratioStr: string; } = entry.computePassRatio(stepOutcomeTable.viewModelOptions.warningsAsSummaryFailure);
+                        let { skippedRatio, skippedRatioStr }: { skippedRatio: number; skippedRatioStr: string; } = entry.computeSkipRatio();
+                        let { truePassRatio, truePassRatioStr }: { truePassRatio: number; truePassRatioStr: string; } = entry.computeCompletedPassRatio(stepOutcomeTable.viewModelOptions.warningsAsSummaryFailure);
 
                         div.style("cursor", "default")
                             .style("background-color", summaryColorScale(ratio))
@@ -728,7 +761,7 @@ function GenerateD3Table(stepOutcomeTable: StepOutcomeTable, handler: StepOutcom
                 .style("height", "25px");
         });
 
-    }, [stepOutcomeTable.tableEntries, stepOutcomeTable.viewModelOptions, handler.lastRefreshDate, getCellStyle]);
+    }, [stepOutcomeTable.tableEntries, stepOutcomeTable.viewModelOptions, handler.lastRefreshDate, stepOutcomeTable.sortingKey, getCellStyle]);
 
     const gridStyle: React.CSSProperties = {
         display: 'grid',
@@ -767,7 +800,7 @@ function GenerateD3Table(stepOutcomeTable: StepOutcomeTable, handler: StepOutcom
             window.removeEventListener("resize", updateMaxHeight);
         };
     }, []);
-
+    ""
     return (
         <div style={gridStyle}>
             {/* Left side: table */}
@@ -791,11 +824,12 @@ function GenerateD3Table(stepOutcomeTable: StepOutcomeTable, handler: StepOutcom
                 }}
             >
                 {/* Top: spinner */}
+                {/* SubTop: refresh button*/}
                 <IncrementalRefreshIndicatorPanel handler={handler} />
-                {/* Bottom: button */}
+                {/* Bottom: download button */}
                 <TooltipHost
                     content="Export the current step outcome history data to CSV."
-                    directionalHint={DirectionalHint.topLeftEdge}
+                    directionalHint={DirectionalHint.leftCenter}
                 >
                     <IconButton
                         iconProps={{ iconName: 'Download' }}
diff --git a/Engine/Source/Programs/Horde/HordeDashboard/src/components/buildhealth/stepoutcome/StepOutcomeDataHandler.ts b/Engine/Source/Programs/Horde/HordeDashboard/src/components/buildhealth/stepoutcome/StepOutcomeDataHandler.ts
--- a/Engine/Source/Programs/Horde/HordeDashboard/src/components/buildhealth/stepoutcome/StepOutcomeDataHandler.ts
+++ b/Engine/Source/Programs/Horde/HordeDashboard/src/components/buildhealth/stepoutcome/StepOutcomeDataHandler.ts
@@ -462,6 +462,8 @@ export class StepOutcomeDataHandler extends PollBase {
 
         this.stepOutcomeTableData.orderTableDataByChange(true);
 
+        this.stepOutcomeTableData.changeOrderingBy(this.stepOutcomeTableData.sortField, false);
+
         this.setUpdated();
 
         runInAction(() => {
diff --git a/Engine/Source/Programs/Horde/HordeDashboard/src/components/buildhealth/stepoutcome/StepOutcomeDataTypes.ts b/Engine/Source/Programs/Horde/HordeDashboard/src/components/buildhealth/stepoutcome/StepOutcomeDataTypes.ts
--- a/Engine/Source/Programs/Horde/HordeDashboard/src/components/buildhealth/stepoutcome/StepOutcomeDataTypes.ts
+++ b/Engine/Source/Programs/Horde/HordeDashboard/src/components/buildhealth/stepoutcome/StepOutcomeDataTypes.ts
@@ -3,7 +3,7 @@
 import { ChangeSummaryData, GetJobResponse, GetJobStepRefResponse, GetLabelStateResponse, GetStepResponse, IssueData, JobStepOutcome, JobStepState, LabelOutcome, LabelState } from "horde/backend/Api";
 import { computeDuration, DURATION_NOT_SET, formatDurationFromMs, toDate } from "./StepOutcomeUtilities";
 import { StepOutcomeDataHandler, StepOutcomeViewOptions } from "./StepOutcomeDataHandler";
-import { action, makeObservable, observable } from "mobx";
+import { action, computed, makeObservable, observable } from "mobx";
 import { NOT_APPLICABLE_LABEL } from "./StepOutcomeSharedUIComponents";
 
 export const WarningStr: string = "Warnings";
@@ -65,7 +65,7 @@ export type UpdatedChangeSummaryData = ChangeSummaryData & {
 */
 export type JobResponseData = GetJobResponse & {
     summarizedOutcome: JobStepOutcome;
-    isManualJobInvocation : boolean;
+    isManualJobInvocation: boolean;
 }
 
 // #region -- Table Entry Types --
@@ -113,6 +113,82 @@ export class SummaryTableEntry {
     stepsWarning: number = 0;
     stepsPending: number = 0;
     stepsSkipped: number = 0;
+
+    /**
+     * Computes the pass ratio of the summary entry.
+     * @param warningsAsSummaryFailre Whether to include warnings as failures.
+     * @returns The numerical and string representation of the pass ratio.
+     */
+    computePassRatio(warningsAsSummaryFailre: boolean) {
+        let ratio: number = -1;
+
+        if (this.steps > 0 && this.stepsIncluded > 0) {
+            let numerator: number = this.stepsPass;
+
+            if (!warningsAsSummaryFailre) {
+                numerator += this.stepsWarning;
+            }
+
+            ratio = numerator / this.stepsIncluded;
+        }
+
+        let ratioStr: string = ratio === -1 ? NOT_APPLICABLE_LABEL : `${(ratio * 100).toFixed(1)}%`;
+        return { ratio, ratioStr };
+    }
+
+    /**
+     * Computes the skip ratio of the summary entry.
+     * @returns The numerical and string representation of the skip ratio.
+     */
+    computeSkipRatio() {
+        let skippedRatio: number = -1;
+
+        if (this.steps > 0 && this.stepsIncluded > 0) {
+            let numerator: number = this.stepsSkipped;
+            skippedRatio = numerator / this.stepsIncluded;
+        }
+
+        let skippedRatioStr: string = skippedRatio === -1 ? NOT_APPLICABLE_LABEL : `${(skippedRatio * 100).toFixed(1)}%`;
+        return { skippedRatio, skippedRatioStr };
+    }
+
+    /**
+     * Computes the completed pass ratio of the summary entry.
+     * @param warningsAsSummaryFailre Whether to include warnings as failures.
+     * @returns The numerical and string representation of the completed pass ratio.
+     */
+    computeCompletedPassRatio(warningsAsSummaryFailre: boolean) {
+        let truePassRatio: number = -1;
+
+        if (this.steps > 0 && this.stepsIncluded > 0) {
+            let numerator: number = this.stepsPass;
+
+            if (!warningsAsSummaryFailre) {
+                numerator += this.stepsWarning;
+            }
+
+            truePassRatio = (this.stepsIncluded - this.stepsSkipped) === 0 ? -1 : numerator / (this.stepsIncluded - this.stepsSkipped);
+        }
+
+        let truePassRatioStr: string = truePassRatio === -1 ? NOT_APPLICABLE_LABEL : `${(truePassRatio * 100).toFixed(1)}%`;
+        return { truePassRatio, truePassRatioStr };
+    }
+
+    /**
+     * The CSV representation of the summary entry.
+     * @returns String representation of CSV row.
+     */
+    toCSVRepresentation() {
+        return `${this.steps},${this.stepsIncluded},${this.stepsPass},${this.stepsFail},${this.stepsWarning},${this.stepsPending},${this.stepsSkipped},${this.stepsFail - this.stepsSkipped},${this.computePassRatio(true).ratioStr},${this.computePassRatio(false).ratioStr},${this.computeCompletedPassRatio(true).truePassRatioStr},${this.computeCompletedPassRatio(false).truePassRatioStr},${this.computeSkipRatio().skippedRatioStr}\n`
+    }
+
+    /**
+     * The CSV header of the summary entry.
+     * @returns The CSV header representing the columns.
+     */
+    static getCSVHeader(): string {
+        return "Step,StepCount,StepsIncluded,StepsPassed,StepsFailed,StepsWarning,StepsPending,StepsSkipped,FailedCompletedSteps,PassRatio(WarningsAsError),PassRatio(WarningsNotAsError),CompletedPassRatio(WarningsAsError),ComplatePassRatio,SkipRatio\n";
+    }
 }
 
 /**
@@ -626,6 +702,9 @@ export type StreamMetadata = {
 export class StepOutcomeTable {
     // #region -- Public Members --
 
+    @observable sortField: "name" | "summary" | "start" | "none" = "none";
+    @observable sortDirectionAsc: boolean = true;
+
     dataCount: number = 0;
 
     /**
@@ -729,6 +808,15 @@ export class StepOutcomeTable {
 
     // #region -- Public API --
 
+    /**
+     * Gets the current sort key.
+     * @returns The current sort key of the table.
+     */
+    @computed
+    get sortingKey() {
+        return `${this.sortField}:${this.sortDirectionAsc}`;
+    }
+
     /**
      * Resets the table data. This will remove all data from the container.
      */
@@ -974,13 +1062,76 @@ export class StepOutcomeTable {
     }
 
     /**
-     * Sets the table data's step (row) order.
+     * Requests a change of ordering type.
+     * @param requestedOrdering The requested ordering field type.
+     * @param toggleIfSameOrdering If true, will togggle the sort direction if the sort field hasn't changed types.
+     */
+    changeOrderingBy(requestedOrdering: "name" | "summary" | "start" | "none", toggleIfSameOrdering: boolean = true): void {
+        let newSortDirectionAsc: boolean = toggleIfSameOrdering && this.sortField === requestedOrdering ? !this.sortDirectionAsc : this.sortDirectionAsc;
+        switch (requestedOrdering) {
+            case "name":
+                this.orderTableDataByStepName(newSortDirectionAsc);
+                break;
+            case "summary":
+                this.orderTableDataBySummaryPassRatio(newSortDirectionAsc);
+                break;
+            case "start":
+                this.orderTableDataByStepStart(newSortDirectionAsc);
+                break;
+            case "none":
+                this.sortField = "none";
+                break;
+        }
+    }
+
+    /**
+     * Orders table data based on step "start".
      * @param ascending True for ascending order; false for descending order.
      */
+    @action
+    orderTableDataByStepStart(ascending: boolean): void {
+        this.sortField = "start";
+        this.sortDirectionAsc = ascending;
+
+        this.stepNameRowHeaders.sort((a, b) => {
+            let aInsertionRow = this.stepRowLookup.get(encodeStepNameFromStepNameHeader(a))!;
+            let bInsertionRow = this.stepRowLookup.get(encodeStepNameFromStepNameHeader(b))!;
+            if (aInsertionRow < bInsertionRow) return ascending ? -1 : 1;
+            if (aInsertionRow > bInsertionRow) return ascending ? 1 : -1;
+            return 0;
+        });
+    }
+
+    /**
+     * Orders table data based on step name.
+     * @param ascending True for ascending order; false for descending order.
+     */
+    @action
     orderTableDataByStepName(ascending: boolean): void {
+        this.sortField = "name";
+        this.sortDirectionAsc = ascending;
+
+        this.stepNameRowHeaders.sort((a, b) => {
+            if (a.stepName < b.stepName) return ascending ? -1 : 1;
+            if (a.stepName > b.stepName) return ascending ? 1 : -1;
+            return 0;
+        });
+    }
+
+    /**
+     * Orders table by the summary pass ratio.
+     * @param ascending True for ascending order; false for descending order.
+     */
+    @action
+    orderTableDataBySummaryPassRatio(ascending: boolean): void {
+        this.sortField = "summary";
+        this.sortDirectionAsc = ascending;
         this.stepNameRowHeaders.sort((a, b) => {
-            if (a < b) return ascending ? -1 : 1;
-            if (a > b) return ascending ? 1 : -1;
+            let aSummary = this.summaryLookup.get(encodeStepNameFromStepNameHeader(a));
+            let bSummary = this.summaryLookup.get(encodeStepNameFromStepNameHeader(b));
+
+            if (aSummary?.computePassRatio(false).ratio! < bSummary?.computePassRatio(false).ratio!) return ascending ? -1 : 1;
+            if (aSummary?.computePassRatio(false).ratio! > bSummary?.computePassRatio(false).ratio!) return ascending ? 1 : -1;
             return 0;
         });
     }
@@ -1091,6 +1242,23 @@ export class StepOutcomeTable {
         return resultString!;
     }
 
+    /**
+     * Generates the csv string representation of the step outcome table's summary column.
+     * @returns A string in CSV format that represents the underyling step outcome table's summary column.
+     */
+    summaryToRowOrderedCSV() {
+        let resultString: string | undefined;
+
+        for (let [stepName, summary] of this.summaryLookup.entries()) {
+            if (resultString === undefined) {
+                resultString = SummaryTableEntry.getCSVHeader();
+            }
+            resultString += `${stepName},` + summary.toCSVRepresentation();
+        }
+
+        return resultString!;
+    }
+
     /**
      * Sets the view model options to use when interacting with the StepOutcomeTable.
      * @param viewModelOptions The view model options to set.
diff --git a/Engine/Source/Programs/Horde/HordeDashboard/src/components/buildhealth/stepoutcome/StepOutcomeSharedUIComponents.tsx b/Engine/Source/Programs/Horde/HordeDashboard/src/components/buildhealth/stepoutcome/StepOutcomeSharedUIComponents.tsx
--- a/Engine/Source/Programs/Horde/HordeDashboard/src/components/buildhealth/stepoutcome/StepOutcomeSharedUIComponents.tsx
+++ b/Engine/Source/Programs/Horde/HordeDashboard/src/components/buildhealth/stepoutcome/StepOutcomeSharedUIComponents.tsx
@@ -1,6 +1,6 @@
 // Copyright Epic Games, Inc. All Rights Reserved.
 
-import { mergeStyleSets, Spinner, SpinnerSize } from "@fluentui/react";
+import { DirectionalHint, IconButton, mergeStyleSets, Spinner, SpinnerSize, TooltipHost } from "@fluentui/react";
 import dashboard, { StatusColor } from "horde/backend/Dashboard";
 import { observer } from "mobx-react-lite";
 import { StepOutcomeDataHandler } from "./StepOutcomeDataHandler";
@@ -16,6 +16,57 @@ export const ROW_HEADER_NAME: string = "stepName" as const;
 
 // #region -- Styles --
 
+/**
+ * Step outcome table Header styles.
+ */
+export const headerStyles = mergeStyleSets({
+    verticalText: {
+        writingMode: "vertical-rl",
+        transform: "rotate(180deg)",
+        display: "inline-block",
+    }
+})
+
+const sortStyles = mergeStyleSets({
+    upArrow: {
+        '::after': {
+            content: '"▲"',
+            position: 'absolute',
+            bottom: 4,
+            right: 4,
+            fontSize: 12,
+            lineHeight: 1,
+        }
+    },
+    downArrow: {
+        position: "relative",
+        selectors: {
+            '::after': {
+                content: '"▼"',
+                position: 'absolute',
+                bottom: 4,
+                right: 4,
+                fontSize: 12,
+                lineHeight: 1,
+            }
+        }
+    },
+
+    dash: {
+        position: "relative",
+        selectors: {
+            '::after': {
+                content: '"-"',
+                position: 'absolute',
+                bottom: 4,
+                right: 4,
+                fontSize: 12,
+                lineHeight: 1,
+            }
+        },
+    }
+});
+
 export const sharedBorderStyle = {
     bottomBorder: {
         borderBottom: `1px solid rgba(80, 80, 80, 1)`
@@ -137,6 +188,22 @@ export function getCellClasses(entry: OutcomeTableEntry | null, jobState: JobRes
     return classes;
 }
 
+/**
+ * Gets the classes for a step outcome table header cell. These are used to add styling to a sortable column.
+ * @param sourceSortField The requesting cells sort field type.
+ * @param sortField The current sort field type.
+ * @param sortDirectionAsc The sort direction.
+ * @returns Class string as it relates to the requestor.
+ */
+export function getCellSortClasses(sourceSortField: "name" | "summary" | "start", sortField: "name" | "summary" | "start" | "none", sortDirectionAsc: boolean): string {
+    if (sourceSortField != sortField) {
+        return sortStyles.dash;
+    }
+
+    return sortDirectionAsc ? sortStyles.upArrow : sortStyles.downArrow;
+}
+
+
 /**
  * Gets the accents for a step outcome table entry cell. These are used to highlight steps that are currently in flight, or have not run to completion (Aborted OR Skipped).
  * @param entry The step outcome table entry of which to base the accents on.
@@ -213,16 +280,35 @@ export function getCellStyle(entry: OutcomeTableEntry | null): { bgColor: string
  * @param handler The data handler used to query step outcome data.
  * @returns React Component.
  */
-export const IncrementalRefreshIndicatorPanel: React.FC<{ handler: StepOutcomeDataHandler }> = observer(({ handler }) => (
-    <div
-        style={{
-            gridRow: '1 / span 2',
-            gridColumn: 3,
-            overflowY: 'auto',
-            width: 30,
-            height: 30,
-        }}
-    >
-        {handler.isInIncrementalDataRefresh && <Spinner size={SpinnerSize.small} />}
-    </div>
-));
\ No newline at end of file
+export const IncrementalRefreshIndicatorPanel: React.FC<{ handler: StepOutcomeDataHandler }> = observer(({ handler }) => {
+    const totalSeconds = Math.floor(handler.refreshTime / 1000);
+    const minutes = Math.floor(totalSeconds / 60);
+    const seconds = totalSeconds % 60;
+    const formatted = `${minutes}m ${seconds}s`;
+
+    return (
+
+        <div
+            style={{
+                gridRow: '1 / span 2',
+                gridColumn: 3,
+                overflowY: 'auto',
+                width: 50,
+                height: 50,
+            }}
+        >
+            {handler.isInIncrementalDataRefresh && <Spinner size={SpinnerSize.medium} />}
+            {!handler.isInIncrementalDataRefresh &&
+                <TooltipHost
+                    content={`Refresh the table. (Refresh rate: ${formatted})`}
+                    directionalHint={DirectionalHint.leftCenter}
+                >
+                    <IconButton
+
+                        iconProps={{ iconName: 'Refresh' }}
+                        onClick={() => handler.poll()}
+                    />
+                </TooltipHost>}
+        </div>
+    )
+});
\ No newline at end of file
