Horde: Updating schema documentation

#rnx

diff --git a/Engine/Source/Programs/Horde/Docs/Config/Schema/AclActions.md b/Engine/Source/Programs/Horde/Docs/Config/Schema/AclActions.md
--- a/Engine/Source/Programs/Horde/Docs/Config/Schema/AclActions.md
+++ b/Engine/Source/Programs/Horde/Docs/Config/Schema/AclActions.md
@@ -81,6 +81,12 @@
 | `UpdateBisectTask` | Ability to update a bisect task |
 | `ViewBisectTask` | Ability to view a bisect task |
 
+## BuildHealth
+
+| Name | Description |
+| ---- | ----------- |
+| `AddUpdateDeleteBuildHealthFilter` | Ability to add, update, and delete build health filters. |
+
 ## Devices
 
 | Name | Description |
@@ -187,6 +193,12 @@
 | `UpdateBisectTask` |  |
 | `ViewBisectTask` |  |
 
+## BuildHealth
+
+| Name | Description |
+| ---- | ----------- |
+| `AddUpdateDeleteBuildHealthFilter` |  |
+
 ## Compute
 
 | Name | Description |
@@ -345,6 +357,12 @@
 | `UpdateBisectTask` |  |
 | `ViewBisectTask` |  |
 
+## BuildHealth
+
+| Name | Description |
+| ---- | ----------- |
+| `AddUpdateDeleteBuildHealthFilter` |  |
+
 ## Compute
 
 | Name | Description |
@@ -503,6 +521,12 @@
 | `UpdateBisectTask` |  |
 | `ViewBisectTask` |  |
 
+## BuildHealth
+
+| Name | Description |
+| ---- | ----------- |
+| `AddUpdateDeleteBuildHealthFilter` |  |
+
 ## Compute
 
 | Name | Description |
@@ -588,12 +612,6 @@
 | `UpdateReplicator` |  |
 | `ViewReplicator` |  |
 
-## Secrets
-
-| Name | Description |
-| ---- | ----------- |
-| `ViewSecret` | View a credential |
-
 ## ServiceAccounts
 
 | Name | Description |
@@ -667,6 +685,12 @@
 | `UpdateBisectTask` |  |
 | `ViewBisectTask` |  |
 
+## BuildHealth
+
+| Name | Description |
+| ---- | ----------- |
+| `AddUpdateDeleteBuildHealthFilter` |  |
+
 ## Compute
 
 | Name | Description |
@@ -756,7 +780,7 @@
 
 | Name | Description |
 | ---- | ----------- |
-| `ViewSecret` |  |
+| `ViewSecret` | View a credential |
 
 ## ServiceAccounts
 
@@ -831,6 +855,12 @@
 | `UpdateBisectTask` |  |
 | `ViewBisectTask` |  |
 
+## BuildHealth
+
+| Name | Description |
+| ---- | ----------- |
+| `AddUpdateDeleteBuildHealthFilter` |  |
+
 ## Compute
 
 | Name | Description |
@@ -1008,6 +1038,12 @@
 | `UpdateBisectTask` |  |
 | `ViewBisectTask` |  |
 
+## BuildHealth
+
+| Name | Description |
+| ---- | ----------- |
+| `AddUpdateDeleteBuildHealthFilter` |  |
+
 ## Compute
 
 | Name | Description |
@@ -1191,6 +1227,12 @@
 | `UpdateBisectTask` |  |
 | `ViewBisectTask` |  |
 
+## BuildHealth
+
+| Name | Description |
+| ---- | ----------- |
+| `AddUpdateDeleteBuildHealthFilter` |  |
+
 ## Compute
 
 | Name | Description |
@@ -1381,6 +1423,12 @@
 | `UpdateBisectTask` |  |
 | `ViewBisectTask` |  |
 
+## BuildHealth
+
+| Name | Description |
+| ---- | ----------- |
+| `AddUpdateDeleteBuildHealthFilter` |  |
+
 ## Compute
 
 | Name | Description |
diff --git a/Engine/Source/Programs/Horde/Docs/Deployment/ServerSettings.md b/Engine/Source/Programs/Horde/Docs/Deployment/ServerSettings.md
--- a/Engine/Source/Programs/Horde/Docs/Deployment/ServerSettings.md
+++ b/Engine/Source/Programs/Horde/Docs/Deployment/ServerSettings.md
@@ -131,8 +131,8 @@ Name | Description
 `build` | [BuildServerConfig](#buildserverconfig)<br>Configuration for the build plugin
 `compute` | [ComputeServerConfig](#computeserverconfig)<br>Configuration for the compute plugin
 `exampleplugin` | [ExamplePluginServerConfig](#examplepluginserverconfig)<br>Configuration for the exampleplugin plugin
-`secrets` | [SecretsServerConfig](#secretsserverconfig)<br>Configuration for the secrets plugin
 `experimental` | [ExperimentalServerConfig](#experimentalserverconfig)<br>Configuration for the experimental plugin
+`secrets` | [SecretsServerConfig](#secretsserverconfig)<br>Configuration for the secrets plugin
 `storage` | [StorageServerConfig](#storageserverconfig)<br>Configuration for the storage plugin
 `symbols` | [PluginServerConfig](#pluginserverconfig)<br>Configuration for the symbols plugin
 `tools` | [ToolsServerConfig](#toolsserverconfig)<br>Configuration for the tools plugin
@@ -369,21 +369,21 @@ Name | Description
 ---- | -----------
 `enabled` | `boolean`<br>Whether the plugin should be enabled or not
 
-## SecretsServerConfig
+## ExperimentalServerConfig
 
-Static configuration for the secrets plugin
+Server configuration for the experimental plugin
 
 Name | Description
 ---- | -----------
-`withAws` | `boolean`<br>Whether to enable Amazon Web Services (AWS) specific features
 `enabled` | `boolean`<br>Whether the plugin should be enabled or not
 
-## ExperimentalServerConfig
+## SecretsServerConfig
 
-Server configuration for the experimental plugin
+Static configuration for the secrets plugin
 
 Name | Description
 ---- | -----------
+`withAws` | `boolean`<br>Whether to enable Amazon Web Services (AWS) specific features
 `enabled` | `boolean`<br>Whether the plugin should be enabled or not
 
 ## StorageServerConfig
