Horde: Use the correct Perforce metadata logger for PerforceMaterializer

Refactored IWorkspaceMaterializer to support loggers.

diff --git a/Engine/Source/Programs/Horde/Drivers/JobDriver.Tests/Execution/FakeWorkspaceMaterializer.cs b/Engine/Source/Programs/Horde/Drivers/JobDriver.Tests/Execution/FakeWorkspaceMaterializer.cs
--- a/Engine/Source/Programs/Horde/Drivers/JobDriver.Tests/Execution/FakeWorkspaceMaterializer.cs
+++ b/Engine/Source/Programs/Horde/Drivers/JobDriver.Tests/Execution/FakeWorkspaceMaterializer.cs
@@ -8,6 +8,7 @@ using System.Threading;
 using System.Threading.Tasks;
 using EpicGames.Core;
 using JobDriver.Execution;
+using Microsoft.Extensions.Logging;
 
 namespace JobDriver.Tests.Execution;
 
@@ -25,6 +26,7 @@ class FakeWorkspaceMaterializer : IWorkspaceMaterializer
 	public IReadOnlyDictionary<string, string> EnvironmentVariables { get; } = new Dictionary<string, string>();
 
 	public bool IsPerforceWorkspace => false;
+	public ILogger GetLogger(ILogger logger) => logger;
 
 	public FakeWorkspaceMaterializer()
 	{
diff --git a/Engine/Source/Programs/Horde/Drivers/JobDriver.Tests/PerforceMaterializerTests.cs b/Engine/Source/Programs/Horde/Drivers/JobDriver.Tests/PerforceMaterializerTests.cs
--- a/Engine/Source/Programs/Horde/Drivers/JobDriver.Tests/PerforceMaterializerTests.cs
+++ b/Engine/Source/Programs/Horde/Drivers/JobDriver.Tests/PerforceMaterializerTests.cs
@@ -4,8 +4,10 @@ using System;
 using System.Collections.Generic;
 using System.IO;
 using System.Linq;
+using System.Text.Json;
 using System.Threading;
 using System.Threading.Tasks;
+using EpicGames.Core;
 using EpicGames.Perforce;
 using EpicGames.Perforce.Fixture;
 using HordeCommon.Rpc.Messages;
@@ -18,6 +20,19 @@ using SyncOptions = JobDriver.Execution.SyncOptions;
 
 namespace JobDriver.Tests;
 
+public record TestLoggerEntry(LogLevel LogLevel, EventId EventId, object? State);
+public class TestLogger : ILogger
+{
+	public List<TestLoggerEntry> Entries { get; } = [];
+	public void Log<TState>(LogLevel logLevel, EventId eventId, TState state, Exception? exception, Func<TState, Exception?, string> formatter)
+	{
+		Entries.Add(new TestLoggerEntry(logLevel, eventId, state));
+	}
+
+	public bool IsEnabled(LogLevel logLevel) => true;
+	public IDisposable? BeginScope<TState>(TState state) where TState : notnull => null;
+}
+
 [TestClass]
 public class PerforceMaterializerTests : BasePerforceFixtureTest
 {
@@ -54,6 +69,35 @@ public class PerforceMaterializerTests : BasePerforceFixtureTest
 		Assert.IsTrue(true, "Submit went through without exception being thrown");
 	}
 
+	[TestMethod]
+	public async Task Logger_RelativeDepotPath_Async()
+	{
+		using PerforceMaterializer pm = CreatePm(PerforceSnapshot.FooMain);
+
+		TestLogger baseLogger = new();
+		ILogger materializerLogger = pm.GetLogger(baseLogger);
+		Assert.AreNotEqual(baseLogger, materializerLogger);
+		await pm.SyncAsync(IWorkspaceMaterializer.LatestChangeNumber, -1, new SyncOptions(), CancellationToken.None);
+
+		FileReference source = FileReference.Combine(pm.SyncDir, "Sub1", "Sub2", "MyFile.txt");
+		LogValue fileValue = LogValue.SourceFile(source, source.GetFileName());
+		materializerLogger.LogError(KnownLogEvents.Compiler, "Testing path resolution {File}", fileValue);
+		Assert.AreEqual(1, baseLogger.Entries.Count);
+
+		if (baseLogger.Entries[0].State is JsonLogEvent jle)
+		{
+			JsonDocument document = JsonDocument.Parse(jle.Data);
+			JsonElement fileElement = document.RootElement.GetProperty("properties").GetProperty("File");
+			Assert.AreEqual("MyFile.txt", fileElement.GetProperty("$text").GetString());
+			Assert.AreEqual("Sub1/Sub2/MyFile.txt", fileElement.GetProperty("relativePath").GetString());
+			Assert.AreEqual("//Foo/Main/Sub1/Sub2/MyFile.txt@9", fileElement.GetProperty("depotPath").GetString());
+		}
+		else
+		{
+			Assert.Fail("Unable to find JsonLogEvent");
+		}
+	}
+
 	[TestMethod]
 	[DataRow(PerforceSnapshot.FooMain, IWorkspaceMaterializer.LatestChangeNumber, 9)]
 	[DataRow(PerforceSnapshot.FooMain, IWorkspaceMaterializer.LatestChangeNumber, 9, true)]
diff --git a/Engine/Source/Programs/Horde/Drivers/JobDriver.Tests/WorkspaceExecutorTests.cs b/Engine/Source/Programs/Horde/Drivers/JobDriver.Tests/WorkspaceExecutorTests.cs
--- a/Engine/Source/Programs/Horde/Drivers/JobDriver.Tests/WorkspaceExecutorTests.cs
+++ b/Engine/Source/Programs/Horde/Drivers/JobDriver.Tests/WorkspaceExecutorTests.cs
@@ -24,6 +24,8 @@ internal class MaterializerStub(string name) : IWorkspaceMaterializer
 	public string Identifier { get; } = "not-set";
 	public IReadOnlyDictionary<string, string> EnvironmentVariables { get; } = new Dictionary<string, string>();
 	public bool IsPerforceWorkspace { get; } = false;
+	public ILogger GetLogger(ILogger logger) => logger;
+
 	public Task SyncAsync(int changeNum, int shelveChangeNum, SyncOptions options, CancellationToken cancellationToken)
 	{
 		throw new NotImplementedException();
diff --git a/Engine/Source/Programs/Horde/Drivers/JobDriver/Execution/ManagedWorkspaceMaterializer.cs b/Engine/Source/Programs/Horde/Drivers/JobDriver/Execution/ManagedWorkspaceMaterializer.cs
--- a/Engine/Source/Programs/Horde/Drivers/JobDriver/Execution/ManagedWorkspaceMaterializer.cs
+++ b/Engine/Source/Programs/Horde/Drivers/JobDriver/Execution/ManagedWorkspaceMaterializer.cs
@@ -27,6 +27,7 @@ public sealed class ManagedWorkspaceMaterializer : IWorkspaceMaterializer
 	private readonly WorkspaceInfo _workspace;
 	private readonly Tracer _tracer;
 	private readonly ILogger _logger;
+	private PerforceMetadataLogger? _perforceLogger;
 
 	/// <inheritdoc/>
 	public DirectoryReference SyncDir => _workspace.WorkspaceDir;
@@ -95,6 +96,13 @@ public sealed class ManagedWorkspaceMaterializer : IWorkspaceMaterializer
 		return new ManagedWorkspaceMaterializer(agentWorkspace, useCacheFile, cleanDuringFinalize, workspace, tracer, logger);
 	}
 
+	/// <inheritdoc/>
+	public ILogger GetLogger(ILogger logger)
+	{
+		_perforceLogger = new PerforceMetadataLogger(logger);
+		return _perforceLogger;
+	}
+
 	/// <inheritdoc/>
 	public async Task SyncAsync(int changeNum, int preflightChangeNum, SyncOptions options, CancellationToken cancellationToken)
 	{
@@ -111,6 +119,7 @@ public sealed class ManagedWorkspaceMaterializer : IWorkspaceMaterializer
 			span.SetAttribute("horde.change_num_latest", latestChangeNum);
 			changeNum = latestChangeNum;
 		}
+		_perforceLogger?.AddClientView(SyncDir, _workspace.StreamView, changeNum);
 
 		FileReference cacheFile = FileReference.Combine(_workspace.MetadataDir, "Contents.dat");
 		if (_useCacheFile)
diff --git a/Engine/Source/Programs/Horde/Drivers/JobDriver/Execution/PerforceMaterializer.cs b/Engine/Source/Programs/Horde/Drivers/JobDriver/Execution/PerforceMaterializer.cs
--- a/Engine/Source/Programs/Horde/Drivers/JobDriver/Execution/PerforceMaterializer.cs
+++ b/Engine/Source/Programs/Horde/Drivers/JobDriver/Execution/PerforceMaterializer.cs
@@ -135,11 +135,13 @@ public sealed class PerforceMaterializer : IWorkspaceMaterializer
 	private PerforceSettings? _perforceSettings;
 	private IPerforceConnection? _perforceWithoutClient;
 	private IPerforceConnection? _perforceWithClient;
+	private List<string>? _streamView;
 	private const string StateFilename = "State.json";
 	private const bool PreferNativePerforceClient = true;
 	private const int MaxFilesLogged = 100;
 	private readonly string _stateFile;
 	private readonly string _stateTempFile;
+	private PerforceMetadataLogger? _perforceLogger;
 
 	/// <summary>
 	/// Constructor
@@ -215,11 +217,21 @@ public sealed class PerforceMaterializer : IWorkspaceMaterializer
 			}
 		}
 		
+		StreamRecord stream = await _perforceWithoutClient.GetStreamAsync(_options.AgentWorkspace.Stream, true, cancellationToken);
+		_streamView = stream.View;
+		
 		string clientName = $"Horde+PM+{WorkspaceInfo.GetNormalizedHostName(hostName)}+{raw.Identifier}{edgeSuffix}";
 		_perforceSettings = new PerforceSettings(_perforceWithoutClient.Settings) { ClientName = clientName, PreferNativeClient = PreferNative };
 		status.Progress = $"({timer.Elapsed.TotalSeconds:0.0}s)";
 	}
-	
+
+	/// <inheritdoc/>
+	public ILogger GetLogger(ILogger logger)
+	{
+		_perforceLogger = new PerforceMetadataLogger(logger);
+		return _perforceLogger;
+	}
+
 	/// <inheritdoc/>
 	public async Task SyncAsync(int changeNum, int shelveChangeNum, SyncOptions options, CancellationToken cancellationToken)
 	{
@@ -239,8 +251,10 @@ public sealed class PerforceMaterializer : IWorkspaceMaterializer
 			|| state.Identifier != Identifier
 			|| state.Stream != _options.AgentWorkspace.Stream;
 		
-		_logger.LogInformation("Status: {Status}\nClient: {PrevClient} -> {Client}\nIdentifier: {PrevId} -> {Id}\nStream: {PrevStream} -> {Stream}\nCL: {PrevCl} -> {Cl}\nShelve: {PrevShelve} -> {Shelve}\nDirty: {Dirty}",
+		_logger.LogInformation("Status: {Status}\nVersion: {PrevVersion} -> {Version}\nClient: {PrevClient} -> {Client}\nIdentifier: {PrevId} -> {Id}\nStream: {PrevStream} -> {Stream}\nCL: {PrevCl} -> {Cl}\nShelve: {PrevShelve} -> {Shelve}\nDirty: {Dirty}",
 			state?.Status,
+			state?.Version,
+			CurrentStateVersion,
 			state?.Client, _perforceSettings.ClientName,
 			state?.Identifier, Identifier,
 			state?.Stream, _options.AgentWorkspace.Stream,
@@ -297,6 +311,11 @@ public sealed class PerforceMaterializer : IWorkspaceMaterializer
 			changeNum = changesRecords[0].Number;
 		}
 
+		if (_streamView != null)
+		{
+			_perforceLogger?.AddClientView(SyncDir, PerforceViewMap.Parse(_streamView), changeNum);
+		}
+
 		await SaveStateAsync(TransactionStatus.Dirty, client.Name, changeNum, shelveChangeNum, [], cancellationToken);
 		
 		{
diff --git a/Engine/Source/Programs/Horde/Drivers/JobDriver/Execution/WorkspaceExecutor.cs b/Engine/Source/Programs/Horde/Drivers/JobDriver/Execution/WorkspaceExecutor.cs
--- a/Engine/Source/Programs/Horde/Drivers/JobDriver/Execution/WorkspaceExecutor.cs
+++ b/Engine/Source/Programs/Horde/Drivers/JobDriver/Execution/WorkspaceExecutor.cs
@@ -127,21 +127,7 @@ namespace JobDriver.Execution
 
 		private ILogger GetLogger(ILogger logger)
 		{
-			if (_workspace.IsPerforceWorkspace)
-			{
-				// Try resolve a PerforceLogger using assumptions about the materializer.
-				// This is to remain compatible with PerforceExecutor.
-				// These Perforce-specific references should ideally not exist in WorkspaceExecutor.
-				WorkspaceInfo? workspaceInfo = (_workspace as ManagedWorkspaceMaterializer)?.GetWorkspaceInfo();
-				WorkspaceInfo? autoSdkWorkspaceInfo = (_autoSdkWorkspace as ManagedWorkspaceMaterializer)?.GetWorkspaceInfo();
-
-				if (workspaceInfo != null)
-				{
-					return PerforceExecutor.CreatePerforceLogger(logger, Batch.Change, workspaceInfo, autoSdkWorkspaceInfo);
-				}
-			}
-
-			return logger;
+			return _workspace.GetLogger(logger);
 		}
 	}
 
diff --git a/Engine/Source/Programs/Horde/Drivers/JobDriver/Execution/WorkspaceMaterializer.cs b/Engine/Source/Programs/Horde/Drivers/JobDriver/Execution/WorkspaceMaterializer.cs
--- a/Engine/Source/Programs/Horde/Drivers/JobDriver/Execution/WorkspaceMaterializer.cs
+++ b/Engine/Source/Programs/Horde/Drivers/JobDriver/Execution/WorkspaceMaterializer.cs
@@ -2,6 +2,7 @@
 
 using EpicGames.Core;
 using HordeCommon.Rpc.Messages;
+using Microsoft.Extensions.Logging;
 
 namespace JobDriver.Execution;
 
@@ -86,6 +87,13 @@ public interface IWorkspaceMaterializer : IDisposable
 	/// </summary>
 	bool IsPerforceWorkspace { get; }
 
+	/// <summary>
+	/// Get a logger capable of decorating log lines with materializer specific hints (resolving of paths or change numbers)
+	/// </summary>
+	/// <param name="logger">Base logger to wrap</param>
+	/// <returns>A wrapped logger</returns>
+	ILogger GetLogger(ILogger logger);
+
 	/// <summary>
 	/// Materialize (or sync) a Perforce stream at a given change number
 	/// Once method has completed, file tree is available on disk.
