Horde: Global config flag to enable and disable RCA notifications to be sent

#Synced-CL 0

diff --git a/Engine/Source/Programs/Horde/Plugins/Build/HordeServer.Build/BuildConfig.cs b/Engine/Source/Programs/Horde/Plugins/Build/HordeServer.Build/BuildConfig.cs
--- a/Engine/Source/Programs/Horde/Plugins/Build/HordeServer.Build/BuildConfig.cs
+++ b/Engine/Source/Programs/Horde/Plugins/Build/HordeServer.Build/BuildConfig.cs
@@ -100,6 +100,11 @@ namespace HordeServer
 		/// </summary>
 		public IssueConfig? Issues { get; set; }
 
+		/// <summary>
+		/// Whether to allow notifications for Root Cause Analysis to be sent
+		/// </summary>
+		public bool EnableIssueRootCauseAnalysisNotifications { get; set; }
+
 		/// <summary>
 		/// Issue channel groupings
 		/// </summary>
diff --git a/Engine/Source/Programs/Horde/Plugins/Build/HordeServer.Build/Notifications/Sinks/SlackNotificationSink.cs b/Engine/Source/Programs/Horde/Plugins/Build/HordeServer.Build/Notifications/Sinks/SlackNotificationSink.cs
--- a/Engine/Source/Programs/Horde/Plugins/Build/HordeServer.Build/Notifications/Sinks/SlackNotificationSink.cs
+++ b/Engine/Source/Programs/Horde/Plugins/Build/HordeServer.Build/Notifications/Sinks/SlackNotificationSink.cs
@@ -3153,6 +3153,12 @@ namespace HordeServer.Notifications.Sinks
 				return null;
 			}
 
+			if (!_buildConfig.CurrentValue.EnableIssueRootCauseAnalysisNotifications)
+			{
+				// When notifications are disabled still allow notifications to resume for issues if enabled later
+				return GetNextEscalationTime(issue, issue.ResolvedAt.Value, utcNow, workflow.EscalateRootCauseAnalysisTimes, notificationName);
+			}
+
 			if (DateTime.Now.DayOfWeek == DayOfWeek.Saturday || DateTime.Now.DayOfWeek == DayOfWeek.Sunday)
 			{
 				_logger.LogInformation("Skipping escalation message for {NotificationName} {IssueId} since it's a weekend", notificationName, issueId);
