[Horde][BuildHealth][Dashboard] Updating naming of request to be more appropriate & minor polish.

#Synced-CL 48916881

diff --git a/Engine/Source/Programs/Horde/HordeDashboard/src/backend/Api.ts b/Engine/Source/Programs/Horde/HordeDashboard/src/backend/Api.ts
--- a/Engine/Source/Programs/Horde/HordeDashboard/src/backend/Api.ts
+++ b/Engine/Source/Programs/Horde/HordeDashboard/src/backend/Api.ts
@@ -5415,7 +5415,7 @@ export type GetServiceAccountResponse = {
  */
 export type AddBuildHealthFiltersRequest = {
 	/// The project id.
-	projectId: string;
+	filterProject: string;
 
 	/// The filter name.
 	filterName: string;
diff --git a/Engine/Source/Programs/Horde/HordeDashboard/src/components/buildhealth/BuildHealthFilterModal.tsx b/Engine/Source/Programs/Horde/HordeDashboard/src/components/buildhealth/BuildHealthFilterModal.tsx
--- a/Engine/Source/Programs/Horde/HordeDashboard/src/components/buildhealth/BuildHealthFilterModal.tsx
+++ b/Engine/Source/Programs/Horde/HordeDashboard/src/components/buildhealth/BuildHealthFilterModal.tsx
@@ -86,7 +86,7 @@ export const BuildHealthFilterModal: React.FC<BuildHealthFilterModalProps> = ({
 
     async function Save(): Promise<GetBuildHealthFilterResponse | null> {
         if (buildHealthFilter?.filterId === null) {
-            let request: AddBuildHealthFiltersRequest = { filterName: state.filterName, projectId: buildHealthFilter.filterProjectId, filterQuery: buildHealthFilter.oldFilterQuery };
+            let request: AddBuildHealthFiltersRequest = { filterName: state.filterName, filterProject: buildHealthFilter.filterProjectId, filterQuery: buildHealthFilter.oldFilterQuery };
             let response: GetBuildHealthFilterResponse = await backend.addBuildHealthProjectFilter(request);
 
             return response;
@@ -102,7 +102,7 @@ export const BuildHealthFilterModal: React.FC<BuildHealthFilterModalProps> = ({
             return null;
         }
 
-        let request: AddBuildHealthFiltersRequest = { filterName: state.filterName, projectId: buildHealthFilter.filterProjectId, filterQuery: buildHealthFilter.newFilterQuery };
+        let request: AddBuildHealthFiltersRequest = { filterName: state.filterName, filterProject: buildHealthFilter.filterProjectId, filterQuery: buildHealthFilter.newFilterQuery };
         let response: GetBuildHealthFilterResponse = await backend.addBuildHealthProjectFilter(request);
 
         return response;
diff --git a/Engine/Source/Programs/Horde/HordeDashboard/src/components/buildhealth/BuildHealthOptions.ts b/Engine/Source/Programs/Horde/HordeDashboard/src/components/buildhealth/BuildHealthOptions.ts
--- a/Engine/Source/Programs/Horde/HordeDashboard/src/components/buildhealth/BuildHealthOptions.ts
+++ b/Engine/Source/Programs/Horde/HordeDashboard/src/components/buildhealth/BuildHealthOptions.ts
@@ -730,7 +730,7 @@ export class BuildHealthOptionsState {
 
     // #region -- User Selected Options --
     
-    @observable buildHealthFilterId : string | null; 
+    @observable buildHealthFilterId : string | null = null; 
     @observable enabledProjects: Record<string, string> = {};
     @observable enabledProjectFilter: string;
     @observable enabledStreams: Record<string, string> = {};
diff --git a/Engine/Source/Programs/Horde/HordeDashboard/src/components/buildhealth/BuildHealthOptionsController.ts b/Engine/Source/Programs/Horde/HordeDashboard/src/components/buildhealth/BuildHealthOptionsController.ts
--- a/Engine/Source/Programs/Horde/HordeDashboard/src/components/buildhealth/BuildHealthOptionsController.ts
+++ b/Engine/Source/Programs/Horde/HordeDashboard/src/components/buildhealth/BuildHealthOptionsController.ts
@@ -322,6 +322,7 @@ export class BuildHealthOptionsWriter implements IBuildHealthOptionWriter {
         this.clearStreams(true);
         this.clearStates();
         this.clearJobStartMethods();
+        this.state.buildHealthFilterId = null;
     }
 
     /**
