Horde: Add step report color hinting

diff --git a/Engine/Source/Programs/Horde/Plugins/Build/HordeServer.Build/Jobs/JobRpcService.cs b/Engine/Source/Programs/Horde/Plugins/Build/HordeServer.Build/Jobs/JobRpcService.cs
--- a/Engine/Source/Programs/Horde/Plugins/Build/HordeServer.Build/Jobs/JobRpcService.cs
+++ b/Engine/Source/Programs/Horde/Plugins/Build/HordeServer.Build/Jobs/JobRpcService.cs
@@ -983,6 +983,11 @@ namespace HordeServer.Jobs
 						reportMetaData.Add($"ReportStepSeverity={newReport.Severity}");
 					}
 
+					if (!String.IsNullOrEmpty(newReport.Color))
+					{
+						reportMetaData.Add($"ReportStepColor={newReport.Color}");
+					}
+
 					Dictionary<JobStepId, List<string>> stepMetaData = new Dictionary<JobStepId, List<string>>() { { stepId, reportMetaData } };					
 					
 					await _jobService.UpdateMetadataAsync(job.Id, null, stepMetaData, context.CancellationToken);
diff --git a/Engine/Source/Programs/Shared/EpicGames.Horde/Jobs/IJob.cs b/Engine/Source/Programs/Shared/EpicGames.Horde/Jobs/IJob.cs
--- a/Engine/Source/Programs/Shared/EpicGames.Horde/Jobs/IJob.cs
+++ b/Engine/Source/Programs/Shared/EpicGames.Horde/Jobs/IJob.cs
@@ -1918,6 +1918,11 @@ namespace EpicGames.Horde.Jobs
 		/// The message severity 
 		/// </summary>
 		string? Severity { get; }
+
+		/// <summary>
+		/// The message color
+		/// </summary>
+		string? Color { get; }
 	}
 
 	/// <summary>
@@ -1939,8 +1944,11 @@ namespace EpicGames.Horde.Jobs
 
 		/// <inheritdoc/>
 		public string? Message { get; set; }
-
+	
 		/// <inheritdoc/>
 		public string? Severity { get; set; }
+
+		/// <inheritdoc/>
+		public string? Color { get; set; }
 	}
 }
