Horde: Updating schema documentation

#rnx

diff --git a/Engine/Source/Programs/Horde/Docs/Config/Schema/Globals.md b/Engine/Source/Programs/Horde/Docs/Config/Schema/Globals.md
--- a/Engine/Source/Programs/Horde/Docs/Config/Schema/Globals.md
+++ b/Engine/Source/Programs/Horde/Docs/Config/Schema/Globals.md
@@ -62,8 +62,8 @@ Name | Description
 `analytics` | [AnalyticsConfig](#analyticsconfig)<br>Configuration for the analytics plugin
 `build` | [BuildConfig](#buildconfig)<br>Configuration for the build plugin
 `compute` | [ComputeConfig](#computeconfig)<br>Configuration for the compute plugin
-`exampleplugin` | [ExamplePluginConfig](#examplepluginconfig)<br>Configuration for the exampleplugin plugin
 `experimental` | [ExperimentalConfig](#experimentalconfig)<br>Configuration for the experimental plugin
+`exampleplugin` | [ExamplePluginConfig](#examplepluginconfig)<br>Configuration for the exampleplugin plugin
 `secrets` | [SecretsConfig](#secretsconfig)<br>Configuration for the secrets plugin
 `storage` | [StorageConfig](#storageconfig)<br>Configuration for the storage plugin
 `symbols` | [SymbolsConfig](#symbolsconfig)<br>Configuration for the symbols plugin
@@ -93,6 +93,7 @@ Name | Description
 `enableJobTaskConformCheck` | `boolean`<br>Whether to check if a conform is required on an agent before assigning a lease for a job task source.
 `issueFixedTag` | `string`<br>Commit tag to use for marking issues as fixed
 `issues` | [IssueConfig](Issues.md)<br>Build health issue configurations
+`enableIssueRootCauseAnalysisNotifications` | `boolean`<br>Whether to allow notifications for Root Cause Analysis to be sent
 `issueFindSuspectsMaxSuccessDays` | `integer`<br>The maximum days since node success to scan for suspects
 `issueFindSuspectsMaxScoringSeconds` | `integer`<br>The maximum number of seconds to spend scording suspects
 
@@ -434,21 +435,21 @@ Name | Description
 `description` | `string`<br>Human-readable description
 `computeId` | `string`<br>Compute ID for this network (used when allocating compute resources)
 
-## ExamplePluginConfig
+## ExperimentalConfig
 
-Configuration for the ExamplePlugin plugin
+Configuration for the experimental plugin
 
 Name | Description
 ---- | -----------
-`greetingMessage` | `string`<br>The Greeting provided by the example plugin
+`notifications` | [NotificationConfig](Notification.md)`[]`<br>Notifications to be sent on the Horde server
 
-## ExperimentalConfig
+## ExamplePluginConfig
 
-Configuration for the experimental plugin
+Configuration for the ExamplePlugin plugin
 
 Name | Description
 ---- | -----------
-`notifications` | [NotificationConfig](Notification.md)`[]`<br>Notifications to be sent on the Horde server
+`greetingMessage` | `string`<br>The Greeting provided by the example plugin
 
 ## SecretsConfig
 
diff --git a/Engine/Source/Programs/Horde/Docs/Deployment/ServerSettings.md b/Engine/Source/Programs/Horde/Docs/Deployment/ServerSettings.md
--- a/Engine/Source/Programs/Horde/Docs/Deployment/ServerSettings.md
+++ b/Engine/Source/Programs/Horde/Docs/Deployment/ServerSettings.md
@@ -130,8 +130,8 @@ Name | Description
 `analytics` | [AnalyticsServerConfig](#analyticsserverconfig)<br>Configuration for the analytics plugin
 `build` | [BuildServerConfig](#buildserverconfig)<br>Configuration for the build plugin
 `compute` | [ComputeServerConfig](#computeserverconfig)<br>Configuration for the compute plugin
-`exampleplugin` | [ExamplePluginServerConfig](#examplepluginserverconfig)<br>Configuration for the exampleplugin plugin
 `experimental` | [ExperimentalServerConfig](#experimentalserverconfig)<br>Configuration for the experimental plugin
+`exampleplugin` | [ExamplePluginServerConfig](#examplepluginserverconfig)<br>Configuration for the exampleplugin plugin
 `secrets` | [SecretsServerConfig](#secretsserverconfig)<br>Configuration for the secrets plugin
 `storage` | [StorageServerConfig](#storageserverconfig)<br>Configuration for the storage plugin
 `symbols` | [PluginServerConfig](#pluginserverconfig)<br>Configuration for the symbols plugin
@@ -361,17 +361,17 @@ Name | Description
 `ComputeQueueAwsMetric` | A no-op strategy that reports metrics to let an external AWS auto-scaling policy scale the fleet
 `LeaseUtilizationAwsMetric` | A no-op strategy that reports metrics to let an external AWS auto-scaling policy scale the fleet
 
-## ExamplePluginServerConfig
+## ExperimentalServerConfig
 
-Server configuration for the ExamplePlugin
+Server configuration for the experimental plugin
 
 Name | Description
 ---- | -----------
 `enabled` | `boolean`<br>Whether the plugin should be enabled or not
 
-## ExperimentalServerConfig
+## ExamplePluginServerConfig
 
-Server configuration for the experimental plugin
+Server configuration for the ExamplePlugin
 
 Name | Description
 ---- | -----------
