Horde: Save lease metrics reported by agent to database

diff --git a/Engine/Source/Programs/Horde/Plugins/Build/HordeServer.Build.Tests/Agents/Leases/LeasesControllerDbTest.cs b/Engine/Source/Programs/Horde/Plugins/Build/HordeServer.Build.Tests/Agents/Leases/LeasesControllerDbTest.cs
--- a/Engine/Source/Programs/Horde/Plugins/Build/HordeServer.Build.Tests/Agents/Leases/LeasesControllerDbTest.cs
+++ b/Engine/Source/Programs/Horde/Plugins/Build/HordeServer.Build.Tests/Agents/Leases/LeasesControllerDbTest.cs
@@ -38,7 +38,7 @@ namespace HordeServer.Tests.Agents.Leases
 			LeaseId id = new LeaseId(BinaryIdUtils.CreateNew());
 			SessionId sessionId = SessionIdUtils.GenerateNewId();
 			ILease lease = await LeaseCollection.AddAsync(id, null, "myLease", new AgentId("agent-1"), sessionId, null, null, null, startTime, Array.Empty<byte>());
-			await LeaseCollection.TrySetOutcomeAsync(id, startTime + duration, LeaseOutcome.Success, LeaseOutcomeReason.None, null);
+			await LeaseCollection.TrySetOutcomeAsync(id, startTime + duration, LeaseOutcome.Success, LeaseOutcomeReason.None, null, null);
 			return lease;
 		}
 	}
diff --git a/Engine/Source/Programs/Horde/Plugins/Build/HordeServer.Build.Tests/Fleet/LeaseUtilizationAwsMetricStrategyTests.cs b/Engine/Source/Programs/Horde/Plugins/Build/HordeServer.Build.Tests/Fleet/LeaseUtilizationAwsMetricStrategyTests.cs
--- a/Engine/Source/Programs/Horde/Plugins/Build/HordeServer.Build.Tests/Fleet/LeaseUtilizationAwsMetricStrategyTests.cs
+++ b/Engine/Source/Programs/Horde/Plugins/Build/HordeServer.Build.Tests/Fleet/LeaseUtilizationAwsMetricStrategyTests.cs
@@ -224,7 +224,7 @@ namespace HordeServer.Tests.Fleet
 			ILease lease = await LeaseCollection.AddAsync(new LeaseId(BinaryIdUtils.CreateNew()), null, "placeholderLease", agent.Id, agent.SessionId!.Value, new StreamId("placeholderStream"), poolId, null, startTime, payload);
 			if (duration != null)
 			{
-				bool wasModified = await LeaseCollection.TrySetOutcomeAsync(lease.Id, startTime + duration.Value, LeaseOutcome.Success, LeaseOutcomeReason.None, null);
+				bool wasModified = await LeaseCollection.TrySetOutcomeAsync(lease.Id, startTime + duration.Value, LeaseOutcome.Success, LeaseOutcomeReason.None, null, null);
 				Assert.IsTrue(wasModified);
 			}
 
diff --git a/Engine/Source/Programs/Horde/Plugins/Build/HordeServer.Build.Tests/Fleet/LeaseUtilizationStrategyTest.cs b/Engine/Source/Programs/Horde/Plugins/Build/HordeServer.Build.Tests/Fleet/LeaseUtilizationStrategyTest.cs
--- a/Engine/Source/Programs/Horde/Plugins/Build/HordeServer.Build.Tests/Fleet/LeaseUtilizationStrategyTest.cs
+++ b/Engine/Source/Programs/Horde/Plugins/Build/HordeServer.Build.Tests/Fleet/LeaseUtilizationStrategyTest.cs
@@ -120,7 +120,7 @@ namespace HordeServer.Tests.Fleet
 			}
 
 			ILease lease = await LeaseCollection.AddAsync(new LeaseId(BinaryIdUtils.CreateNew()), null, "placeholderLease", agent.Id, agent.SessionId!.Value, new StreamId("placeholderStream"), pool.Id, null, startTime, payload);
-			bool wasModified = await LeaseCollection.TrySetOutcomeAsync(lease.Id, startTime + duration, LeaseOutcome.Success, LeaseOutcomeReason.None,null);
+			bool wasModified = await LeaseCollection.TrySetOutcomeAsync(lease.Id, startTime + duration, LeaseOutcome.Success, LeaseOutcomeReason.None, null, null);
 			Assert.IsTrue(wasModified);
 			return lease;
 		}
diff --git a/Engine/Source/Programs/Horde/Plugins/Compute/HordeServer.Compute.Tests/Agents/Leases/LeaseCollectionTests.cs b/Engine/Source/Programs/Horde/Plugins/Compute/HordeServer.Compute.Tests/Agents/Leases/LeaseCollectionTests.cs
new file mode 100644
--- /dev/null
+++ b/Engine/Source/Programs/Horde/Plugins/Compute/HordeServer.Compute.Tests/Agents/Leases/LeaseCollectionTests.cs
@@ -0,0 +1,55 @@
+// Copyright Epic Games, Inc. All Rights Reserved.
+
+extern alias HordeAgent;
+using EpicGames.Horde.Agents;
+using EpicGames.Horde.Agents.Leases;
+using EpicGames.Horde.Agents.Sessions;
+using HordeServer.Agents.Leases;
+
+namespace HordeServer.Tests.Agents.Leases;
+
+[TestClass]
+public class LeaseCollectionTests : ComputeTestSetup
+{
+	[TestMethod]
+	public async Task SetOutcome_WithMetrics_Async()
+	{
+		ILease? lease = await CreateLeaseAsync();
+
+		lease = await LeaseCollection.GetAsync(lease.Id);
+		Assert.IsNull(lease!.Metrics);
+		Assert.AreEqual(LeaseOutcome.Unspecified, lease.Outcome);
+
+		LeaseCpuTime cpuTime = new (TimeSpan.FromSeconds(3), TimeSpan.FromSeconds(2), TimeSpan.FromSeconds(1));
+		LeaseMetrics metrics = new (TimeSpan.FromSeconds(10), TimeSpan.FromSeconds(20), 5, 0.59, cpuTime);
+		Assert.IsTrue(await LeaseCollection.TrySetOutcomeAsync(lease.Id, DateTime.UtcNow, LeaseOutcome.Success, LeaseOutcomeReason.None, null, metrics));
+		
+		lease = await LeaseCollection.GetAsync(lease.Id);
+		Assert.IsNotNull(lease!.Metrics);
+		Assert.AreEqual(LeaseOutcome.Success, lease.Outcome);
+		Assert.AreEqual(TimeSpan.FromSeconds(10), lease.Metrics.SetupTime);
+		Assert.AreEqual(TimeSpan.FromSeconds(20), lease.Metrics.TeardownTime);
+		Assert.AreEqual(5, lease.Metrics.CpuCount);
+		Assert.AreEqual(0.59, lease.Metrics.GlobalCpuUtilization, delta: 0.001);
+		Assert.AreEqual(TimeSpan.FromSeconds(3), lease.Metrics.GlobalCpuTime.User);
+		Assert.AreEqual(TimeSpan.FromSeconds(2), lease.Metrics.GlobalCpuTime.System);
+		Assert.AreEqual(TimeSpan.FromSeconds(1), lease.Metrics.GlobalCpuTime.Idle);
+	}
+	
+	[TestMethod]
+	public async Task SetOutcome_WithoutMetrics_Async()
+	{
+		ILease? lease = await CreateLeaseAsync();
+		
+		Assert.IsTrue(await LeaseCollection.TrySetOutcomeAsync(lease.Id, DateTime.UtcNow, LeaseOutcome.Failed, LeaseOutcomeReason.None, null, null));
+		lease = await LeaseCollection.GetAsync(lease.Id);
+		Assert.IsNull(lease!.Metrics);
+		Assert.AreEqual(LeaseOutcome.Failed, lease.Outcome);
+	}
+
+	private Task<ILease> CreateLeaseAsync()
+	{
+		return LeaseCollection.AddAsync(LeaseId.Parse("aaaaaaaaaaaaaaaaaaaaaaaa"), null, "lease1", new AgentId("agent1"),
+			SessionId.Parse("bbbbbbbbbbbbbbbbbbbbbbbb"), null, null, null, DateTime.UtcNow, Array.Empty<byte>());
+	}
+}
\ No newline at end of file
diff --git a/Engine/Source/Programs/Horde/Plugins/Compute/HordeServer.Compute.Tests/ComputeTestSetup.cs b/Engine/Source/Programs/Horde/Plugins/Compute/HordeServer.Compute.Tests/ComputeTestSetup.cs
--- a/Engine/Source/Programs/Horde/Plugins/Compute/HordeServer.Compute.Tests/ComputeTestSetup.cs
+++ b/Engine/Source/Programs/Horde/Plugins/Compute/HordeServer.Compute.Tests/ComputeTestSetup.cs
@@ -1,6 +1,7 @@
 // Copyright Epic Games, Inc. All Rights Reserved.
 
 using HordeServer.Agents;
+using HordeServer.Agents.Leases;
 using HordeServer.Compute;
 using HordeServer.Logs;
 using HordeServer.Storage;
@@ -11,6 +12,7 @@ namespace HordeServer.Tests
 	public class ComputeTestSetup : ServerTestSetup
 	{
 		public ComputeService ComputeService => ServiceProvider.GetRequiredService<ComputeService>();
+		public ILeaseCollection LeaseCollection => ServiceProvider.GetRequiredService<ILeaseCollection>();
 		public IAgentCollection AgentCollection => ServiceProvider.GetRequiredService<IAgentCollection>();
 		public IAgentScheduler AgentScheduler => ServiceProvider.GetRequiredService<IAgentScheduler>();
 		public AgentService AgentService => ServiceProvider.GetRequiredService<AgentService>();
diff --git a/Engine/Source/Programs/Horde/Plugins/Compute/HordeServer.Compute/Agents/AgentCollection.cs b/Engine/Source/Programs/Horde/Plugins/Compute/HordeServer.Compute/Agents/AgentCollection.cs
--- a/Engine/Source/Programs/Horde/Plugins/Compute/HordeServer.Compute/Agents/AgentCollection.cs
+++ b/Engine/Source/Programs/Horde/Plugins/Compute/HordeServer.Compute/Agents/AgentCollection.cs
@@ -2,6 +2,7 @@
 
 using System.Collections.Concurrent;
 using System.Collections.ObjectModel;
+using System.Diagnostics.CodeAnalysis;
 using System.Linq.Expressions;
 using System.Runtime.CompilerServices;
 using EpicGames.Core;
@@ -964,19 +965,42 @@ namespace HordeServer.Agents
 					if (rpcLease is { State: RpcLeaseState.Completed })
 					{
 						_logger.LogInformation("Setting outcome {Outcome} for completed lease {LeaseId}", oldLease.Id, rpcLease.Outcome);
-						await _leaseCollection.TrySetOutcomeAsync(oldLease.Id, _clock.UtcNow, (LeaseOutcome)rpcLease.Outcome, LeaseOutcomeReason.None, rpcLease.Output.ToArray(), CancellationToken.None);
+						await _leaseCollection.TrySetOutcomeAsync(oldLease.Id, _clock.UtcNow, (LeaseOutcome)rpcLease.Outcome, LeaseOutcomeReason.None, rpcLease.Output.ToArray(), ToNativeLeaseMetrics(rpcLease.Metrics), CancellationToken.None);
 					}
 					else
 					{
 						LeaseOutcomeReason reason = rpcLease?.ToNativeOutcomeReason() ?? LeaseOutcomeReason.None;
 						_logger.LogInformation("Setting outcome {Outcome} with reason {Reason} for lease {LeaseId} ({State})", rpcLease?.Outcome, reason, oldLease.Id, rpcLease?.State);
-						await _leaseCollection.TrySetOutcomeAsync(oldLease.Id, _clock.UtcNow, LeaseOutcome.Cancelled, reason, null, CancellationToken.None);
+						await _leaseCollection.TrySetOutcomeAsync(oldLease.Id, _clock.UtcNow, LeaseOutcome.Cancelled, reason, null, ToNativeLeaseMetrics(rpcLease?.Metrics), CancellationToken.None);
 					}
 				}
 			}
 
 			return await CreateAgentObjectAsync(newDocument, newSession, cancellationToken);
 		}
+		
+		[SuppressMessage("Maintainability", "CA1508:Avoid dead conditional code", Justification = "Protobuf generated code may contain null fields")]
+		static LeaseMetrics? ToNativeLeaseMetrics(RpcLeaseMetrics? metrics)
+		{
+			if (metrics == null)
+			{
+				return null;
+			}
+			
+			LeaseCpuTime cpuTime = metrics.GlobalCpuTime != null
+				? new LeaseCpuTime(
+					TimeSpan.FromMilliseconds(metrics.GlobalCpuTime.UserMs),
+					TimeSpan.FromMilliseconds(metrics.GlobalCpuTime.SystemMs),
+					TimeSpan.FromMilliseconds(metrics.GlobalCpuTime.IdleMs))
+				: new LeaseCpuTime(TimeSpan.Zero,  TimeSpan.Zero, TimeSpan.Zero);
+
+			return new LeaseMetrics(
+				TimeSpan.FromMilliseconds(metrics.SetupTimeMs),
+				TimeSpan.FromMilliseconds(metrics.TeardownTimeMs),
+				metrics.CpuCount,
+				metrics.GlobalCpuUtilization,
+				cpuTime);
+		}
 
 		static void SetPoolProperties(List<string> properties, List<PoolId> pools)
 		{
diff --git a/Engine/Source/Programs/Horde/Plugins/Compute/HordeServer.Compute/Agents/AgentService.cs b/Engine/Source/Programs/Horde/Plugins/Compute/HordeServer.Compute/Agents/AgentService.cs
--- a/Engine/Source/Programs/Horde/Plugins/Compute/HordeServer.Compute/Agents/AgentService.cs
+++ b/Engine/Source/Programs/Horde/Plugins/Compute/HordeServer.Compute/Agents/AgentService.cs
@@ -596,7 +596,7 @@ namespace HordeServer.Agents
 			if (agent == null)
 			{
 				_logger.LogWarning("Forcing lease {LeaseId} to cancelled state; agent {AgentId} not found", lease.Id, lease.AgentId);
-				await _leases.TrySetOutcomeAsync(lease.Id, _clock.UtcNow, LeaseOutcome.Cancelled, outcomeReason, null, cancellationToken);
+				await _leases.TrySetOutcomeAsync(lease.Id, _clock.UtcNow, LeaseOutcome.Cancelled, outcomeReason, null, null, cancellationToken);
 				return false;
 			}
 
@@ -605,7 +605,7 @@ namespace HordeServer.Agents
 			if (index == -1)
 			{
 				_logger.LogWarning("Forcing lease {LeaseId} to cancelled state; no longer running on agent {AgentId}", lease.Id, lease.AgentId);
-				await _leases.TrySetOutcomeAsync(lease.Id, _clock.UtcNow, LeaseOutcome.Cancelled, outcomeReason, null, cancellationToken);
+				await _leases.TrySetOutcomeAsync(lease.Id, _clock.UtcNow, LeaseOutcome.Cancelled, outcomeReason, null, null, cancellationToken);
 				return false;
 			}
 
diff --git a/Engine/Source/Programs/Horde/Plugins/Compute/HordeServer.Compute/Agents/Leases/ILease.cs b/Engine/Source/Programs/Horde/Plugins/Compute/HordeServer.Compute/Agents/Leases/ILease.cs
--- a/Engine/Source/Programs/Horde/Plugins/Compute/HordeServer.Compute/Agents/Leases/ILease.cs
+++ b/Engine/Source/Programs/Horde/Plugins/Compute/HordeServer.Compute/Agents/Leases/ILease.cs
@@ -85,8 +85,38 @@ namespace HordeServer.Agents.Leases
 		/// Output from executing the lease
 		/// </summary>
 		public ReadOnlyMemory<byte> Output { get; }
+
+		/// <summary>
+		/// Metrics collected during execution of lease
+		/// </summary>
+		public LeaseMetrics? Metrics { get; }
 	}
 
+	/// <summary>
+	/// CPU timings for a lease
+	/// </summary>
+	/// <param name="User">CPU time spent in user mode</param>
+	/// <param name="System">CPU time spent in system/kernel mode</param>
+	/// <param name="Idle">CPU time spent idle</param>
+	public record LeaseCpuTime(TimeSpan User, TimeSpan System, TimeSpan Idle);
+
+	/// <summary>
+	/// Metrics collected during execution of lease
+	/// Machine-wide metrics assumes the lease has exclusive access to the agent and other background work is negligible.
+	/// </summary>
+	/// <param name="SetupTime">Time taken to prepare for lease execution (e.g syncing VCS data)</param>
+	/// <param name="TeardownTime">Time taken to clean up after a lease execution (removing files created etc)</param>
+	/// <param name="CpuCount">Number of logical CPU cores that were available to the lease</param>
+	/// <param name="GlobalCpuUtilization">CPU utilization in percent based on number of CPU seconds that was available during the lease</param>
+	/// <param name="GlobalCpuTime">Machine-wide CPU time spent during this lease</param>
+	public record LeaseMetrics(
+		TimeSpan SetupTime,
+		TimeSpan TeardownTime,
+		int CpuCount,
+		double GlobalCpuUtilization,
+		LeaseCpuTime GlobalCpuTime
+	);
+
 	/// <summary>
 	/// Extension methods for leases
 	/// </summary>
diff --git a/Engine/Source/Programs/Horde/Plugins/Compute/HordeServer.Compute/Agents/Leases/ILeaseCollection.cs b/Engine/Source/Programs/Horde/Plugins/Compute/HordeServer.Compute/Agents/Leases/ILeaseCollection.cs
--- a/Engine/Source/Programs/Horde/Plugins/Compute/HordeServer.Compute/Agents/Leases/ILeaseCollection.cs
+++ b/Engine/Source/Programs/Horde/Plugins/Compute/HordeServer.Compute/Agents/Leases/ILeaseCollection.cs
@@ -103,8 +103,9 @@ namespace HordeServer.Agents.Leases
 		/// <param name="outcome">Outcome of the lease</param>
 		/// <param name="outcomeReason">Reason for the outcome</param>
 		/// <param name="output">Output data from the task</param>
+		/// <param name="metrics">Metrics</param>
 		/// <param name="cancellationToken">Cancellation token for the operation</param>
 		/// <returns>True if the lease was updated, false otherwise</returns>
-		Task<bool> TrySetOutcomeAsync(LeaseId leaseId, DateTime finishTime, LeaseOutcome outcome, LeaseOutcomeReason outcomeReason, byte[]? output, CancellationToken cancellationToken = default);
+		Task<bool> TrySetOutcomeAsync(LeaseId leaseId, DateTime finishTime, LeaseOutcome outcome, LeaseOutcomeReason outcomeReason, byte[]? output, LeaseMetrics? metrics, CancellationToken cancellationToken = default);
 	}
 }
diff --git a/Engine/Source/Programs/Horde/Plugins/Compute/HordeServer.Compute/Agents/Leases/LeaseCollection.cs b/Engine/Source/Programs/Horde/Plugins/Compute/HordeServer.Compute/Agents/Leases/LeaseCollection.cs
--- a/Engine/Source/Programs/Horde/Plugins/Compute/HordeServer.Compute/Agents/Leases/LeaseCollection.cs
+++ b/Engine/Source/Programs/Horde/Plugins/Compute/HordeServer.Compute/Agents/Leases/LeaseCollection.cs
@@ -25,6 +25,41 @@ namespace HordeServer.Agents.Leases
 	/// </summary>
 	public class LeaseCollection : ILeaseCollection
 	{
+		record LeaseMetricsDocument(
+			[property: BsonElement("setup")] int SetupTimeMs,
+			[property: BsonElement("teardown")] int TeardownTimeMs,
+			[property: BsonElement("cpus")] int CpuCount,
+			[property: BsonElement("cpuUtil")] double GlobalCpuUtilization,
+			[property: BsonElement("cpuUsr")] long GlobalCpuUserMs,
+			[property: BsonElement("cpuSys")] long GlobalCpuSystemMs,
+			[property: BsonElement("cpuIdl")] long GlobalCpuIdleMs
+		)
+		{
+			public LeaseMetrics ToLeaseMetrics()
+			{
+				return new LeaseMetrics(
+					TimeSpan.FromMilliseconds(SetupTimeMs),
+					TimeSpan.FromMilliseconds(TeardownTimeMs),
+					CpuCount,
+					GlobalCpuUtilization,
+					new LeaseCpuTime(TimeSpan.FromMilliseconds(GlobalCpuUserMs), TimeSpan.FromMilliseconds(GlobalCpuSystemMs), TimeSpan.FromMilliseconds(GlobalCpuIdleMs))
+				);
+			}
+
+			public static LeaseMetricsDocument FromLeaseMetrics(LeaseMetrics metrics)
+			{
+				return new LeaseMetricsDocument(
+					(int)metrics.SetupTime.TotalMilliseconds,
+					(int)metrics.TeardownTime.TotalMilliseconds,
+					metrics.CpuCount,
+					metrics.GlobalCpuUtilization,
+					(long)metrics.GlobalCpuTime.User.TotalMilliseconds,
+					(long)metrics.GlobalCpuTime.System.TotalMilliseconds,
+					(long)metrics.GlobalCpuTime.Idle.TotalMilliseconds
+				);
+			}
+		}
+		
 		/// <summary>
 		/// Concrete implementation of a lease document
 		/// </summary>
@@ -45,6 +80,7 @@ namespace HordeServer.Agents.Leases
 			public DateTime? FinishTime { get; set; }
 			public byte[] Payload { get; set; }
 			public LeaseOutcome Outcome { get; set; }
+			public LeaseMetricsDocument? Metrics { get; set; }
 			
 			[BsonRepresentation(BsonType.Int64)]
 			public LeaseOutcomeReason OutcomeReason { get; set; } = LeaseOutcomeReason.None;
@@ -54,6 +90,7 @@ namespace HordeServer.Agents.Leases
 
 			ReadOnlyMemory<byte> ILease.Payload => Payload;
 			ReadOnlyMemory<byte> ILease.Output => Output ?? Array.Empty<byte>();
+			LeaseMetrics? ILease.Metrics => Metrics?.ToLeaseMetrics();
 
 			[BsonConstructor]
 			private LeaseDocument()
@@ -197,7 +234,7 @@ namespace HordeServer.Agents.Leases
 		}
 
 		/// <inheritdoc/>
-		public async Task<bool> TrySetOutcomeAsync(LeaseId leaseId, DateTime finishTime, LeaseOutcome outcome, LeaseOutcomeReason outcomeReason, byte[]? output, CancellationToken cancellationToken)
+		public async Task<bool> TrySetOutcomeAsync(LeaseId leaseId, DateTime finishTime, LeaseOutcome outcome, LeaseOutcomeReason outcomeReason, byte[]? output, LeaseMetrics? metrics, CancellationToken cancellationToken)
 		{
 			FilterDefinitionBuilder<LeaseDocument> filterBuilder = Builders<LeaseDocument>.Filter;
 			FilterDefinition<LeaseDocument> filter = filterBuilder.Eq(x => x.Id, leaseId) & filterBuilder.Eq(x => x.FinishTime, null);
@@ -208,11 +245,16 @@ namespace HordeServer.Agents.Leases
 				.Set(x => x.Outcome, outcome)
 				.Set(x => x.OutcomeReason, outcomeReason);
 
-			if (output != null && output.Length > 0)
+			if (output is { Length: > 0 })
 			{
 				update = update.Set(x => x.Output, output);
 			}
 
+			if (metrics != null)
+			{
+				update = update.Set(x => x.Metrics, LeaseMetricsDocument.FromLeaseMetrics(metrics));
+			}
+
 			LeaseDocument? document = await _leases.FindOneAndUpdateAsync(filter, update, new FindOneAndUpdateOptions<LeaseDocument, LeaseDocument> { ReturnDocument = ReturnDocument.After }, cancellationToken: cancellationToken);
 			if (document == null)
 			{
