Horde: Initialize PerforceMetadataLogger even if retrieved after a sync

diff --git a/Engine/Source/Programs/Horde/Drivers/JobDriver.Tests/ManagedWorkspaceMaterializerTests.cs b/Engine/Source/Programs/Horde/Drivers/JobDriver.Tests/ManagedWorkspaceMaterializerTests.cs
new file mode 100644
--- /dev/null
+++ b/Engine/Source/Programs/Horde/Drivers/JobDriver.Tests/ManagedWorkspaceMaterializerTests.cs
@@ -0,0 +1,78 @@
+// Copyright Epic Games, Inc. All Rights Reserved.
+
+using System.Text.Json;
+using System.Threading;
+using System.Threading.Tasks;
+using EpicGames.Core;
+using EpicGames.Perforce.Fixture;
+using HordeCommon.Rpc.Messages;
+using JobDriver.Execution;
+using Microsoft.Extensions.Logging;
+using Microsoft.VisualStudio.TestTools.UnitTesting;
+using OpenTelemetry.Trace;
+using SyncOptions = JobDriver.Execution.SyncOptions;
+
+namespace JobDriver.Tests;
+
+[TestClass]
+public class ManagedWorkspaceMaterializerTests : BasePerforceFixtureTest
+{
+	private static Tracer NoOpTracer { get; } = TracerProvider.Default.GetTracer("NoOp");
+
+	[TestMethod]
+	[DataRow(true, DisplayName = "Get logger before sync")]
+	[DataRow(false, DisplayName = "Get logger after sync")]
+	public async Task Logger_RelativeDepotPath_Async(bool getLoggerBeforeSync)
+	{
+		using ManagedWorkspaceMaterializer pm = await CreateMwmAsync(PerforceSnapshot.FooMain);
+		TestLogger baseLogger = new();
+		ILogger materializerLogger;
+		
+		if (getLoggerBeforeSync)
+		{
+			materializerLogger = pm.GetLogger(baseLogger);
+			await pm.SyncAsync(IWorkspaceMaterializer.LatestChangeNumber, -1, new SyncOptions(), CancellationToken.None);
+		}
+		else
+		{
+			await pm.SyncAsync(IWorkspaceMaterializer.LatestChangeNumber, -1, new SyncOptions(), CancellationToken.None);
+			materializerLogger = pm.GetLogger(baseLogger);
+		}
+
+		Assert.AreNotEqual(baseLogger, materializerLogger);
+		FileReference source = FileReference.Combine(pm.SyncDir, "Sub1", "Sub2", "MyFile.txt");
+		LogValue fileValue = LogValue.SourceFile(source, source.GetFileName());
+		materializerLogger.LogError(KnownLogEvents.Compiler, "Testing path resolution {File}", fileValue);
+		Assert.AreEqual(1, baseLogger.Entries.Count);
+
+		if (baseLogger.Entries[0].State is JsonLogEvent jle)
+		{
+			JsonDocument document = JsonDocument.Parse(jle.Data);
+			JsonElement fileElement = document.RootElement.GetProperty("properties").GetProperty("File");
+			Assert.AreEqual("MyFile.txt", fileElement.GetProperty("$text").GetString());
+			Assert.AreEqual("Sub1/Sub2/MyFile.txt", fileElement.GetProperty("relativePath").GetString());
+			Assert.AreEqual("//Foo/Main/Sub1/Sub2/MyFile.txt@9", fileElement.GetProperty("depotPath").GetString());
+		}
+		else
+		{
+			Assert.Fail("Unable to find JsonLogEvent");
+		}
+	}
+
+	private async Task<ManagedWorkspaceMaterializer> CreateMwmAsync(string stream, string? dataDir = null, string? identifier = "mwmTest", string[]? view = null)
+	{
+		ILogger logger = LoggerFactory.CreateLogger("MWM");
+		RpcAgentWorkspace raw = new ()
+		{
+			ServerAndPort = PerforceConnection.ServerAndPort,
+			UserName = PerforceConnection.UserName,
+			Password = PerforceConnection.Settings.Password,
+			Identifier = $"{TestGuid}-{identifier}",
+			Stream = stream,
+			MinScratchSpace = 1
+		};
+		raw.View.AddRange(view ?? []);
+		return await ManagedWorkspaceMaterializer.CreateAsync(raw, dataDir is null ? TempDir : new DirectoryReference(dataDir), false, true, NoOpTracer, logger, CancellationToken.None);
+	}
+}
+
diff --git a/Engine/Source/Programs/Horde/Drivers/JobDriver.Tests/PerforceMaterializerTests.cs b/Engine/Source/Programs/Horde/Drivers/JobDriver.Tests/PerforceMaterializerTests.cs
--- a/Engine/Source/Programs/Horde/Drivers/JobDriver.Tests/PerforceMaterializerTests.cs
+++ b/Engine/Source/Programs/Horde/Drivers/JobDriver.Tests/PerforceMaterializerTests.cs
@@ -70,15 +70,26 @@ public class PerforceMaterializerTests : BasePerforceFixtureTest
 	}
 
 	[TestMethod]
-	public async Task Logger_RelativeDepotPath_Async()
+	[DataRow(true, DisplayName = "Get logger before sync")]
+	[DataRow(false, DisplayName = "Get logger after sync")]
+	public async Task Logger_RelativeDepotPath_Async(bool getLoggerBeforeSync)
 	{
 		using PerforceMaterializer pm = CreatePm(PerforceSnapshot.FooMain);
-
 		TestLogger baseLogger = new();
-		ILogger materializerLogger = pm.GetLogger(baseLogger);
-		Assert.AreNotEqual(baseLogger, materializerLogger);
-		await pm.SyncAsync(IWorkspaceMaterializer.LatestChangeNumber, -1, new SyncOptions(), CancellationToken.None);
+		ILogger materializerLogger;
+		
+		if (getLoggerBeforeSync)
+		{
+			materializerLogger = pm.GetLogger(baseLogger);
+			await pm.SyncAsync(IWorkspaceMaterializer.LatestChangeNumber, -1, new SyncOptions(), CancellationToken.None);
+		}
+		else
+		{
+			await pm.SyncAsync(IWorkspaceMaterializer.LatestChangeNumber, -1, new SyncOptions(), CancellationToken.None);
+			materializerLogger = pm.GetLogger(baseLogger);
+		}
 
+		Assert.AreNotEqual(baseLogger, materializerLogger);
 		FileReference source = FileReference.Combine(pm.SyncDir, "Sub1", "Sub2", "MyFile.txt");
 		LogValue fileValue = LogValue.SourceFile(source, source.GetFileName());
 		materializerLogger.LogError(KnownLogEvents.Compiler, "Testing path resolution {File}", fileValue);
diff --git a/Engine/Source/Programs/Horde/Drivers/JobDriver/Execution/ManagedWorkspaceMaterializer.cs b/Engine/Source/Programs/Horde/Drivers/JobDriver/Execution/ManagedWorkspaceMaterializer.cs
--- a/Engine/Source/Programs/Horde/Drivers/JobDriver/Execution/ManagedWorkspaceMaterializer.cs
+++ b/Engine/Source/Programs/Horde/Drivers/JobDriver/Execution/ManagedWorkspaceMaterializer.cs
@@ -28,6 +28,7 @@ public sealed class ManagedWorkspaceMaterializer : IWorkspaceMaterializer
 	private readonly Tracer _tracer;
 	private readonly ILogger _logger;
 	private PerforceMetadataLogger? _perforceLogger;
+	private int? _lastSyncedChangeNum;
 
 	/// <inheritdoc/>
 	public DirectoryReference SyncDir => _workspace.WorkspaceDir;
@@ -100,6 +101,7 @@ public sealed class ManagedWorkspaceMaterializer : IWorkspaceMaterializer
 	public ILogger GetLogger(ILogger logger)
 	{
 		_perforceLogger = new PerforceMetadataLogger(logger);
+		AddClientViewToLogger();
 		return _perforceLogger;
 	}
 
@@ -119,7 +121,9 @@ public sealed class ManagedWorkspaceMaterializer : IWorkspaceMaterializer
 			span.SetAttribute("horde.change_num_latest", latestChangeNum);
 			changeNum = latestChangeNum;
 		}
-		_perforceLogger?.AddClientView(SyncDir, _workspace.StreamView, changeNum);
+
+		_lastSyncedChangeNum = changeNum;
+		AddClientViewToLogger();
 
 		FileReference cacheFile = FileReference.Combine(_workspace.MetadataDir, "Contents.dat");
 		if (_useCacheFile)
@@ -167,6 +171,14 @@ public sealed class ManagedWorkspaceMaterializer : IWorkspaceMaterializer
 		return _workspace;
 	}
 
+	private void AddClientViewToLogger()
+	{
+		if (_lastSyncedChangeNum is not null)
+		{
+			_perforceLogger?.AddClientView(SyncDir, _workspace.StreamView, _lastSyncedChangeNum.Value);
+		}
+	}
+
 	private TelemetrySpan CreateTraceSpan(string operationName)
 	{
 		TelemetrySpan span = _tracer.StartActiveSpan(operationName);
diff --git a/Engine/Source/Programs/Horde/Drivers/JobDriver/Execution/PerforceMaterializer.cs b/Engine/Source/Programs/Horde/Drivers/JobDriver/Execution/PerforceMaterializer.cs
--- a/Engine/Source/Programs/Horde/Drivers/JobDriver/Execution/PerforceMaterializer.cs
+++ b/Engine/Source/Programs/Horde/Drivers/JobDriver/Execution/PerforceMaterializer.cs
@@ -136,6 +136,7 @@ public sealed class PerforceMaterializer : IWorkspaceMaterializer
 	private IPerforceConnection? _perforceWithoutClient;
 	private IPerforceConnection? _perforceWithClient;
 	private List<string>? _streamView;
+	private int? _lastSyncedChangeNum;
 	private const string StateFilename = "State.json";
 	private const bool PreferNativePerforceClient = true;
 	private const int MaxFilesLogged = 100;
@@ -229,6 +230,7 @@ public sealed class PerforceMaterializer : IWorkspaceMaterializer
 	public ILogger GetLogger(ILogger logger)
 	{
 		_perforceLogger = new PerforceMetadataLogger(logger);
+		AddClientViewToLogger();
 		return _perforceLogger;
 	}
 
@@ -311,10 +313,8 @@ public sealed class PerforceMaterializer : IWorkspaceMaterializer
 			changeNum = changesRecords[0].Number;
 		}
 
-		if (_streamView != null)
-		{
-			_perforceLogger?.AddClientView(SyncDir, PerforceViewMap.Parse(_streamView), changeNum);
-		}
+		_lastSyncedChangeNum = changeNum;
+		AddClientViewToLogger();
 
 		await SaveStateAsync(TransactionStatus.Dirty, client.Name, changeNum, shelveChangeNum, [], cancellationToken);
 		
@@ -734,6 +734,14 @@ public sealed class PerforceMaterializer : IWorkspaceMaterializer
 		}
 	}
 
+	private void AddClientViewToLogger()
+	{
+		if (_streamView is not null && _lastSyncedChangeNum is not null)
+		{
+			_perforceLogger?.AddClientView(SyncDir, PerforceViewMap.Parse(_streamView), _lastSyncedChangeNum.Value);
+		}
+	}
+
 	private static int GetDefaultThreadCount(int defaultValue) => Math.Min(defaultValue, Math.Max(Environment.ProcessorCount - 1, 1));
 	
 	internal async Task DeleteClientForTestAsync(CancellationToken cancellationToken)
