Horde: Update release notes

#rnx

diff --git a/Engine/Source/Programs/Horde/Docs/ReleaseNotes.md b/Engine/Source/Programs/Horde/Docs/ReleaseNotes.md
--- a/Engine/Source/Programs/Horde/Docs/ReleaseNotes.md
+++ b/Engine/Source/Programs/Horde/Docs/ReleaseNotes.md
@@ -1,5 +1,31 @@
 # Release Notes
 
+## 2025-12-04
+
+* Initialize PerforceMetadataLogger even if retrieved after a sync (48953485)
+* Save lease metrics reported by agent to database (48944287)
+* Add step report color hinting (48920015)
+* Updating naming of request to be more appropriate & minor polish. (48919764)
+* Global config flag to enable and disable RCA notifications to be sent (48910456)
+* Use the correct Perforce metadata logger for PerforceMaterializer
+  Refactored IWorkspaceMaterializer to support loggers. (48901078)
+* Export Summary column to CSV report, and sorting step rows by capabilities. (48897489)
+* BuildHealthFilter support (48861859)
+* BuildHealthFilter collection + API endpoints (48860965)
+* Improve logging when saving lease outcomes (48847790)
+* Removing dev suffix. (48835366)
+* Adding new columns for SubmittedBy, PreflightTag, Virtualized (string-bool), as well as UserId conversions from usernames. (48831291)
+* Capture CPU seconds spent during a lease on the agent (48806744)
+* Fix .NET 10 warnings (48803780)
+* Store list of files for replicated Perforce changelists (48803432)
+* Duration label cleanup. (48789113)
+* Add the shared LogMatchers project and its tests to the Shared solution. (48783326)
+* Support -x for opened (48770142)
+* Align submitted naming with existing schema (48768458)
+* [Converting commit telemetry stream to prod (48764675)
+* Have the compile event matcher recognise extensionless files containing underscores
+  UBT fails to recognise a diagnostic in files like "type_traits" (48741032)
+
 ## 2025-11-25
 
 * Perf & tests for commit regex (48627869)
