[Horde][BuildHealth][Dashboard] Minor polish for sticky & more appropriate default height.

#Synced-CL 48916881

diff --git a/Engine/Source/Programs/Horde/HordeDashboard/src/components/buildhealth/BuildhealthView.tsx b/Engine/Source/Programs/Horde/HordeDashboard/src/components/buildhealth/BuildhealthView.tsx
--- a/Engine/Source/Programs/Horde/HordeDashboard/src/components/buildhealth/BuildhealthView.tsx
+++ b/Engine/Source/Programs/Horde/HordeDashboard/src/components/buildhealth/BuildhealthView.tsx
@@ -175,7 +175,7 @@ const SearchableDropdown: React.FC<ISearchableDropdownProps> = (props) => {
                     calloutMain: {
                         width: 500,
                         minWidth: 500,
-                        height: 500,
+                        height: 750,
                         maxHeight: 1000
                     }
                 },
diff --git a/Engine/Source/Programs/Horde/HordeDashboard/src/components/buildhealth/stepoutcome/StepOutcomeD3Component.tsx b/Engine/Source/Programs/Horde/HordeDashboard/src/components/buildhealth/stepoutcome/StepOutcomeD3Component.tsx
--- a/Engine/Source/Programs/Horde/HordeDashboard/src/components/buildhealth/stepoutcome/StepOutcomeD3Component.tsx
+++ b/Engine/Source/Programs/Horde/HordeDashboard/src/components/buildhealth/stepoutcome/StepOutcomeD3Component.tsx
@@ -301,6 +301,8 @@ function GenerateD3Table(stepOutcomeTable: StepOutcomeTable, handler: StepOutcom
                         .style("min-width", "30px")
                         .style("max-width", "30px")
                         .style("cursor", "pointer")
+                        .style("left", "0")
+
                     ),
                 update => update
                     .attr("class", d => `start-header ${getCellSortClasses("start", stepOutcomeTable.sortField, stepOutcomeTable.sortDirectionAsc)}`)
