[Horde][BuildHealth][Dashboard] More polish on default hights for right top edge alignment

#Synced-CL 48916881

diff --git a/Engine/Source/Programs/Horde/HordeDashboard/src/components/buildhealth/BuildhealthView.tsx b/Engine/Source/Programs/Horde/HordeDashboard/src/components/buildhealth/BuildhealthView.tsx
--- a/Engine/Source/Programs/Horde/HordeDashboard/src/components/buildhealth/BuildhealthView.tsx
+++ b/Engine/Source/Programs/Horde/HordeDashboard/src/components/buildhealth/BuildhealthView.tsx
@@ -167,16 +167,11 @@ const SearchableDropdown: React.FC<ISearchableDropdownProps> = (props) => {
                 directionalHint: DirectionalHint.rightTopEdge,
                 alignTargetEdge: true,
                 styles: {
-                    root: {
-                        border: `1px solid ${modeColors.content}`,
-                        borderRadius: 4,
-                        boxSizing: "border-box",
-                    },
                     calloutMain: {
                         width: 500,
                         minWidth: 500,
-                        height: 750,
-                        maxHeight: 1000
+                        minHeight: 200,
+                        maxHeight: 800
                     }
                 },
             }}
@@ -186,7 +181,7 @@ const SearchableDropdown: React.FC<ISearchableDropdownProps> = (props) => {
                 }
 
                 return (
-                    <div ref={dropdownRef} style={{ maxHeight: 1000, overflowY: 'auto', width: 500 }}>
+                    <div ref={dropdownRef} style={{ maxHeight: 800, overflowY: 'auto', width: 500 }}>
                         <Stack
                             horizontal
                             verticalAlign="center"
