[Horde][BuildHealth][Dashboard] Pass along old project id to prevent full refresh.

#Synced-CL 48916881

diff --git a/Engine/Source/Programs/Horde/HordeDashboard/src/components/buildhealth/BuildHealthFilterComponent.tsx b/Engine/Source/Programs/Horde/HordeDashboard/src/components/buildhealth/BuildHealthFilterComponent.tsx
--- a/Engine/Source/Programs/Horde/HordeDashboard/src/components/buildhealth/BuildHealthFilterComponent.tsx
+++ b/Engine/Source/Programs/Horde/HordeDashboard/src/components/buildhealth/BuildHealthFilterComponent.tsx
@@ -271,10 +271,10 @@ export const BuildHealthFilterComponent: React.FC<{ handler: BuildHealthDataHand
             buildHealthOptions.setBuildHealthFilterId(savedFilter.id);
             handler.getProjectFiltersData(savedFilter.filterProject);
         },
-        onDeleteSuccess: () => {
+        onDeleteSuccess: (oldFilterProjectId?: string) => {
             setModalVisible(false);
             buildHealthOptions.setBuildHealthFilterId(null);
-            handler.getProjectFiltersData();
+            handler.getProjectFiltersData(oldFilterProjectId);
         },
         onRevert: () => {
             setModalVisible(false);
diff --git a/Engine/Source/Programs/Horde/HordeDashboard/src/components/buildhealth/BuildHealthFilterModal.tsx b/Engine/Source/Programs/Horde/HordeDashboard/src/components/buildhealth/BuildHealthFilterModal.tsx
--- a/Engine/Source/Programs/Horde/HordeDashboard/src/components/buildhealth/BuildHealthFilterModal.tsx
+++ b/Engine/Source/Programs/Horde/HordeDashboard/src/components/buildhealth/BuildHealthFilterModal.tsx
@@ -57,7 +57,7 @@ export type BuildHealthFilterModalProps = {
     onRevert?: () => any;
 
     // The delegate to execute when the Build Health Filter has been deleted.
-    onDeleteSuccess?: () => any;
+    onDeleteSuccess?: (oldFilterProjectId? : string) => any;
 };
 
 // #endregion -- Private Types --
@@ -136,10 +136,10 @@ export const BuildHealthFilterModal: React.FC<BuildHealthFilterModalProps> = ({
         }
     };
 
-    const handleDeleteClick = async () => {
+    const handleDeleteClick = async (oldFilterProjectId? : string) => {
         const savedFilter = await Delete();
         if (savedFilter) {
-            onDeleteSuccess?.();
+            onDeleteSuccess?.(oldFilterProjectId);
         }
     };
 
@@ -237,7 +237,7 @@ export const BuildHealthFilterModal: React.FC<BuildHealthFilterModalProps> = ({
             {(isNewFilter || isEditorOwner) && <PrimaryButton disabled={(!isNewFilter && !isOldAndNewDiff)} onClick={() => { handleSaveClick(); }} text="Save" />}
             {!isNewFilter && <PrimaryButton onClick={() => { handleSaveAsCopyClick(); }} text="Save As Copy" />}
             {!isNewFilter && isOldAndNewDiff && <PrimaryButton onClick={() => { handleRevertClick(); }} text="Revert" />}
-            {!isNewFilter && <PrimaryButton disabled={buildHealthFilter.filterOwner !== null && buildHealthFilter.filterOwner !== dashboard.userId} onClick={() => { handleDeleteClick(); }} text="Delete" />}
+            {!isNewFilter && <PrimaryButton disabled={buildHealthFilter.filterOwner !== null && buildHealthFilter.filterOwner !== dashboard.userId} onClick={() => { handleDeleteClick(buildHealthFilter?.filterProjectId); }} text="Delete" />}
         </Dialog >
     );
 };
