Horde: Ensure untracked files are removed in PerforceMaterializer

diff --git a/Engine/Source/Programs/Horde/Drivers/JobDriver.Tests/PerforceMaterializerTests.cs b/Engine/Source/Programs/Horde/Drivers/JobDriver.Tests/PerforceMaterializerTests.cs
--- a/Engine/Source/Programs/Horde/Drivers/JobDriver.Tests/PerforceMaterializerTests.cs
+++ b/Engine/Source/Programs/Horde/Drivers/JobDriver.Tests/PerforceMaterializerTests.cs
@@ -129,13 +129,41 @@ public class PerforceMaterializerTests : BasePerforceFixtureTest
 		PerforceMaterializer.State? state = await pm.LoadStateForTestAsync(CancellationToken.None);
 		Assert.AreEqual(expectedChangeNum, state!.Changelist);
 	}
+
+	[TestMethod]
+	public async Task Sync_RemoveUntracked_Async()
+	{
+		using PerforceMaterializer pm = CreatePm(PerforceSnapshot.FooMain);
+		WorkspaceState ws = PerforceSnapshot.Main(2) with
+		{
+			RemoveUntracked = true,
+			OnPostSync = () => SimulateWorkspaceChanges(pm.SyncDir.FullName)
+		};
+		await SyncToStateAsync(pm, ws, CancellationToken.None);
+	}
+	
+	[TestMethod]
+	public async Task Sync_KeepUntracked_Async()
+	{
+		using PerforceMaterializer pm = CreatePm(PerforceSnapshot.FooMain);
+		string addedFile = $"{pm.SyncDir.FullName}/Data/added.txt";
+		WorkspaceState ws = PerforceSnapshot.Main(2) with
+		{
+			RemoveUntracked = false,
+			OnPostSync = () => { File.WriteAllText(addedFile, "New file"); }
+		};
+		await SyncToStateAsync(pm, ws, CancellationToken.None);
+		
+		// Verify file is *not* cleaned up by materializer
+		Assert.AreEqual("New file", await File.ReadAllTextAsync(addedFile));
+	}
 	
 	[TestMethod]
 	public async Task Conform_RemoveUntracked_Async()
 	{
 		using PerforceMaterializer pm = CreatePm(PerforceSnapshot.FooMain);
-		await pm.SyncAsync(5, -1, new SyncOptions(), CancellationToken.None);
-		await File.WriteAllTextAsync(Path.Join(pm.SyncDir.FullName, _untrackedFile.clientFile), _untrackedFile.content);
+		await pm.SyncAsync(2, -1, new SyncOptions(), CancellationToken.None);
+		SimulateWorkspaceChanges(pm.SyncDir.FullName);
 		
 		await pm.ConformAsync(true, CancellationToken.None);
 		PerforceSnapshot.GetSnapshot().GetStream(PerforceSnapshot.FooMain).GetWorkspace(9).AssertDepotFiles(pm.SyncDir.FullName);
@@ -175,6 +203,10 @@ public class PerforceMaterializerTests : BasePerforceFixtureTest
 			Method = "name=perforce"
 		};
 		
+		using PerforceMaterializer pm = CreatePm(raw.Stream, dataDir: Path.Join(TempDir.FullName, raw.Identifier), identifier: raw.Identifier);
+		await pm.SyncAsync(2, -1, new SyncOptions(), CancellationToken.None);
+		SimulateWorkspaceChanges(pm.SyncDir.FullName);
+		
 		ILogger logger = LoggerFactory.CreateLogger("PM");
 		ServiceCollection services = [];
 		services.AddSingleton<ILogger<PerforceMaterializer>>(x => LoggerFactory.CreateLogger<PerforceMaterializer>());
@@ -280,9 +312,20 @@ public class PerforceMaterializerTests : BasePerforceFixtureTest
 
 	private static async Task SyncToStateAsync(PerforceMaterializer pm, WorkspaceState state, CancellationToken cancellationToken)
 	{
+		state.OnPreSync?.Invoke();
 		await pm.SyncAsync(state.Cl, state.ShelvedCl, new SyncOptions { RemoveUntracked = state.RemoveUntracked }, cancellationToken);
 		PerforceSnapshot.GetSnapshot().GetStream(state.Stream).GetWorkspace(state.Cl, state.ShelvedCl).AssertDepotFiles(pm.SyncDir.FullName);
+
+		state.OnPostSync?.Invoke(); // Perform imaginary file operations to workspace
 		await pm.FinalizeAsync(cancellationToken);
+
+		if (state.RemoveUntracked)
+		{
+			// Only verify files on disk if they are expected to be cleaned up.
+			// Note that shelved CL is not checked for as state file/have-list does not represent that.
+			// Materializer simply p4 prints those. 
+			PerforceSnapshot.GetSnapshot().GetStream(state.Stream).GetWorkspace(state.Cl).AssertDepotFiles(pm.SyncDir.FullName);	
+		}
 		
 		PerforceMaterializer.State? pmState = await pm.LoadStateForTestAsync(cancellationToken);
 		Assert.AreEqual(state.Cl, pmState!.Changelist);
@@ -319,5 +362,20 @@ public class PerforceMaterializerTests : BasePerforceFixtureTest
 		change = await perforce.CreateChangeAsync(change);
 		await perforce.SubmitAsync(change.Number, SubmitOptions.None);
 	}
-}
 
+	private static void SimulateWorkspaceChanges(string rootDir)
+	{
+		// With CleanOptions.ModifiedTimes, last modified time of a file is too close in time for Perforce client to detect a change.
+		// We nudge it just by a second to simulate real-world behavior.
+		// Using digest-based change detection is simply unfeasible for the size of streams Horde handles.
+		_ = new FileInfo($"{rootDir}/main.cpp") { IsReadOnly = false };
+		File.AppendAllText($"{rootDir}/main.cpp", " appended text!");
+		_ = new FileInfo($"{rootDir}/main.cpp") { LastWriteTime = DateTime.Now - TimeSpan.FromSeconds(1) };
+
+		File.Move($"{rootDir}/main.h", $"{rootDir}/main-moved.h");
+		File.WriteAllText($"{rootDir}/Data/added.txt", "New file");
+				
+		_ = new FileInfo($"{rootDir}/common.h") { IsReadOnly = false };
+		File.Delete($"{rootDir}/common.h");
+	}
+}
\ No newline at end of file
diff --git a/Engine/Source/Programs/Horde/Drivers/JobDriver/Execution/PerforceExecutor.cs b/Engine/Source/Programs/Horde/Drivers/JobDriver/Execution/PerforceExecutor.cs
--- a/Engine/Source/Programs/Horde/Drivers/JobDriver/Execution/PerforceExecutor.cs
+++ b/Engine/Source/Programs/Horde/Drivers/JobDriver/Execution/PerforceExecutor.cs
@@ -517,6 +517,7 @@ namespace JobDriver.Execution
 					{
 						logger.LogInformation("Deleting directory {DeleteDir}", dir);
 						FileUtils.ForceDeleteDirectory(dir);
+						logger.LogInformation("Finished deleting directory");
 					}
 				}
 				logger.LogInformation("");
diff --git a/Engine/Source/Programs/Horde/Drivers/JobDriver/Execution/PerforceMaterializer.cs b/Engine/Source/Programs/Horde/Drivers/JobDriver/Execution/PerforceMaterializer.cs
--- a/Engine/Source/Programs/Horde/Drivers/JobDriver/Execution/PerforceMaterializer.cs
+++ b/Engine/Source/Programs/Horde/Drivers/JobDriver/Execution/PerforceMaterializer.cs
@@ -137,9 +137,10 @@ public sealed class PerforceMaterializer : IWorkspaceMaterializer
 	private IPerforceConnection? _perforceWithClient;
 	private List<string>? _streamView;
 	private int? _lastSyncedChangeNum;
+	private SyncOptions? _lastSyncOptions;
 	private const string StateFilename = "State.json";
 	private const bool PreferNativePerforceClient = true;
-	private const int MaxFilesLogged = 100;
+	private const int MaxFilesLogged = 500;
 	private readonly string _stateFile;
 	private readonly string _stateTempFile;
 	private PerforceMetadataLogger? _perforceLogger;
@@ -244,14 +245,16 @@ public sealed class PerforceMaterializer : IWorkspaceMaterializer
 		{
 			throw new WorkspaceMaterializationException("Perforce settings/client not initialized!");
 		}
-		
+
+		_lastSyncOptions = options;
 		State? state = await LoadStateAsync(cancellationToken);
 		
 		bool isDirty = state == null
 			|| state.Version != CurrentStateVersion
 			|| state.Status == TransactionStatus.Dirty
 			|| state.Identifier != Identifier
-			|| state.Stream != _options.AgentWorkspace.Stream;
+			|| state.Stream != _options.AgentWorkspace.Stream
+			|| state.Changelist == -1;
 		
 		_logger.LogInformation("Status: {Status}\nVersion: {PrevVersion} -> {Version}\nClient: {PrevClient} -> {Client}\nIdentifier: {PrevId} -> {Id}\nStream: {PrevStream} -> {Stream}\nCL: {PrevCl} -> {Cl}\nShelve: {PrevShelve} -> {Shelve}\nDirty: {Dirty}",
 			state?.Status,
@@ -360,6 +363,10 @@ public sealed class PerforceMaterializer : IWorkspaceMaterializer
 		if (_perforceWithClient != null)
 		{
 			await RevertInternalAsync(_perforceWithClient, cancellationToken);
+			if (_lastSyncOptions is { RemoveUntracked: true })
+			{
+				await CleanAsync(_perforceWithClient, cancellationToken);
+			}
 		}
 	}
 
@@ -374,16 +381,7 @@ public sealed class PerforceMaterializer : IWorkspaceMaterializer
 
 		if (removeUntrackedFiles)
 		{
-			using TelemetrySpan cleanSpan = _tracer.StartActiveSpan($"{nameof(PerforceMaterializer)}.{nameof(ConformAsync)}.Clean");
-			using ILoggerProgress cleanStatus = _logger.BeginProgressScope("Cleaning files...");
-			Stopwatch cleanTimer = Stopwatch.StartNew();
-			CleanOptions cleanOptions = CleanOptions.ModifiedTimes | CleanOptions.Added | CleanOptions.Deleted | CleanOptions.Edited;
-			PerforceResponseList<CleanRecord> cleanRes = await perforceWithClient.TryCleanAsync(cleanOptions, new List<string>() { "//..."}, cancellationToken);
-			if (!cleanRes.Succeeded)
-			{
-				_logger.LogInformation("Clean skipped.");
-			}
-			cleanStatus.Progress = $"({cleanTimer.Elapsed.TotalSeconds:0.0}s)";
+			await CleanAsync(perforceWithClient, cancellationToken);
 		}
 
 		// Sync from scratch again. Any files are reverted by SyncAsync
@@ -614,6 +612,41 @@ public sealed class PerforceMaterializer : IWorkspaceMaterializer
 		span.SetAttribute("horde.pm.num_files", revertRecords.Count);
 		status.Progress = $"({timer.Elapsed.TotalSeconds:0.0}s)";
 	}
+
+	private async Task CleanAsync(IPerforceConnection perforce, CancellationToken cancellationToken)
+	{
+		if (_perforceSettings == null || _perforceWithClient == null)
+		{
+			throw new WorkspaceMaterializationException("Perforce settings/client not initialized!");
+		}
+		
+		using TelemetrySpan span = _tracer.StartActiveSpan($"{nameof(PerforceMaterializer)}.{nameof(CleanAsync)}");
+		using ILoggerProgress status = _logger.BeginProgressScope("Cleaning files...");
+		Stopwatch timer = Stopwatch.StartNew();
+
+		await SaveStateAsync(TransactionStatus.Dirty, cancellationToken);
+		PerforceResponseList<CleanRecord> res = await perforce.TryCleanAsync(CleanOptions.ModifiedTimes, "//...", cancellationToken);
+		if (res.Succeeded)
+		{
+			List<CleanRecord> cleanedFiles = res.Data;
+			foreach (CleanRecord cr in cleanedFiles.Take(MaxFilesLogged))
+			{
+				_logger.LogInformation("Cleaned file: {File} {Action}", cr.DepotFile, cr.Action);
+			}
+			if (cleanedFiles.Count > MaxFilesLogged)
+			{
+				_logger.LogInformation(".. with {MoreCount} more files", cleanedFiles.Count - MaxFilesLogged);
+			}
+			span.SetAttribute("horde.pm.num_files", cleanedFiles.Count);
+		}
+		else
+		{
+			_logger.LogInformation("Clean skipped");
+		}
+		
+		await SaveStateAsync(TransactionStatus.Clean, cancellationToken);
+		status.Progress = $"({timer.Elapsed.TotalSeconds:0.0}s)";
+	}
 	
 	private async Task<(ClientRecord client, bool wasRecreated)> CreateOrUpdateClientAsync(bool forceCreate, CancellationToken cancellationToken)
 	{
@@ -695,16 +728,23 @@ public sealed class PerforceMaterializer : IWorkspaceMaterializer
 		return $"{exclude}{stream}{path} //{clientName}{path}";
 	}
 
-	private Task SaveStateAsync(TransactionStatus status, string client ,int changelist, int shelvedChangelist, List<WorkspaceFile> dirtyFiles, CancellationToken cancellationToken)
+	private Task SaveStateAsync(TransactionStatus status, string client, int changelist, int shelvedChangelist, List<WorkspaceFile> dirtyFiles, CancellationToken cancellationToken)
 	{
 		State pms = new (CurrentStateVersion, status, client, Identifier, _options.AgentWorkspace.Stream, changelist, shelvedChangelist, dirtyFiles);
 		return SaveStateAsync(pms, cancellationToken);
 	}
 	
+	private async Task SaveStateAsync(TransactionStatus status, CancellationToken cancellationToken)
+	{
+		State state = await LoadStateAsync(cancellationToken) ?? new State(CurrentStateVersion, status, "<not initialized>", Identifier, _options.AgentWorkspace.Stream, -1, -1, []);
+		await SaveStateAsync(state with { Status = status }, cancellationToken);
+	}
+	
 	private async Task SaveStateAsync(State state, CancellationToken cancellationToken)
 	{
 		try
 		{
+			Directory.CreateDirectory(_metadataDir.FullName); // Ensure metadata directory exists
 			string json = JsonSerializer.Serialize(state, s_jsonOptions);
 			await File.WriteAllTextAsync(_stateTempFile, json, cancellationToken);
 			File.Move(_stateTempFile, _stateFile, overwrite: true);
diff --git a/Engine/Source/Programs/Shared/EpicGames.Perforce.Fixture/PerforceSnapshot.cs b/Engine/Source/Programs/Shared/EpicGames.Perforce.Fixture/PerforceSnapshot.cs
--- a/Engine/Source/Programs/Shared/EpicGames.Perforce.Fixture/PerforceSnapshot.cs
+++ b/Engine/Source/Programs/Shared/EpicGames.Perforce.Fixture/PerforceSnapshot.cs
@@ -102,7 +102,7 @@ public record WorkspaceSnapshot(int Change, int ShelvedChange, HashSet<Workspace
 	}
 }
 
-public record WorkspaceState(string Stream, int Cl, int ShelvedCl = 0, bool RemoveUntracked = true);
+public record WorkspaceState(string Stream, int Cl, int ShelvedCl = 0, bool RemoveUntracked = true, Action? OnPreSync = null, Action? OnPostSync = null);
 
 public class PerforceSnapshot(PerforceConnection perforce)
 {
