Introduce Test Health service for test health report generation
+ Add support for Test Health Report notification

#rnx
#Synced-CL 48996119

diff --git a/Engine/Source/Programs/Horde/Plugins/Build/HordeServer.Build.Tests/BuildTestSetup.cs b/Engine/Source/Programs/Horde/Plugins/Build/HordeServer.Build.Tests/BuildTestSetup.cs
--- a/Engine/Source/Programs/Horde/Plugins/Build/HordeServer.Build.Tests/BuildTestSetup.cs
+++ b/Engine/Source/Programs/Horde/Plugins/Build/HordeServer.Build.Tests/BuildTestSetup.cs
@@ -100,6 +100,7 @@ namespace HordeServer.Tests
 		public DeviceService DeviceService => ServiceProvider.GetRequiredService<DeviceService>();
 		public StorageService StorageService => ServiceProvider.GetRequiredService<StorageService>();
 		public TestDataService TestDataService => ServiceProvider.GetRequiredService<TestDataService>();
+		public TestHealthService TestHealthService => ServiceProvider.GetRequiredService<TestHealthService>();
 		public ComputeService ComputeService => ServiceProvider.GetRequiredService<ComputeService>();
 
 		public StreamsController StreamsController => GetStreamsController();
@@ -189,6 +190,7 @@ namespace HordeServer.Tests
 			services.AddSingleton<ScheduleService>();
 			services.AddSingleton<DeviceService>();
 			services.AddSingleton<TestDataService>();
+			services.AddSingleton<TestHealthService>();
 			services.AddSingleton<ComputeService>();
 			services.AddSingleton<EnrollmentService>();
 
diff --git a/Engine/Source/Programs/Horde/Plugins/Build/HordeServer.Build.Tests/Jobs/TestData/TestDataTests.cs b/Engine/Source/Programs/Horde/Plugins/Build/HordeServer.Build.Tests/Jobs/TestData/TestDataTests.cs
--- a/Engine/Source/Programs/Horde/Plugins/Build/HordeServer.Build.Tests/Jobs/TestData/TestDataTests.cs
+++ b/Engine/Source/Programs/Horde/Plugins/Build/HordeServer.Build.Tests/Jobs/TestData/TestDataTests.cs
@@ -78,6 +78,7 @@ namespace HordeServer.Tests.Jobs.TestData
 			StreamConfig streamConfig = new StreamConfig { Id = streamId, Name = streamName };
 			streamConfig.Tabs.Add(new TabConfig { Title = "General", Templates = new List<TemplateId> { new TemplateId("test-template") } });
 			streamConfig.Templates.Add(new TemplateRefConfig { Id = new TemplateId("test-template") });
+			streamConfig.Workflows.Add(new HordeServer.Issues.WorkflowConfig() { Id = new WorkflowId("test-workflow"), ReportChannel = "some-channel" });
 			return streamConfig;
 		}
 
@@ -737,6 +738,133 @@ namespace HordeServer.Tests.Jobs.TestData
 			Assert.AreEqual("TeamTest", streams[0].Tests[0].Team);
 		}
 
+
+		[TestMethod]
+		public async Task TestHealthReportAsync()
+		{
+			await TestDataService.StartAsync(CancellationToken.None);
+
+			string[] streamIds = new string[] { _mainStreamId.ToString(), _releaseStreamId.ToString() };
+
+			IJob job = CreateJob(_mainStreamId, 105, "Test Build", _graph);
+			IJobStep step = job.Batches[0].Steps[0];
+			IJob job2 = CreateJob(_releaseStreamId, 105, "Test Build", _graph);
+			IJobStep step2 = job2.Batches[0].Steps[0];
+			IJob job3 = CreateJob(_mainStreamId, 106, "Test Build", _graph);
+			IJobStep step3 = job3.Batches[0].Steps[0];
+			IJob job4 = CreateJob(_releaseStreamId, 107, "Test Build", _graph);
+			IJobStep step4 = job4.Batches[0].Steps[0];
+
+			BsonDocument testData = BsonDocument.Parse(String.Join('\n', _sessionTestDataV2Lines5));
+			List<(string key, BsonDocument)> data = new List<(string key, BsonDocument)>();
+			foreach (BsonValue item in testData.GetValue("items").AsBsonArray.ToList())
+			{
+				BsonDocument value = item.AsBsonDocument.GetValue("data").AsBsonDocument;
+				data.Add((item.AsBsonDocument.GetValue("key").AsString, value));
+			}
+
+			BsonDocument testData2 = BsonDocument.Parse(String.Join('\n', _sessionTestDataV2Lines4));
+			List<(string key, BsonDocument)> data2 = new List<(string key, BsonDocument)>();
+			foreach (BsonValue item in testData2.GetValue("items").AsBsonArray.ToList())
+			{
+				BsonDocument value = item.AsBsonDocument.GetValue("data").AsBsonDocument;
+				data2.Add((item.AsBsonDocument.GetValue("key").AsString, value));
+			}
+
+			await TestDataCollectionV2.AddAsync(job, step, data.ToArray());
+			await TestDataCollectionV2.AddAsync(job2, step2, data.ToArray());
+			await TestDataCollectionV2.AddAsync(job3, step3, data2.ToArray());
+
+			await TestDataService.TickForTestingAsync();
+
+			string? userIdString = TestDataController.HttpContext.User.GetUserId()?.ToString();
+
+			GetStreamTestsRequest streamRequest = new GetStreamTestsRequest() { StreamIds = streamIds.ToList() };
+			ActionResult<List<GetTestSessionStreamResponse>> streamResult = await TestDataController.GetTestStreamsV2Async(streamRequest);
+			Assert.IsNotNull(streamResult);
+			Assert.IsNotNull(streamResult.Value);
+			List<GetTestSessionStreamResponse> streams = streamResult.Value;
+			Assert.AreNotEqual(0, streams.Count);
+			string streamId1 = streams[0].StreamId;
+			string[] testIds = streams.SelectMany(s => s.Tests.Select(t => t.Id)).Distinct().ToArray();
+			Assert.AreEqual(1, testIds.Length);
+			string testId1 = testIds[0];
+
+			AddTestAuditRequest testAuditPostRequest = new AddTestAuditRequest()
+			{
+				TestId = testId1,
+				Owner = userIdString,
+				Customers = [userIdString!],
+				Notifications = new Dictionary<string, string>() { { _mainStreamId.ToString(), "test-workflow" }, { _releaseStreamId.ToString(), "test-workflow" } }
+			};
+			await TestDataController.AddTestAuditV2Async(testAuditPostRequest);
+
+			GetTestHealthReportRequest testRequestOne = new GetTestHealthReportRequest() { StreamIds = [streamId1], TestIds = [testId1] };
+			ActionResult<List<GetTestHealthReportResponse>> testHealthReports = await TestDataController.GetTestHealthReportsV2Async(testRequestOne);
+
+			Assert.IsNotNull(testHealthReports);
+			Assert.IsNotNull(testHealthReports.Value);
+			List<GetTestHealthReportResponse> reports = testHealthReports.Value;
+			Assert.AreEqual(0, reports.Count);
+
+			await TestHealthService.TickForTestingAsync();
+
+			testHealthReports = await TestDataController.GetTestHealthReportsV2Async(testRequestOne);
+
+			Assert.IsNotNull(testHealthReports);
+			Assert.IsNotNull(testHealthReports.Value);
+			reports = testHealthReports.Value;
+			Assert.AreEqual(1, reports.Count);
+
+			GetTestHealthReportRequest testRequestAll = new GetTestHealthReportRequest() { StreamIds = streamIds.ToList() };
+			testHealthReports = await TestDataController.GetTestHealthReportsV2Async(testRequestAll);
+
+			Assert.IsNotNull(testHealthReports);
+			Assert.IsNotNull(testHealthReports.Value);
+			reports = testHealthReports.Value;
+			Assert.AreEqual(2, reports.Count);
+
+			Assert.AreEqual(true, reports[0].IsHealthy);
+			Assert.AreEqual("Fair", reports[0].State);
+			Assert.AreEqual(33, reports[0].SuccessRate);
+			Assert.AreEqual(25, reports[0].CatastrophicFailureRate);
+			Assert.AreEqual(25, reports[0].RedundantErrorRate);
+			Assert.IsNull(reports[0].NotificationLastDate);
+
+			Assert.AreEqual(false, reports[1].IsHealthy);
+			Assert.AreEqual("Unstable", reports[1].State);
+			Assert.AreEqual(0, reports[1].SuccessRate);
+			Assert.AreEqual(33, reports[1].CatastrophicFailureRate);
+			Assert.AreEqual(0, reports[1].RedundantErrorRate);
+			Assert.IsNotNull(reports[1].NotificationLastDate);
+
+			// Add another run that should improve only the unstable test
+			await TestDataCollectionV2.AddAsync(job4, step4, data2.ToArray());
+			await TestDataService.TickForTestingAsync();
+			await TestHealthService.TickForTestingAsync(true);
+
+			testHealthReports = await TestDataController.GetTestHealthReportsV2Async(testRequestAll);
+
+			Assert.IsNotNull(testHealthReports);
+			Assert.IsNotNull(testHealthReports.Value);
+			reports = testHealthReports.Value;
+			Assert.AreEqual(2, reports.Count);
+
+			Assert.AreEqual(true, reports[0].IsHealthy);
+			Assert.AreEqual("Fair", reports[0].State);
+			Assert.AreEqual(33, reports[0].SuccessRate);
+			Assert.AreEqual(25, reports[0].CatastrophicFailureRate);
+			Assert.AreEqual(25, reports[0].RedundantErrorRate);
+			Assert.IsNull(reports[0].NotificationLastDate);
+
+			Assert.AreEqual(true, reports[1].IsHealthy);
+			Assert.AreEqual("Fair", reports[1].State);
+			Assert.AreEqual(33, reports[1].SuccessRate);
+			Assert.AreEqual(25, reports[1].CatastrophicFailureRate);
+			Assert.AreEqual(25, reports[1].RedundantErrorRate);
+			Assert.IsNull(reports[1].NotificationLastDate);
+		}
+
 		private readonly string[] _simpleTestDataLines =
 		{
 			@"{",
@@ -1428,5 +1556,103 @@ namespace HordeServer.Tests.Jobs.TestData
 			@"    ]",
 			@"}"
 		};
+
+		private readonly string[] _sessionTestDataV2Lines5 =
+		{
+			@"{",
+			@"  ""items"": [",
+			@"        {",
+			@"            ""key"": ""UE.Automation(Group:Editor) Lyra"",",
+			@"            ""data"": {",
+			@"   		      ""version"" : 2,",
+			@"                ""summary"": {",
+			@"                   ""testName"": ""Editor"",",
+			@"        		     ""dateTime"": ""2025-03-03T18:14:04"",",
+			@"		             ""timeElapseSec"": 207.02324,",
+			@"                   ""commandline"": ""run ue.automation -group=editor -project=Lyra"",",
+			@"        		     ""phasesTotalCount"": 3,",
+			@"		             ""phasesSucceededCount"": 0,",
+			@"        		     ""phasesUndefinedCount"": 1,",
+			@"		             ""phasesFailedCount"": 2",
+			@"        		  },",
+			@"                ""phases"": [",
+			@"                    {",
+			@"                        ""key"": ""874a6711"",",
+			@"                        ""name"": ""EditConditionParser.EvaluateBitFlags"",",
+			@"                        ""dateTime"": ""2025-03-03T18:07:55"",",
+			@"                        ""timeElapseSec"": 0.0879574,",
+			@"                        ""outcome"": ""Failed"",",
+			@"                        ""errorFingerprint"": ""a6efc189"",",
+			@"                        ""deviceKeys"": [",
+			@"                            ""9AF6B8CC484096D351EC56ABDFBB919F""",
+			@"                        ]",
+			@"                    },",
+			@"                    {",
+			@"                        ""key"": ""03f772e3"",",
+			@"                        ""name"": ""EditConditionParser.EvaluateBool"",",
+			@"                        ""dateTime"": ""2025-03-03T18:08:08"",",
+			@"                        ""timeElapseSec"": 0.008163903,",
+			@"                        ""outcome"": ""Failed"",",
+			@"                        ""errorFingerprint"": ""a6efc188"",",
+			@"                        ""deviceKeys"": [",
+			@"                            ""9AF6B8CC484096D351EC56ABDFBB919F""",
+			@"                        ]",
+			@"                   },",
+			@"                   {",
+			@"                        ""key"": ""62d99b42"",",
+			@"                        ""name"": ""Editor.Failing.Test"",",
+			@"                        ""dateTime"": ""2025-03-03T18:08:14"",",
+			@"                        ""timeElapseSec"": 0.1961886,",
+			@"                        ""outcome"": ""Failed"",",
+			@"                        ""errorFingerprint"": ""a6efc180"",",
+			@"                        ""deviceKeys"": [",
+			@"                            ""9AF6B8CC484096D351EC56ABDFBB919F""",
+			@"                        ],",
+			@"                   }",
+			@"                ]",
+			@"                ""description"": {",
+			@"                    ""intent"": ""test intent"",",
+ 			@"                    ""profile"": ""functional"",",
+			@"                    ""harness"": ""UE Functional Test"",",
+ 			@"                    ""team"": ""UE""",
+			@"                }",
+			@"                ""recipe"": {",
+			@"                    ""nodeClass"": ""testClass"",",
+ 			@"                    ""version"": ""1.0.0"",",
+ 			@"                    ""timeoutMinutes"": 10,",
+			@"                    ""harness"": ""UE Functional Test"",",
+ 			@"                    ""params"": {",
+			@"                      ""group"": ""editor"",",
+			@"                      ""project"": ""Lyra""",
+			@"                    }",
+			@"                }",
+			@"                ""metadata"": {",
+			@"                    ""Platform"": ""Win64"",",
+ 			@"                    ""BuildTarget"": ""Editor"",",
+			@"                    ""Configuration"": ""Development"",",
+ 			@"                    ""Project"": ""Lyra"",",
+			@"                    ""RHI"": ""default""",
+			@"                }",
+			@"                ""devices"": {",
+			@"                  ""9AF6B8CC484096D351EC56ABDFBB919F"": {",
+			@"                    ""name"": ""M-SIHLOYQHDOUKO"",",
+			@"                    ""appInstanceLogPath"": ""Engine/Programs/AutomationTool/Saved/Logs/UE.EditorAutomation(RunTest=Group_Editor)_(Win64_Development_Editor)/Editor/EditorOutput.log"",",
+			@"                    ""metadata"": {",
+			@"                      ""platform"": ""WindowsEditor"",",
+			@"                      ""os_version"": ""Windows 10 (22H2) [10.0.19045.5487] "",",
+			@"                      ""model"": ""Default"",",
+			@"                      ""gpu"": ""Parsec Virtual Display Adapter"",",
+			@"                      ""cpumodel"": ""AMD EPYC 7R32"",",
+			@"                      ""ram_in_gb"": ""127"",",
+			@"                      ""render_mode"": ""D3D12_SM6"",",
+			@"                      ""rhi"": ""D3D12 (SM6)""",
+			@"                    }",
+			@"                  }",
+			@"                },",
+			@"            }",
+			@"        }",
+			@"    ]",
+			@"}"
+		};
 	}
 }
diff --git a/Engine/Source/Programs/Horde/Plugins/Build/HordeServer.Build.Tests/Notifications/NotificationServiceTests.cs b/Engine/Source/Programs/Horde/Plugins/Build/HordeServer.Build.Tests/Notifications/NotificationServiceTests.cs
--- a/Engine/Source/Programs/Horde/Plugins/Build/HordeServer.Build.Tests/Notifications/NotificationServiceTests.cs
+++ b/Engine/Source/Programs/Horde/Plugins/Build/HordeServer.Build.Tests/Notifications/NotificationServiceTests.cs
@@ -1,9 +1,9 @@
 // Copyright Epic Games, Inc. All Rights Reserved.
 
 using EpicGames.Core;
+using EpicGames.Horde.Agents;
 using EpicGames.Horde.Agents.Leases;
 using EpicGames.Horde.Agents.Pools;
-using EpicGames.Horde.Agents;
 using EpicGames.Horde.Agents.Sessions;
 using EpicGames.Horde.Commits;
 using EpicGames.Horde.Jobs;
@@ -19,6 +19,7 @@ using HordeServer.Configuration;
 using HordeServer.Devices;
 using HordeServer.Issues;
 using HordeServer.Jobs;
+using HordeServer.Jobs.TestData;
 using HordeServer.Logs;
 using HordeServer.Notifications;
 using HordeServer.Projects;
@@ -94,6 +95,7 @@ namespace HordeServer.Tests.Notifications
 		public Task SendAgentReportAsync(AgentReport report, CancellationToken cancellationToken) => throw new NotImplementedException();
 		public Task SendIssueReportAsync(IssueReportGroup report, CancellationToken cancellationToken) => throw new NotImplementedException();
 		public Task SendDeviceIssueReportAsync(DeviceIssueReport report, CancellationToken cancellationToken) => throw new NotImplementedException();
+		public Task NotifyTestHealthReportAsync(ITestHealthReport report, string recipient, string[]? carbonCopies, CancellationToken cancellationToken) => throw new NotImplementedException();
 	}
 
 	[TestClass]
diff --git a/Engine/Source/Programs/Horde/Plugins/Build/HordeServer.Build/BuildAclModifier.cs b/Engine/Source/Programs/Horde/Plugins/Build/HordeServer.Build/BuildAclModifier.cs
--- a/Engine/Source/Programs/Horde/Plugins/Build/HordeServer.Build/BuildAclModifier.cs
+++ b/Engine/Source/Programs/Horde/Plugins/Build/HordeServer.Build/BuildAclModifier.cs
@@ -7,6 +7,7 @@ using HordeServer.Artifacts;
 using HordeServer.Devices;
 using HordeServer.Jobs;
 using HordeServer.Jobs.Bisect;
+using HordeServer.Jobs.TestData;
 using HordeServer.Logs;
 using HordeServer.Notifications;
 using HordeServer.Projects;
@@ -30,6 +31,7 @@ namespace HordeServer
 			acl.AddDefaultReadAction(ArtifactAclAction.ReadArtifact);
 			acl.AddDefaultReadAction(BisectTaskAclAction.ViewBisectTask);
 			acl.AddDefaultReadAction(DeviceAclAction.DeviceRead);
+			acl.AddDefaultReadAction(TestDataAclAction.TestHealthRead);
 			acl.AddDefaultReadAction(JobAclAction.ViewJob);
 			acl.AddDefaultReadAction(NotificationAclAction.CreateSubscription);
 			acl.AddDefaultReadAction(ProjectAclAction.ViewProject);
@@ -41,6 +43,7 @@ namespace HordeServer
 			acl.AddDefaultWriteAction(JobAclAction.UpdateJob);
 			acl.AddDefaultWriteAction(JobAclAction.RetryJobStep);
 			acl.AddDefaultWriteAction(DeviceAclAction.DeviceWrite);
+			acl.AddDefaultWriteAction(TestDataAclAction.TestHealthWrite);
 			acl.AddDefaultWriteAction(BisectTaskAclAction.CreateBisectTask);
 			acl.AddDefaultWriteAction(BisectTaskAclAction.UpdateBisectTask);
 		}
diff --git a/Engine/Source/Programs/Horde/Plugins/Build/HordeServer.Build/BuildPlugin.cs b/Engine/Source/Programs/Horde/Plugins/Build/HordeServer.Build/BuildPlugin.cs
--- a/Engine/Source/Programs/Horde/Plugins/Build/HordeServer.Build/BuildPlugin.cs
+++ b/Engine/Source/Programs/Horde/Plugins/Build/HordeServer.Build/BuildPlugin.cs
@@ -116,6 +116,7 @@ namespace HordeServer
 			services.AddSingleton<IAuditLog<DeviceId>>(sp => sp.GetRequiredService<IAuditLogFactory<DeviceId>>().Create("Devices.Log", "DeviceId"));
 			services.AddSingleton<IBlockCache>(sp => CreateBlockCache());
 			services.AddSingleton<TestDataService>();
+			services.AddSingleton<IAuditLog<TestId>>(sp => sp.GetRequiredService<IAuditLogFactory<TestId>>().Create("TestHealthAudits.Log", "TestId"));
 
 			services.AddSingleton<IssueService>();
 			services.AddSingleton<JobService>();
@@ -149,6 +150,7 @@ namespace HordeServer
 				services.AddHostedService(provider => provider.GetRequiredService<DeviceService>());
 				services.AddHostedService<DeviceReportService>();
 				services.AddHostedService(provider => provider.GetRequiredService<TestDataService>());
+				services.AddHostedService<TestHealthService>();
 			}
 
 			services.AddHostedService(provider => provider.GetRequiredService<IExternalIssueService>());
diff --git a/Engine/Source/Programs/Horde/Plugins/Build/HordeServer.Build/BuildServerConfig.cs b/Engine/Source/Programs/Horde/Plugins/Build/HordeServer.Build/BuildServerConfig.cs
--- a/Engine/Source/Programs/Horde/Plugins/Build/HordeServer.Build/BuildServerConfig.cs
+++ b/Engine/Source/Programs/Horde/Plugins/Build/HordeServer.Build/BuildServerConfig.cs
@@ -140,6 +140,16 @@ namespace HordeServer
 		/// </summary>
 		public int TestDataRetainMonths { get; set; } = 6;
 
+		/// <summary>
+		/// The number of days to go back in historical data to evaluate test health
+		/// </summary>
+		public int TestHealthHistoryEvaluationDays { get; set; } = 14;
+
+		/// <summary>
+		/// The number of days after which a test health report remind for improvement
+		/// </summary>
+		public int TestHealthReportRemindDays { get; set; } = 14;
+
 		/// <summary>
 		/// Directory to store the fine-grained block cache. This caches individual exports embedded in bundles.
 		/// </summary>
diff --git a/Engine/Source/Programs/Horde/Plugins/Build/HordeServer.Build/Jobs/TestData/ITestData.cs b/Engine/Source/Programs/Horde/Plugins/Build/HordeServer.Build/Jobs/TestData/ITestData.cs
--- a/Engine/Source/Programs/Horde/Plugins/Build/HordeServer.Build/Jobs/TestData/ITestData.cs
+++ b/Engine/Source/Programs/Horde/Plugins/Build/HordeServer.Build/Jobs/TestData/ITestData.cs
@@ -436,6 +436,77 @@ namespace HordeServer.Jobs.TestData
 		IReadOnlyList<ITestParamEntry>? Params { get; }
 	}
 
+	/// <summary>
+	/// A test health report
+	/// </summary>
+	public interface ITestHealthReport
+	{
+		/// <summary>
+		/// The report id
+		/// </summary>
+		ObjectId Id { get; }
+
+		/// <summary>
+		/// The test id associated with this report
+		/// </summary>
+		TestId TestId { get; }
+
+		/// <summary>
+		/// The test name associated with this report
+		/// </summary>
+		string TestName { get; }
+
+		/// <summary>
+		/// The stream id associated with this report
+		/// </summary>
+		StreamId StreamId { get; }
+
+		/// <summary>
+		/// The last update date associated with this report
+		/// </summary>
+		DateTime LastUpdateDateUtc { get; }
+
+		/// <summary>
+		/// Whether or not the test is healthy
+		/// </summary>
+		bool IsHealthy { get; }
+
+		/// <summary>
+		/// The state associated with this test health
+		/// </summary>
+		string State { get; }
+
+		/// <summary>
+		/// The previous state associated with this test health
+		/// </summary>
+		string? PreviousState { get; }
+
+		/// <summary>
+		/// The success rate of the test health
+		/// </summary>
+		int SuccessRate { get; }
+
+		/// <summary>
+		/// The failure rate of the test health
+		/// </summary>
+		int FailureRate { get; }
+
+		/// <summary>
+		/// The catastrophic rate of the failure
+		/// </summary>
+		int CatastrophicFailureRate { get; }
+
+		/// <summary>
+		/// The rate of the redundant errors
+		/// </summary>
+		int RedundantErrorRate { get; }
+
+		/// <summary>
+		/// The date of the last notification sent for this report
+		/// </summary>
+		DateTime? NotificationLastDateUtc { get; }
+	}
+
 	/// <summary>
 	/// A test notification settings
 	/// </summary>
diff --git a/Engine/Source/Programs/Horde/Plugins/Build/HordeServer.Build/Jobs/TestData/ITestDataCollection.cs b/Engine/Source/Programs/Horde/Plugins/Build/HordeServer.Build/Jobs/TestData/ITestDataCollection.cs
--- a/Engine/Source/Programs/Horde/Plugins/Build/HordeServer.Build/Jobs/TestData/ITestDataCollection.cs
+++ b/Engine/Source/Programs/Horde/Plugins/Build/HordeServer.Build/Jobs/TestData/ITestDataCollection.cs
@@ -5,6 +5,7 @@ using EpicGames.Horde.Issues;
 using EpicGames.Horde.Jobs;
 using EpicGames.Horde.Streams;
 using EpicGames.Horde.Users;
+using HordeServer.Auditing;
 using MongoDB.Bson;
 
 namespace HordeServer.Jobs.TestData
@@ -270,6 +271,49 @@ namespace HordeServer.Jobs.TestData
 		/// <returns></returns>
 		Task<IReadOnlyList<ITestAudit>> FindTestAuditsAsync(TestId[] testIds, DateTime? minCreateTime = null, CancellationToken cancellationToken = default);
 
+		/// <summary>
+		/// Find audits related to tests with notification workflow
+		/// </summary>
+		/// <param name="cancellationToken"></param>
+		/// <returns></returns>
+		Task<IReadOnlyList<ITestAudit>> FindTestWithAuditNotificationWorkflowAsync(CancellationToken cancellationToken = default);
+
+		/// <summary>
+		/// Get the logger for the specified test id
+		/// </summary>
+		/// <param name="testId"></param>
+		/// <returns></returns>
+		IAuditLogChannel<TestId> GetTestAuditLogger(TestId testId);
+
+		/// <summary>
+		/// Add or update test health report
+		/// </summary>
+		/// <param name="testId"></param>
+		/// <param name="streamId"></param>
+		/// <param name="testName"></param>
+		/// <param name="testHealth"></param>
+		/// <param name="cancellationToken"></param>
+		/// <returns></returns>
+		Task<ITestHealthReport> AddOrUpdateTestHealthReportAsync(TestId testId, StreamId streamId, string testName, ComputedTestHealth testHealth, CancellationToken cancellationToken = default);
+
+		/// <summary>
+		/// Find Test health reports associated with test ids and stream ids
+		/// </summary>
+		/// <param name="streamIds"></param>
+		/// <param name="testIds"></param>
+		/// <param name="cancellationToken"></param>
+		/// <returns></returns>
+		Task<IReadOnlyList<ITestHealthReport>> FindTestHealthReportsAsync(StreamId[] streamIds, TestId[]? testIds, CancellationToken cancellationToken = default);
+
+		/// <summary>
+		/// Update test health report notification date
+		/// </summary>
+		/// <param name="reportId"></param>
+		/// <param name="resetDate"></param>
+		/// <param name="cancellationToken"></param>
+		/// <returns></returns>
+		Task UpdateTestHealthReportNotificationDateAsync(ObjectId reportId, bool resetDate, CancellationToken cancellationToken = default);
+
 		/// <summary>
 		/// Prune test data from collection
 		/// </summary>
@@ -362,4 +406,40 @@ namespace HordeServer.Jobs.TestData
 		/// </summary>
 		public string? Notes { get; set; }
 	}
+
+	/// <summary>
+	/// The computed health for a report
+	/// </summary>
+	public class ComputedTestHealth
+	{
+		/// <summary>
+		/// Whether the rates distribution is deemed healthy
+		/// </summary>
+		public bool IsHealthy { get; set; }
+
+		/// <summary>
+		/// A label that describes the health state
+		/// </summary>
+		public string State { get; set; } = String.Empty;
+
+		/// <summary>
+		/// The rate of success
+		/// </summary>
+		public int SuccessRate { get; set; }
+
+		/// <summary>
+		/// The rate of failure
+		/// </summary>
+		public int FailureRate { get; set; }
+
+		/// <summary>
+		/// The rate of catastrophic failure
+		/// </summary>
+		public int CatastrophicFailureRate { get; set; }
+
+		/// <summary>
+		/// The rate of redundant error
+		/// </summary>
+		public int RedundantErrorRate { get; set; }
+	}
 }
diff --git a/Engine/Source/Programs/Horde/Plugins/Build/HordeServer.Build/Jobs/TestData/TestDataAclAction.cs b/Engine/Source/Programs/Horde/Plugins/Build/HordeServer.Build/Jobs/TestData/TestDataAclAction.cs
new file mode 100644
--- /dev/null
+++ b/Engine/Source/Programs/Horde/Plugins/Build/HordeServer.Build/Jobs/TestData/TestDataAclAction.cs
@@ -0,0 +1,22 @@
+// Copyright Epic Games, Inc. All Rights Reserved.
+
+using EpicGames.Horde.Acls;
+
+namespace HordeServer.Jobs.TestData
+{
+	/// <summary>
+	/// ACL actions which apply to test data
+	/// </summary>
+	public static class TestDataAclAction
+	{
+		/// <summary>
+		/// Ability to read test health
+		/// </summary>
+		public static AclAction TestHealthRead { get; } = new AclAction("TestHealthRead");
+
+		/// <summary>
+		/// Ability to write test health
+		/// </summary>
+		public static AclAction TestHealthWrite { get; } = new AclAction("TestHealthWrite");
+	}
+}
diff --git a/Engine/Source/Programs/Horde/Plugins/Build/HordeServer.Build/Jobs/TestData/TestDataCollectionV2.cs b/Engine/Source/Programs/Horde/Plugins/Build/HordeServer.Build/Jobs/TestData/TestDataCollectionV2.cs
--- a/Engine/Source/Programs/Horde/Plugins/Build/HordeServer.Build/Jobs/TestData/TestDataCollectionV2.cs
+++ b/Engine/Source/Programs/Horde/Plugins/Build/HordeServer.Build/Jobs/TestData/TestDataCollectionV2.cs
@@ -7,6 +7,7 @@ using EpicGames.Horde.Jobs.Templates;
 using EpicGames.Horde.Jobs.TestData;
 using EpicGames.Horde.Streams;
 using EpicGames.Horde.Users;
+using HordeServer.Auditing;
 using HordeServer.Server;
 using HordeServer.Utilities;
 using Microsoft.Extensions.Logging;
@@ -508,6 +509,62 @@ namespace HordeServer.Jobs.TestData
 			}
 		}
 
+		/// <summary>
+		/// Store test health report info
+		/// </summary>
+		class TestHealthReportDocument : ITestHealthReport
+		{
+			[BsonRequired, BsonId]
+			public ObjectId Id { get; set; }
+
+			[BsonRequired, BsonElement("tid")]
+			public TestId TestId { get; set; }
+
+			[BsonRequired, BsonElement("tn")]
+			public string TestName { get; set; } = String.Empty;
+
+			[BsonRequired, BsonElement("sid")]
+			public StreamId StreamId { get; set; }
+
+			[BsonRequired, BsonElement("lud")]
+			public DateTime LastUpdateDateUtc { get; set; }
+
+			[BsonRequired, BsonElement("h")]
+			public bool IsHealthy { get; set; }
+
+			[BsonRequired, BsonElement("s")]
+			public string State { get; set; } = String.Empty;
+
+			[BsonIgnoreIfNull, BsonElement("ps")]
+			public string? PreviousState { get; set; }
+
+			[BsonRequired, BsonElement("sr")]
+			public int SuccessRate { get; set; }
+
+			[BsonRequired, BsonElement("fr")]
+			public int FailureRate { get; set; }
+
+			[BsonRequired, BsonElement("cr")]
+			public int CatastrophicFailureRate { get; set; }
+
+			[BsonRequired, BsonElement("rer")]
+			public int RedundantErrorRate { get; set; }
+
+			[BsonIgnoreIfNull, BsonElement("nld")]
+			public DateTime? NotificationLastDateUtc { get; set; }
+
+			private TestHealthReportDocument() { }
+
+			public TestHealthReportDocument(TestId testId, string testName, StreamId streamId)
+			{
+				Id = ObjectId.GenerateNewId();
+				TestId = testId;
+				TestName = testName;
+				StreamId = streamId;
+				LastUpdateDateUtc = DateTime.UtcNow;
+			}
+		}
+
 		/// <summary>
 		/// Store test notification settings for a stream
 		/// </summary>
@@ -673,16 +730,19 @@ namespace HordeServer.Jobs.TestData
 		readonly IMongoCollection<TestPhaseSessionDocument> _testPhaseSessions;
 		readonly IMongoCollection<TestSessionStreamDocument> _testStreams;
 		readonly IMongoCollection<TestAuditDocument> _testAudits;
+		readonly IMongoCollection<TestHealthReportDocument> _testHealthReports;
 		readonly Tracer _tracer;
 		readonly ILogger _logger;
+		readonly IAuditLog<TestId> _auditLog;
 
 		/// <summary>
 		/// Constructor
 		/// </summary>
-		public TestDataCollectionV2(IMongoService mongoService, Tracer tracer, ILogger<TestDataCollection> logger)
+		public TestDataCollectionV2(IMongoService mongoService, Tracer tracer, ILogger<TestDataCollection> logger, IAuditLog<TestId> auditLog)
 		{
 			_tracer = tracer;
 			_logger = logger;
+			_auditLog = auditLog;
 
 			List<MongoIndex<TestDataDocument>> indexes = new List<MongoIndex<TestDataDocument>>();
 			indexes.Add(keys => keys.Descending(x => x.JobId).Descending(x => x.StepId));
@@ -723,6 +783,10 @@ namespace HordeServer.Jobs.TestData
 			List<MongoIndex<TestAuditDocument>> auditIndexes = new List<MongoIndex<TestAuditDocument>>();
 			auditIndexes.Add(keys => keys.Ascending(x => x.TestId).Descending(x => x.Archived));
 			_testAudits = mongoService.GetCollection<TestAuditDocument>("TestDataV2.TestAudits", auditIndexes);
+
+			List<MongoIndex<TestHealthReportDocument>> healthReporIndexes = new List<MongoIndex<TestHealthReportDocument>>();
+			healthReporIndexes.Add(keys => keys.Ascending(x => x.StreamId).Ascending(x => x.TestId), unique: true);
+			_testHealthReports = mongoService.GetCollection<TestHealthReportDocument>("TestDataV2.TestHealthReports", healthReporIndexes);			
 		}
 
 		/// <inheritdoc/>
@@ -1503,6 +1567,95 @@ namespace HordeServer.Jobs.TestData
 			return results.ConvertAll<ITestAudit>(x => x);
 		}
 
+		/// <inheritdoc/>
+		public async Task<IReadOnlyList<ITestAudit>> FindTestWithAuditNotificationWorkflowAsync(CancellationToken cancellationToken = default)
+		{
+			FilterDefinition<TestAuditDocument> filter = FilterDefinition<TestAuditDocument>.Empty;
+			FilterDefinitionBuilder<TestAuditDocument> filterBuilder = Builders<TestAuditDocument>.Filter;
+			SortDefinition<TestAuditDocument> sortDefinition = Builders<TestAuditDocument>.Sort.Descending(x => x.Id);
+
+			filter &= filterBuilder.Eq(x => x.Archived, false);
+			filter &= filterBuilder.Ne(x => x.Notifications, null);
+
+			List<TestAuditDocument> results = await _testAudits.Find(filter).Sort(sortDefinition).ToListAsync(cancellationToken);
+
+			return results.ConvertAll<ITestAudit>(x => x);
+		}
+
+		/// <inheritdoc/>
+		public IAuditLogChannel<TestId> GetTestAuditLogger(TestId testId)
+		{
+			return _auditLog[testId];
+		}
+
+		/// <inheritdoc/>
+		public async Task<ITestHealthReport> AddOrUpdateTestHealthReportAsync(TestId testId, StreamId streamId, string testName, ComputedTestHealth testHealth, CancellationToken cancellationToken = default)
+		{
+			FilterDefinitionBuilder<TestHealthReportDocument> filterBuilder = Builders<TestHealthReportDocument>.Filter;
+			FilterDefinition<TestHealthReportDocument> filter = FilterDefinition<TestHealthReportDocument>.Empty;
+
+			filter &= filterBuilder.Eq(i => i.TestId, testId);
+			filter &= filterBuilder.Eq(i => i.StreamId, streamId);
+
+			TestHealthReportDocument? healthReport = await _testHealthReports.Find(filter).FirstOrDefaultAsync(cancellationToken);
+
+			if (healthReport == null)
+			{
+				healthReport = new TestHealthReportDocument(testId, testName, streamId);
+				healthReport.State = testHealth.State;
+				healthReport.IsHealthy = testHealth.IsHealthy;
+				healthReport.SuccessRate = testHealth.SuccessRate;
+				healthReport.FailureRate = testHealth.FailureRate;
+				healthReport.CatastrophicFailureRate = testHealth.CatastrophicFailureRate;
+				healthReport.RedundantErrorRate = testHealth.RedundantErrorRate;
+				await _testHealthReports.InsertOneAsync(healthReport, null, cancellationToken);
+			}
+			else
+			{
+				UpdateDefinitionBuilder<TestHealthReportDocument> updateBuilder = Builders<TestHealthReportDocument>.Update;
+				List<UpdateDefinition<TestHealthReportDocument>> updates = new List<UpdateDefinition<TestHealthReportDocument>>();
+				updates.Add(updateBuilder.Set(x => x.LastUpdateDateUtc, DateTime.UtcNow));
+				updates.Add(updateBuilder.Set(x => x.TestName, testName));
+				updates.Add(updateBuilder.Set(x => x.IsHealthy, testHealth.IsHealthy));
+				updates.Add(updateBuilder.Set(x => x.State, testHealth.State));
+				updates.Add(updateBuilder.Set(x => x.PreviousState, healthReport.State));
+				updates.Add(updateBuilder.Set(x => x.SuccessRate, testHealth.SuccessRate));
+				updates.Add(updateBuilder.Set(x => x.FailureRate, testHealth.FailureRate));
+				updates.Add(updateBuilder.Set(x => x.CatastrophicFailureRate, testHealth.CatastrophicFailureRate));
+				updates.Add(updateBuilder.Set(x => x.RedundantErrorRate, testHealth.RedundantErrorRate));
+				FindOneAndUpdateOptions<TestHealthReportDocument> options = new FindOneAndUpdateOptions<TestHealthReportDocument> { IsUpsert = true, ReturnDocument = ReturnDocument.After };
+				healthReport = await _testHealthReports.FindOneAndUpdateAsync(filterBuilder.Eq(x => x.Id, healthReport.Id), updateBuilder.Combine(updates), options, cancellationToken);
+			}
+
+			return healthReport;
+		}
+
+		/// <inheritdoc/>
+		public async Task<IReadOnlyList<ITestHealthReport>> FindTestHealthReportsAsync(StreamId[] streamIds, TestId[]? testIds, CancellationToken cancellationToken = default)
+		{
+			FilterDefinition<TestHealthReportDocument> filter = FilterDefinition<TestHealthReportDocument>.Empty;
+			FilterDefinitionBuilder<TestHealthReportDocument> filterBuilder = Builders<TestHealthReportDocument>.Filter;
+
+			filter &= filterBuilder.In(i => i.StreamId, streamIds);
+			if (testIds is not null)
+			{
+				filter &= filterBuilder.In(i => i.TestId, testIds);
+			}
+			
+			List<TestHealthReportDocument> healthReports = await _testHealthReports.Find(filter).ToListAsync(cancellationToken);
+
+			return healthReports;
+		}
+
+		/// <inheritdoc/>
+		public async Task UpdateTestHealthReportNotificationDateAsync(ObjectId reportId, bool resetDate, CancellationToken cancellationToken = default)
+		{
+			UpdateDefinitionBuilder<TestHealthReportDocument> updateBuilder = Builders<TestHealthReportDocument>.Update;
+			List<UpdateDefinition<TestHealthReportDocument>> updates = new List<UpdateDefinition<TestHealthReportDocument>>();
+			updates.Add(updateBuilder.Set(x => x.NotificationLastDateUtc, resetDate ? DateTime.UtcNow : null));
+			await _testHealthReports.FindOneAndUpdateAsync(x => x.Id == reportId, updateBuilder.Combine(updates), null, cancellationToken);
+		}
+
 		/// <inheritdoc/>
 		public async Task<bool> UpdateAsync(int retainMonths, CancellationToken cancellationToken)
 		{
@@ -1691,6 +1844,17 @@ namespace HordeServer.Jobs.TestData
 					_logger.LogInformation("Pruned {NumDeleted} test audit documents", result.DeletedCount);
 				}
 			}
+			// pruned all test health reports related to tests that were NOT kept
+			{
+				FilterDefinition<TestHealthReportDocument> filter = FilterDefinition<TestHealthReportDocument>.Empty;
+				FilterDefinitionBuilder<TestHealthReportDocument> filterBuilder = Builders<TestHealthReportDocument>.Filter;
+				filter &= filterBuilder.Nin(x => x.TestId, testIds);
+				DeleteResult result = await _testHealthReports.DeleteManyAsync(filter, cancellationToken);
+				if (result.DeletedCount > 0)
+				{
+					_logger.LogInformation("Pruned {NumDeleted} test health report documents", result.DeletedCount);
+				}
+			}
 		}
 	}
 }
diff --git a/Engine/Source/Programs/Horde/Plugins/Build/HordeServer.Build/Jobs/TestData/TestDataController.cs b/Engine/Source/Programs/Horde/Plugins/Build/HordeServer.Build/Jobs/TestData/TestDataController.cs
--- a/Engine/Source/Programs/Horde/Plugins/Build/HordeServer.Build/Jobs/TestData/TestDataController.cs
+++ b/Engine/Source/Programs/Horde/Plugins/Build/HordeServer.Build/Jobs/TestData/TestDataController.cs
@@ -8,6 +8,7 @@ using EpicGames.Horde.Jobs;
 using EpicGames.Horde.Jobs.TestData;
 using EpicGames.Horde.Streams;
 using EpicGames.Horde.Users;
+using HordeServer.Auditing;
 using HordeServer.Streams;
 using HordeServer.Utilities;
 using Microsoft.AspNetCore.Authorization;
@@ -923,6 +924,11 @@ namespace HordeServer.Jobs.TestData
 		[ProducesResponseType(typeof(GetTestAuditResponse), 200)]
 		public async Task<ActionResult<GetTestAuditResponse>> AddTestAuditV2Async([FromBody] AddTestAuditRequest request)
 		{
+			if (!_buildConfig.Value.Authorize(TestDataAclAction.TestHealthWrite, User))
+			{
+				return Forbid();
+			}
+
 			GetTestAuditResponse responses = new GetTestAuditResponse();
 
 			UserId? userId = User.GetUserId();
@@ -974,6 +980,11 @@ namespace HordeServer.Jobs.TestData
 			[FromBody] GetTestAuditRequest request,
 			[FromQuery] DateTimeOffset? minCreateTime = null)
 		{
+			if (!_buildConfig.Value.Authorize(TestDataAclAction.TestHealthRead, User))
+			{
+				return Forbid();
+			}
+
 			TestId[] testIds = request.TestIds.Select(x => TestId.Parse(x)).ToHashSet().ToArray();
 
 			IReadOnlyList<ITestAudit> audits = await _testDataService.FindTestAuditsV2Async(testIds, minCreateTime?.UtcDateTime);
@@ -1003,5 +1014,92 @@ namespace HordeServer.Jobs.TestData
 				Notes = audit.Notes?.ToString()
 			};
 		}
+
+		/// <summary>
+		/// Retrieve historical information about a specific test audit
+		/// </summary>
+		/// <param name="testId"></param>
+		/// <param name="minTime">Minimum time for records to return</param>
+		/// <param name="maxTime">Maximum time for records to return</param>
+		/// <param name="index">Offset of the first result</param>
+		/// <param name="count">Number of records to return</param>
+		/// <returns>Information about the requested test audit</returns>
+		[HttpGet]
+		[Route("/api/v3/testdata/audit/log/{testId}")]
+		public async Task<ActionResult> GetAuditLogHistoryAsync(string testId, [FromQuery] DateTime? minTime = null, [FromQuery] DateTime? maxTime = null, [FromQuery] int index = 0, [FromQuery] int count = 50)
+		{
+			if (!_buildConfig.Value.Authorize(TestDataAclAction.TestHealthRead, User))
+			{
+				return Forbid();
+			}
+
+			TestId test;
+			if (String.IsNullOrEmpty(testId) || !TestId.TryParse(testId, out test))
+			{
+				return BadRequest("Invalid test id");
+			}
+
+			Response.ContentType = "application/json";
+			Response.StatusCode = 200;
+			await Response.StartAsync();
+			await _testDataService.GetTestAuditLogger(test).FindAsync(Response.BodyWriter, minTime, maxTime, index, count);
+			await Response.CompleteAsync();
+
+			return Empty;
+		}
+
+		/// <summary>
+		/// Get test health reports
+		/// </summary>
+		/// <param name="request"></param>
+		/// <returns></returns>
+		[HttpPost]
+		[Route("/api/v3/testdata/health/reports")]
+		[ProducesResponseType(typeof(List<GetTestHealthReportResponse>), 200)]
+		public async Task<ActionResult<List<GetTestHealthReportResponse>>> GetTestHealthReportsV2Async([FromBody] GetTestHealthReportRequest request)
+		{
+			if (!_buildConfig.Value.Authorize(TestDataAclAction.TestHealthRead, User))
+			{
+				return Forbid();
+			}
+
+			StreamId[] streamIdValues = request.StreamIds.Select(x => new StreamId(x)).ToArray();
+
+			List<StreamId> queryStreams = new List<StreamId>();
+			List<GetTestHealthReportResponse> responses = new List<GetTestHealthReportResponse>();
+
+			// authorize streams
+			foreach (StreamId streamId in streamIdValues)
+			{
+				if (_buildConfig.Value.TryGetStream(streamId, out StreamConfig? streamConfig) && streamConfig.Authorize(JobAclAction.ViewJob, User))
+				{
+					queryStreams.Add(streamId);
+				}
+			}
+
+			if (queryStreams.Count == 0)
+			{
+				return responses;
+			}
+
+			TestId[]? queryTestIds = request.TestIds?.Select(x => TestId.Parse(x)).ToArray();
+
+			IReadOnlyList<ITestHealthReport> reports = await _testDataService.FindTestHealthReportsV2Async(queryStreams.ToArray(), queryTestIds);
+
+			return reports.Select(report => new GetTestHealthReportResponse()
+			{
+				Id = report.Id.ToString(),
+				TestId = report.TestId.ToString(),
+				StreamId = report.StreamId.ToString(),
+				LastUpdateDate = report.LastUpdateDateUtc,
+				IsHealthy = report.IsHealthy,
+				State = report.State,
+				SuccessRate = report.SuccessRate,
+				FailureRate	= report.FailureRate,
+				CatastrophicFailureRate = report.CatastrophicFailureRate,
+				RedundantErrorRate = report.RedundantErrorRate,
+				NotificationLastDate = report.NotificationLastDateUtc
+			}).ToList();
+		}
 	}
 }
diff --git a/Engine/Source/Programs/Horde/Plugins/Build/HordeServer.Build/Jobs/TestData/TestDataService.cs b/Engine/Source/Programs/Horde/Plugins/Build/HordeServer.Build/Jobs/TestData/TestDataService.cs
--- a/Engine/Source/Programs/Horde/Plugins/Build/HordeServer.Build/Jobs/TestData/TestDataService.cs
+++ b/Engine/Source/Programs/Horde/Plugins/Build/HordeServer.Build/Jobs/TestData/TestDataService.cs
@@ -5,6 +5,7 @@ using EpicGames.Horde.Commits;
 using EpicGames.Horde.Jobs;
 using EpicGames.Horde.Streams;
 using EpicGames.Horde.Users;
+using HordeServer.Auditing;
 using Microsoft.Extensions.Hosting;
 using Microsoft.Extensions.Logging;
 using Microsoft.Extensions.Options;
@@ -328,5 +329,53 @@ namespace HordeServer.Jobs.TestData
 		{
 			return await _testDataV2.FindTestAuditsAsync(testIds, minCreateTime, cancellationToken);
 		}
+
+		/// <summary>
+		/// Add or update test health report
+		/// </summary>
+		/// <param name="testId"></param>
+		/// <param name="streamId"></param>
+		/// <param name="testName"></param>
+		/// <param name="testHealth"></param>
+		/// <param name="cancellationToken"></param>
+		/// <returns></returns>
+		public async Task<ITestHealthReport> AddOrUpdateTestHealthReportV2Async(TestId testId, StreamId streamId, string testName, ComputedTestHealth testHealth, CancellationToken cancellationToken = default)
+		{
+			return await _testDataV2.AddOrUpdateTestHealthReportAsync(testId, streamId, testName, testHealth, cancellationToken);
+		}
+
+		/// <summary>
+		/// Get the logger for the specified test id
+		/// </summary>
+		/// <param name="testId"></param>
+		/// <returns></returns>
+		public IAuditLogChannel<TestId> GetTestAuditLogger(TestId testId)
+		{
+			return _testDataV2.GetTestAuditLogger(testId);
+		}
+
+		/// <summary>
+		/// Find Test health report associated with stream and test ids
+		/// </summary>
+		/// <param name="streamIds"></param>
+		/// <param name="testIds"></param>
+		/// <param name="cancellationToken"></param>
+		/// <returns></returns>
+		public async Task<IReadOnlyList<ITestHealthReport>> FindTestHealthReportsV2Async(StreamId[] streamIds, TestId[]? testIds, CancellationToken cancellationToken = default)
+		{
+			return await _testDataV2.FindTestHealthReportsAsync(streamIds, testIds, cancellationToken);
+		}
+
+		/// <summary>
+		/// Update test health report notification information
+		/// </summary>
+		/// <param name="reportId"></param>
+		/// <param name="resetDate"></param>
+		/// <param name="cancellationToken"></param>
+		/// <returns></returns>
+		public async Task UpdateTestHealthReportNotificationDateAsync(ObjectId reportId, bool resetDate, CancellationToken cancellationToken = default)
+		{
+			await _testDataV2.UpdateTestHealthReportNotificationDateAsync(reportId, resetDate, cancellationToken);
+		}
 	}
 }
\ No newline at end of file
diff --git a/Engine/Source/Programs/Horde/Plugins/Build/HordeServer.Build/Jobs/TestData/TestHealthService.cs b/Engine/Source/Programs/Horde/Plugins/Build/HordeServer.Build/Jobs/TestData/TestHealthService.cs
new file mode 100644
--- /dev/null
+++ b/Engine/Source/Programs/Horde/Plugins/Build/HordeServer.Build/Jobs/TestData/TestHealthService.cs
@@ -0,0 +1,324 @@
+// Copyright Epic Games, Inc. All Rights Reserved.
+
+using EpicGames.Core;
+using EpicGames.Horde.Streams;
+using HordeServer.Auditing;
+using HordeServer.Issues;
+using HordeServer.Notifications;
+using HordeServer.Server;
+using HordeServer.Streams;
+using HordeServer.Utilities;
+using Microsoft.Extensions.Hosting;
+using Microsoft.Extensions.Logging;
+using Microsoft.Extensions.Options;
+using OpenTelemetry.Trace;
+
+namespace HordeServer.Jobs.TestData
+{
+
+	[SingletonDocument("test-health-report-state")]
+	class TestHealthReportState : SingletonBase
+	{
+		public DateTime ReportTime { get; set; } = DateTime.MinValue;
+	}
+
+	enum TestHealthState
+	{
+		Healthy = 5,
+		Fair = 4,
+		Unstable = 3,
+		Bad = 2,
+		Broken = 1,
+	}
+
+	/// <summary>
+	/// Test Health Service
+	/// </summary>
+	public sealed class TestHealthService : IHostedService, IAsyncDisposable
+	{
+		readonly ILogger<TestHealthService> _logger;
+		readonly Tracer _tracer;
+		readonly ITicker _ticker;
+		readonly IClock _clock;
+		readonly ITestDataCollectionV2 _testData;
+		readonly SingletonDocument<TestHealthReportState> _state;
+		readonly IOptions<BuildServerConfig> _staticBuildConfig;
+		readonly IOptionsSnapshot<BuildConfig> _buildConfig;
+		readonly INotificationService _notificationService;
+
+		readonly int _reportIntervalMinutes = 12 * 60; // every 12 hours
+		private bool _forceReport = false;
+
+		/// <summary>
+		/// Test Health Service constructor
+		/// </summary>
+		/// <param name="mongoService"></param>
+		/// <param name="testData"></param>
+		/// <param name="notificationService"></param>
+		/// <param name="clock"></param>
+		/// <param name="logger"></param>
+		/// <param name="tracer"></param>
+		/// <param name="staticBuildConfig"></param>
+		/// <param name="buildConfig"></param>
+		public TestHealthService(IMongoService mongoService, ITestDataCollectionV2 testData, INotificationService notificationService, IClock clock, ILogger<TestHealthService> logger, Tracer tracer, IOptions<BuildServerConfig> staticBuildConfig, IOptionsSnapshot<BuildConfig> buildConfig)
+		{
+			_state = new SingletonDocument<TestHealthReportState>(mongoService);
+			_testData = testData;
+			_clock = clock;
+			_logger = logger;
+			_tracer = tracer;
+			_ticker = clock.AddSharedTicker<TestHealthService>(TimeSpan.FromMinutes(10.0), TickAsync, logger);
+			_staticBuildConfig = staticBuildConfig;
+			_buildConfig = buildConfig;
+			_notificationService = notificationService;
+		}
+
+		/// <inheritdoc/>
+		public async Task StartAsync(CancellationToken cancellationToken)
+		{
+			await _ticker.StartAsync();
+		}
+
+		/// <inheritdoc/>
+		public async Task StopAsync(CancellationToken cancellationToken)
+		{
+			await _ticker.StopAsync();
+		}
+
+		/// <inheritdoc/>
+		public async ValueTask DisposeAsync()
+		{
+			await _ticker.DisposeAsync();
+		}
+
+		internal async Task TickForTestingAsync(bool bForceReport = false)
+		{
+			_forceReport = bForceReport;
+			await TickAsync(CancellationToken.None);
+		}
+
+		/// <summary>
+		/// Ticks service
+		/// </summary>
+		async ValueTask TickAsync(CancellationToken cancellationToken)
+		{
+			using TelemetrySpan span = _tracer.StartActiveSpan($"{nameof(TestHealthService)}.{nameof(TickAsync)}");
+
+			TestHealthReportState state = await _state.GetAsync(cancellationToken);
+			DateTime currentTime = _clock.UtcNow;
+			DateTime lastReportTime = state.ReportTime == DateTime.MinValue ? currentTime.Subtract(TimeSpan.FromMinutes(_reportIntervalMinutes + 1)) : state.ReportTime;
+
+			if (!_forceReport && (currentTime - lastReportTime).TotalMinutes <= _reportIntervalMinutes)
+			{
+				return;
+			}
+
+			IReadOnlyList<ITestAudit> audits = await _testData.FindTestWithAuditNotificationWorkflowAsync(cancellationToken);
+			if (audits.Any())
+			{
+				_logger.LogInformation("Generating test health reports");
+
+				audits = audits.DistinctBy(a => a.TestId).ToList();
+				span.SetAttribute("horde.test_health.audit_workflow_count", audits.Count);
+
+				IReadOnlyDictionary<TestId, ITestAudit> testAuditByTestId = audits.ToDictionary(a => a.TestId);
+
+				IReadOnlyList<ITestNameRef> testNameRefs = await _testData.FindTestNameRefsAsync(testAuditByTestId.Keys.ToArray(), cancellationToken: cancellationToken);
+				IReadOnlyDictionary<TestId, ITestNameRef> nameRefByTestId = testNameRefs.ToDictionary(t => t.Id);
+
+				IReadOnlyDictionary<StreamId, StreamConfig> configByStreamId = _buildConfig.Value.Projects.SelectMany(p => p.Streams).ToDictionary(s => s.Id);
+
+				int historyEvaluationDays = _staticBuildConfig.Value.TestHealthHistoryEvaluationDays;
+				DateTime minCreateTime = currentTime.Subtract(TimeSpan.FromDays(historyEvaluationDays));
+
+				int reportRemindDays = _staticBuildConfig.Value.TestHealthReportRemindDays;
+				DateTime reportRemindTime = currentTime.Subtract(TimeSpan.FromDays(reportRemindDays));
+
+				StreamId[] streamIds = audits.SelectMany(a => a.Notifications!.Select(n => n.Stream)).Distinct().ToArray();
+				span.SetAttribute("horde.test_health.audit_stream_count", streamIds.Length);
+				foreach (StreamId streamId in streamIds)
+				{
+					StreamConfig? streamConfig = configByStreamId.GetValueOrDefault(streamId);
+
+					if (streamConfig is null)
+					{
+						_logger.LogWarning("Skipping Test Health reporting for stream '{StreamId}'; reason: No stream config found for this stream", streamId);
+						continue;
+					}
+
+					TestId[] testIds = audits.Where(a => a.Notifications!.Any(n => n.Stream == streamId)).Select(a => a.TestId).ToArray();
+					IReadOnlyList<ITestSession> sessions = await _testData.FindTestSessionsAsync([streamId], testIds, minCreateTime: minCreateTime, cancellationToken: cancellationToken);
+					IReadOnlyList<ITestHealthReport> previousReports = await _testData.FindTestHealthReportsAsync([streamId], testIds, cancellationToken);
+					IReadOnlyDictionary<TestId, ITestHealthReport> previousRepostsByTestId = previousReports.ToDictionary(r => r.TestId);
+
+					// Re-organize data
+					IReadOnlyDictionary<TestId, List<ITestSession>> sessionByTestId = testIds.ToDictionary(t => t, t => new List<ITestSession>());
+					foreach (ITestSession session in sessions)
+					{
+						sessionByTestId.GetValueOrDefault(session.NameRef)!.Add(session);
+					}
+
+					// Apply Health Heuristic to each test
+					foreach ((TestId testId, List<ITestSession> history) in sessionByTestId)
+					{
+						IAuditLogChannel<TestId> auditLog = _testData.GetTestAuditLogger(testId);
+						ITestNameRef? nameRef = nameRefByTestId.GetValueOrDefault(testId);
+						if (nameRef is null)
+						{
+							auditLog.LogError("No name reference found for test {TestId}", testId);
+							continue;
+						}
+
+						if (history.Count == 0)
+						{
+							auditLog.LogInformation("No session history found for stream {Stream} in the last {Days} days", streamConfig.Name, historyEvaluationDays);
+							continue;
+						}
+
+						ITestAudit testAudit = testAuditByTestId.GetValueOrDefault(testId)!;
+
+						ITestNotificationSettings? notificationSettings = testAudit.Notifications?.FirstOrDefault(n => n.Stream == streamId);
+						if (notificationSettings is null)
+						{
+							auditLog.LogError("No notification settings found for stream {Stream}", streamConfig.Name);
+							continue;
+						}
+
+						WorkflowConfig? workflowConfig = streamConfig.Workflows.FirstOrDefault(w => w.Id == notificationSettings.Workflow);
+						if (workflowConfig?.ReportChannel is null)
+						{
+							auditLog.LogError("No workflow config or report channel found for '{WorkflowId}' in stream {Stream}", notificationSettings.Workflow, streamConfig.Name);
+							continue;
+						}
+
+						ITestHealthReport? previousReport = previousRepostsByTestId.GetValueOrDefault(testId);
+						if (previousReport is not null && history.First().Id.CreationTime.AddMinutes(1) < previousReport.LastUpdateDateUtc)
+						{
+							// Nothing new to compute
+							continue;
+						}
+
+						ComputedTestHealth? computedHealth = ComputeHealthHeuristic(history);
+						if (computedHealth is not null)
+						{
+							ITestHealthReport report = await _testData.AddOrUpdateTestHealthReportAsync(testId, streamId, nameRef.Name, computedHealth, cancellationToken);
+
+							if (testAudit.IsUnderAudit ?? false)
+							{
+								// No notification if under audit
+								continue;
+							}
+
+							bool hasOpenNotificationThread = report.NotificationLastDateUtc is not null;
+							if ((report.IsHealthy && !hasOpenNotificationThread) /* no open notification thread */
+								|| (!report.IsHealthy && hasOpenNotificationThread && report.PreviousState == report.State && report.NotificationLastDateUtc > reportRemindTime)) /* no state change under remind time */
+							{
+								continue;
+							}
+
+							string[]? userIds = GetAuditAttachedUsers(testAudit);
+							await _notificationService.NotifyTestHealthReportAsync(report, workflowConfig.ReportChannel, userIds, cancellationToken);
+							// Update notification date or remove it if healthy
+							await _testData.UpdateTestHealthReportNotificationDateAsync(report.Id, !report.IsHealthy, cancellationToken);
+						}
+					}
+				}
+			}
+
+			await _state.UpdateAsync(s => s.ReportTime = currentTime, cancellationToken);
+			_forceReport = false;
+		}
+
+		/// <summary>
+		/// Compute Health Heuristic
+		/// </summary>
+		/// <param name="history"></param>
+		/// <returns></returns>
+		private static ComputedTestHealth? ComputeHealthHeuristic(List<ITestSession> history)
+		{
+			double totalPhaseCount = 0;
+			double totalUndefinedCount = 0;
+			double totalFailureCount = 0;
+			double totalSuccessCount = 0;
+			HashSet<string> errorFingerprints = new HashSet<string>();
+
+			foreach (ITestSession session in history)
+			{
+				totalPhaseCount += session.PhasesSucceededCount + session.PhasesUndefinedCount + session.PhasesFailedCount;
+				totalFailureCount += session.PhasesUndefinedCount + session.PhasesFailedCount;
+				totalUndefinedCount += session.PhasesUndefinedCount;
+				totalSuccessCount += session.PhasesSucceededCount;
+
+				session.ErrorFingerprints?.ToList().ForEach(f => errorFingerprints.Add(f.Key));
+			}
+
+			if (totalPhaseCount == 0)
+			{
+				return null;
+			}
+
+			int catastrophicFailureRate = totalFailureCount == 0 ? 0 : (int)Math.Round(totalUndefinedCount / totalFailureCount * 100);
+			int failureRate = (int)Math.Round(totalFailureCount / totalPhaseCount * 100);
+			int successRate = (int)Math.Round(totalSuccessCount / totalPhaseCount * 100);
+			int redundantErrorRate = totalFailureCount == 0 ? 0 : (int)Math.Round((1 - (double)errorFingerprints.Count / totalFailureCount) * 100);
+			int nStars = (int)TestHealthState.Fair;
+			if (failureRate > 15)
+			{
+				if (catastrophicFailureRate > 40 || successRate == 0)
+				{
+					nStars -= 2;
+				}
+				else
+				{
+					nStars -= 1;
+				}
+				if (redundantErrorRate < 80)
+				{
+					nStars += redundantErrorRate < 60 ? 1 : 0;
+				}
+				else if (failureRate > 40)
+				{
+					nStars -= 1;
+				}
+			}
+			if (successRate > 60)
+			{
+				nStars += 1;
+			}
+
+			bool isHealthy = nStars >= (int)TestHealthState.Fair;
+			string state = Enum.GetName(typeof(TestHealthState), nStars) ?? "Unknown";
+
+			return new ComputedTestHealth()
+			{
+				IsHealthy = isHealthy,
+				State = state,
+				SuccessRate = successRate,
+				FailureRate = failureRate,
+				CatastrophicFailureRate = catastrophicFailureRate,
+				RedundantErrorRate = redundantErrorRate
+			};
+		}
+
+		private static string[]? GetAuditAttachedUsers(ITestAudit testAudit)
+		{
+			List<string>? userIds = null;
+			bool hasOwner = testAudit.Owner is not null;
+			bool hasCustomers = testAudit.Customers?.Any() ?? false;
+			if (hasOwner || hasCustomers)
+			{
+				userIds = new List<string>();
+				if (hasOwner)
+				{
+					userIds.Add(testAudit.Owner!.Value.ToString());
+				}
+				if (hasCustomers)
+				{
+					userIds.AddRange(testAudit.Customers!.Select(c => c.ToString()));
+				}
+			}
+
+			return userIds?.ToArray();
+		}
+	}
+}
diff --git a/Engine/Source/Programs/Horde/Plugins/Build/HordeServer.Build/Jobs/TestData/TestSessionId.cs b/Engine/Source/Programs/Horde/Plugins/Build/HordeServer.Build/Jobs/TestData/TestSessionId.cs
--- a/Engine/Source/Programs/Horde/Plugins/Build/HordeServer.Build/Jobs/TestData/TestSessionId.cs
+++ b/Engine/Source/Programs/Horde/Plugins/Build/HordeServer.Build/Jobs/TestData/TestSessionId.cs
@@ -49,6 +49,11 @@ namespace HordeServer.Jobs.TestData
 
 		/// <inheritdoc/>
 		public override string ToString() => Id.ToString();
+
+		/// <summary>
+		/// The creation time of the id (derived from timestamp)
+		/// </summary>
+		public DateTime CreationTime => Id.CreationTime;
 	}
 
 	/// <summary>
diff --git a/Engine/Source/Programs/Horde/Plugins/Build/HordeServer.Build/Notifications/INotificationService.cs b/Engine/Source/Programs/Horde/Plugins/Build/HordeServer.Build/Notifications/INotificationService.cs
--- a/Engine/Source/Programs/Horde/Plugins/Build/HordeServer.Build/Notifications/INotificationService.cs
+++ b/Engine/Source/Programs/Horde/Plugins/Build/HordeServer.Build/Notifications/INotificationService.cs
@@ -9,6 +9,7 @@ using HordeServer.Agents;
 using HordeServer.Configuration;
 using HordeServer.Devices;
 using HordeServer.Issues;
+using HordeServer.Jobs.TestData;
 using HordeServer.Streams;
 using HordeServer.Users;
 
@@ -193,5 +194,15 @@ namespace HordeServer.Notifications
 		/// <param name="report">The report data to send</param>
 		/// <param name="cancellationToken">Cancellation token for the operation</param>
 		Task SendIssueReportAsync(IssueReportGroup report, CancellationToken cancellationToken);
+
+		/// <summary>
+		/// Post a notification for test health report
+		/// </summary>
+		/// <param name="report"></param>
+		/// <param name="recipient"></param>
+		/// <param name="carbonCopies"></param>
+		/// <param name="cancellationToken"></param>
+		/// <returns></returns>
+		Task NotifyTestHealthReportAsync(ITestHealthReport report, string recipient, string[]? carbonCopies, CancellationToken cancellationToken);
 	}
 }
diff --git a/Engine/Source/Programs/Horde/Plugins/Build/HordeServer.Build/Notifications/INotificationSink.cs b/Engine/Source/Programs/Horde/Plugins/Build/HordeServer.Build/Notifications/INotificationSink.cs
--- a/Engine/Source/Programs/Horde/Plugins/Build/HordeServer.Build/Notifications/INotificationSink.cs
+++ b/Engine/Source/Programs/Horde/Plugins/Build/HordeServer.Build/Notifications/INotificationSink.cs
@@ -6,6 +6,7 @@ using HordeServer.Agents;
 using HordeServer.Configuration;
 using HordeServer.Devices;
 using HordeServer.Issues;
+using HordeServer.Jobs.TestData;
 using HordeServer.Logs;
 using HordeServer.Streams;
 using HordeServer.Users;
@@ -153,6 +154,16 @@ namespace HordeServer.Notifications
 		/// Notify the status of agents in the farm
 		/// </summary>
 		Task SendAgentReportAsync(AgentReport report, CancellationToken cancellationToken);
+
+		/// <summary>
+		/// Post a notification for test health report
+		/// </summary>
+		/// <param name="report"></param>
+		/// <param name="recipient"></param>
+		/// <param name="carbonCopies"></param>
+		/// <param name="cancellationToken"></param>
+		/// <returns></returns>
+		Task NotifyTestHealthReportAsync(ITestHealthReport report, string recipient, string[]? carbonCopies, CancellationToken cancellationToken);
 	}
 }
 
diff --git a/Engine/Source/Programs/Horde/Plugins/Build/HordeServer.Build/Notifications/NotificationService.cs b/Engine/Source/Programs/Horde/Plugins/Build/HordeServer.Build/Notifications/NotificationService.cs
--- a/Engine/Source/Programs/Horde/Plugins/Build/HordeServer.Build/Notifications/NotificationService.cs
+++ b/Engine/Source/Programs/Horde/Plugins/Build/HordeServer.Build/Notifications/NotificationService.cs
@@ -15,6 +15,7 @@ using HordeServer.Devices;
 using HordeServer.Issues;
 using HordeServer.Jobs;
 using HordeServer.Jobs.Graphs;
+using HordeServer.Jobs.TestData;
 using HordeServer.Logs;
 using HordeServer.Server;
 using HordeServer.Streams;
@@ -355,6 +356,22 @@ namespace HordeServer.Notifications
 			}
 		}
 
+		/// <inheritdoc/>
+		public async Task NotifyTestHealthReportAsync(ITestHealthReport report, string recipient, string[]? carbonCopies, CancellationToken cancellationToken)
+		{
+			foreach (INotificationSink sink in _sinks)
+			{
+				try
+				{
+					await sink.NotifyTestHealthReportAsync(report, recipient, carbonCopies, cancellationToken);
+				}
+				catch (Exception e)
+				{
+					_logger.LogError(e, "Failed sending test health report to {Recipient}", recipient);
+				}
+			}
+		}
+
 		/// <summary>
 		/// Enqueues an async task
 		/// </summary>
diff --git a/Engine/Source/Programs/Horde/Plugins/Build/HordeServer.Build/Notifications/Sinks/SlackNotificationSink.cs b/Engine/Source/Programs/Horde/Plugins/Build/HordeServer.Build/Notifications/Sinks/SlackNotificationSink.cs
--- a/Engine/Source/Programs/Horde/Plugins/Build/HordeServer.Build/Notifications/Sinks/SlackNotificationSink.cs
+++ b/Engine/Source/Programs/Horde/Plugins/Build/HordeServer.Build/Notifications/Sinks/SlackNotificationSink.cs
@@ -30,12 +30,13 @@ using HordeServer.Configuration;
 using HordeServer.Devices;
 using HordeServer.Issues;
 using HordeServer.Issues.External;
+using HordeServer.Jobs.TestData;
 using HordeServer.Logs;
-using HordeServer.VersionControl.Perforce;
 using HordeServer.Server;
 using HordeServer.Streams;
 using HordeServer.Users;
 using HordeServer.Utilities;
+using HordeServer.VersionControl.Perforce;
 using Microsoft.AspNetCore.Hosting;
 using Microsoft.Extensions.Caching.Memory;
 using Microsoft.Extensions.Hosting;
@@ -2698,6 +2699,94 @@ namespace HordeServer.Notifications.Sinks
 
 		#endregion
 
+		#region test health notification
+
+		/// <inheritdoc/>
+		public async Task NotifyTestHealthReportAsync(ITestHealthReport report, string recipient, string[]? carbonCopies, CancellationToken cancellationToken)
+		{
+			string eventId = GetTestHealthEventId(report);
+			if (report.IsHealthy)
+			{
+				// make sure there is a thread to update if test is now healthy
+				MessageStateDocument? prevState = await GetMessageStateAsync(recipient, eventId, cancellationToken);
+				if (prevState is null)
+				{
+					return;
+				}
+			}
+
+			BuildConfig buildConfig = _buildConfig.CurrentValue;
+			StreamConfig? streamConfig;
+			if (!buildConfig.TryGetStream(report.StreamId, out streamConfig))
+			{
+				return;
+			}
+
+			List<string> ccUsers = new List<string>();
+			if (carbonCopies?.Any() ?? false)
+			{
+				foreach (string user in carbonCopies)
+				{
+					UserId userId;
+					if (UserId.TryParse(user, out userId))
+					{
+						ccUsers.Add(await FormatMentionAsync(userId, true, cancellationToken));
+					}
+				}
+			}
+
+			SlackMessage message = new SlackMessage();
+			string testHealthDetailsUrlLink = new Uri($"{_serverInfo.DashboardUrl}/test-automation?stream={report.StreamId}&view=1&health={report.TestId}").ToString();
+
+			string title = GenerateTestHealthSlackMessageTitle(report, streamConfig, testHealthDetailsUrlLink);
+			message.Blocks.Add(new SectionBlock(title));
+			if (ccUsers.Any())
+			{
+				message.Blocks.Add(new SectionBlock(String.Join(" ", ccUsers)));
+			}
+
+			(MessageStateDocument state, bool isNew) = await SendOrUpdateMessageAsync(recipient, eventId, null, message, cancellationToken);
+			SlackMessageId threadId = state.MessageId;
+			if (!isNew)
+			{
+				if (report.IsHealthy)
+				{
+					// mark fixed
+					SlackMessage fixedMessage = new SlackMessage();
+					fixedMessage.Text = $"Test Health was restored to {report.State}.";
+					await SendOrUpdateMessageToThreadAsync(recipient, eventId, null, threadId, fixedMessage, cancellationToken);
+
+					// delete message state
+					await DeleteMessageStateAsync(recipient, eventId, cancellationToken);
+				}
+				else
+				{
+					// ping
+					int daysSinceLastUpdate = (int)(DateTime.UtcNow - report.LastUpdateDateUtc).TotalDays;
+					SlackMessage pingMessage = new SlackMessage();
+					pingMessage.Text = report.State == report.PreviousState ?
+							$"No health improvement has been recorded for {daysSinceLastUpdate} days. Consider disabling the Test or set it under Audit." : $"Test health has changed from {report.PreviousState} to {report.State}.";
+					await SendOrUpdateMessageToThreadAsync(recipient, eventId, null, threadId, pingMessage, cancellationToken);
+				}
+			}
+		}
+
+		static string GenerateTestHealthSlackMessageTitle(ITestHealthReport report, StreamConfig streamConfig, string testHealthDetailsUrlLink)
+		{
+			string state = report.PreviousState is not null && report.IsHealthy ? report.PreviousState : report.State;
+			string title = $"*[{streamConfig.Name}]* {RedBadge} <{testHealthDetailsUrlLink}|test *{report.TestName}* health> has degraded to *{state}*";
+			if (report.IsHealthy)
+			{
+				// strike title
+				title = $"~{title}~";
+			}
+			return title;
+		}
+
+		static string GetTestHealthEventId(ITestHealthReport report) => $"test-health-{report.Id}";
+
+		#endregion
+
 		static string QuoteText(string text, int maxLength = MaxLineLength)
 		{
 			maxLength -= 6;
diff --git a/Engine/Source/Programs/Horde/Plugins/Experimental/HordeServer.Experimental.Tests/Notifications/ExperimentalNotificationTests.cs b/Engine/Source/Programs/Horde/Plugins/Experimental/HordeServer.Experimental.Tests/Notifications/ExperimentalNotificationTests.cs
--- a/Engine/Source/Programs/Horde/Plugins/Experimental/HordeServer.Experimental.Tests/Notifications/ExperimentalNotificationTests.cs
+++ b/Engine/Source/Programs/Horde/Plugins/Experimental/HordeServer.Experimental.Tests/Notifications/ExperimentalNotificationTests.cs
@@ -20,11 +20,12 @@ using HordeServer.Devices;
 using HordeServer.Experimental.Notifications;
 using HordeServer.Issues;
 using HordeServer.Jobs;
+using HordeServer.Jobs.TestData;
 using HordeServer.Logs;
 using HordeServer.Notifications;
+using HordeServer.Notifications.Sinks;
 using HordeServer.Projects;
 using HordeServer.Server;
-using HordeServer.Notifications.Sinks;
 using HordeServer.Streams;
 using HordeServer.Tests;
 using HordeServer.Users;
@@ -127,6 +128,7 @@ namespace HordeServer.Experimental.Tests.Notifications
 		public Task SendAgentReportAsync(AgentReport report, CancellationToken cancellationToken) => throw new NotImplementedException();
 		public Task SendIssueReportAsync(IssueReportGroup report, CancellationToken cancellationToken) => throw new NotImplementedException();
 		public Task SendDeviceIssueReportAsync(DeviceIssueReport report, CancellationToken cancellationToken) => throw new NotImplementedException();
+		public Task NotifyTestHealthReportAsync(ITestHealthReport report, string recipient, string[]? carbonCopies, CancellationToken cancellationToken) => throw new NotImplementedException();
 		public void Dispose()
 		{
 			Dispose(true);
diff --git a/Engine/Source/Programs/Horde/Plugins/Experimental/HordeServer.Experimental/Notifications/Sinks/ExperimentalSlackNotificationSink.cs b/Engine/Source/Programs/Horde/Plugins/Experimental/HordeServer.Experimental/Notifications/Sinks/ExperimentalSlackNotificationSink.cs
--- a/Engine/Source/Programs/Horde/Plugins/Experimental/HordeServer.Experimental/Notifications/Sinks/ExperimentalSlackNotificationSink.cs
+++ b/Engine/Source/Programs/Horde/Plugins/Experimental/HordeServer.Experimental/Notifications/Sinks/ExperimentalSlackNotificationSink.cs
@@ -8,6 +8,7 @@ using HordeServer.Devices;
 using HordeServer.Experimental.Notifications;
 using HordeServer.Issues;
 using HordeServer.Jobs;
+using HordeServer.Jobs.TestData;
 using HordeServer.Logs;
 using HordeServer.Streams;
 using HordeServer.Users;
@@ -116,6 +117,10 @@ namespace HordeServer.Notifications.Sinks
 		public Task SendIssueReportAsync(IssueReportGroup report, CancellationToken cancellationToken)
 			=> Task.CompletedTask;
 
+		/// <inheritdoc/>
+		public Task NotifyTestHealthReportAsync(ITestHealthReport report, string recipient, string[]? carbonCopies, CancellationToken cancellationToken)
+			=> Task.CompletedTask;
+
 		/// <inheritdoc/>
 		public void Dispose()
 			=> _slackNotificationProcessor.Dispose();
diff --git a/Engine/Source/Programs/Shared/EpicGames.Horde/Jobs/TestData/TestDataMessages.cs b/Engine/Source/Programs/Shared/EpicGames.Horde/Jobs/TestData/TestDataMessages.cs
--- a/Engine/Source/Programs/Shared/EpicGames.Horde/Jobs/TestData/TestDataMessages.cs
+++ b/Engine/Source/Programs/Shared/EpicGames.Horde/Jobs/TestData/TestDataMessages.cs
@@ -389,6 +389,22 @@ namespace EpicGames.Horde.Jobs.TestData
 		public List<string>? MetaIds { get; set; }
 	}
 
+	/// <summary>
+	/// Get test health report from streams and tests
+	/// </summary>
+	public class GetTestHealthReportRequest
+	{
+		/// <summary>
+		/// Stream ids to get
+		/// </summary>
+		public List<string> StreamIds { get; set; } = new List<string>();
+
+		/// <summary>
+		/// Test ids to get
+		/// </summary>
+		public List<string>? TestIds { get; set; }
+	}
+
 	/// <summary>
 	/// Get phase sessions from stream and phase ids
 	/// </summary>
@@ -558,6 +574,67 @@ namespace EpicGames.Horde.Jobs.TestData
 		public string? Notes { get; set; }
 	}
 
+	/// <summary>
+	/// A test health report
+	/// </summary>
+	public class GetTestHealthReportResponse
+	{
+		/// <summary>
+		/// The report id
+		/// </summary>
+		public string Id { get; set; } = String.Empty;
+
+		/// <summary>
+		/// The test id associated to the report
+		/// </summary>
+		public string TestId { get; set; } = String.Empty;
+
+		/// <summary>
+		/// The stream id associated to the report
+		/// </summary>
+		public string StreamId { get; set; } = String.Empty;
+
+		/// <summary>
+		/// The date of the last update
+		/// </summary>
+		public DateTime LastUpdateDate { get; set; }
+
+		/// <summary>
+		/// Whether or not the test healthy
+		/// </summary>
+		public bool IsHealthy {  get; set; }
+
+		/// <summary>
+		/// The state associated with the test health
+		/// </summary>
+		public string State { get; set; } = String.Empty;
+
+		/// <summary>
+		/// The success rate of the test health
+		/// </summary>
+		public int SuccessRate { get; set; }
+
+		/// <summary>
+		/// The failure rate of the test health
+		/// </summary>
+		public int FailureRate { get; set; }
+
+		/// <summary>
+		/// The catastrophic rate of the test failure
+		/// </summary>
+		public int CatastrophicFailureRate { get; set; }
+
+		/// <summary>
+		/// The redundant error rate of the test failure
+		/// </summary>
+		public int RedundantErrorRate { get; set; }
+
+		/// <summary>
+		/// The last date of the notification sent
+		/// </summary>
+		public DateTime? NotificationLastDate { get; set; }
+	}
+
 	/// <summary>
 	/// A test suite that runs in a stream, contain subtests
 	/// </summary>
