Revert TestHealth Acl check

#rnx
#Synced-CL 48996119

diff --git a/Engine/Source/Programs/Horde/Plugins/Build/HordeServer.Build/Jobs/TestData/TestDataController.cs b/Engine/Source/Programs/Horde/Plugins/Build/HordeServer.Build/Jobs/TestData/TestDataController.cs
--- a/Engine/Source/Programs/Horde/Plugins/Build/HordeServer.Build/Jobs/TestData/TestDataController.cs
+++ b/Engine/Source/Programs/Horde/Plugins/Build/HordeServer.Build/Jobs/TestData/TestDataController.cs
@@ -924,11 +924,6 @@ namespace HordeServer.Jobs.TestData
 		[ProducesResponseType(typeof(GetTestAuditResponse), 200)]
 		public async Task<ActionResult<GetTestAuditResponse>> AddTestAuditV2Async([FromBody] AddTestAuditRequest request)
 		{
-			if (!_buildConfig.Value.Authorize(TestDataAclAction.TestHealthWrite, User))
-			{
-				return Forbid();
-			}
-
 			GetTestAuditResponse responses = new GetTestAuditResponse();
 
 			UserId? userId = User.GetUserId();
@@ -980,11 +975,6 @@ namespace HordeServer.Jobs.TestData
 			[FromBody] GetTestAuditRequest request,
 			[FromQuery] DateTimeOffset? minCreateTime = null)
 		{
-			if (!_buildConfig.Value.Authorize(TestDataAclAction.TestHealthRead, User))
-			{
-				return Forbid();
-			}
-
 			TestId[] testIds = request.TestIds.Select(x => TestId.Parse(x)).ToHashSet().ToArray();
 
 			IReadOnlyList<ITestAudit> audits = await _testDataService.FindTestAuditsV2Async(testIds, minCreateTime?.UtcDateTime);
@@ -1028,11 +1018,6 @@ namespace HordeServer.Jobs.TestData
 		[Route("/api/v3/testdata/audit/log/{testId}")]
 		public async Task<ActionResult> GetAuditLogHistoryAsync(string testId, [FromQuery] DateTime? minTime = null, [FromQuery] DateTime? maxTime = null, [FromQuery] int index = 0, [FromQuery] int count = 50)
 		{
-			if (!_buildConfig.Value.Authorize(TestDataAclAction.TestHealthRead, User))
-			{
-				return Forbid();
-			}
-
 			TestId test;
 			if (String.IsNullOrEmpty(testId) || !TestId.TryParse(testId, out test))
 			{
@@ -1058,11 +1043,6 @@ namespace HordeServer.Jobs.TestData
 		[ProducesResponseType(typeof(List<GetTestHealthReportResponse>), 200)]
 		public async Task<ActionResult<List<GetTestHealthReportResponse>>> GetTestHealthReportsV2Async([FromBody] GetTestHealthReportRequest request)
 		{
-			if (!_buildConfig.Value.Authorize(TestDataAclAction.TestHealthRead, User))
-			{
-				return Forbid();
-			}
-
 			StreamId[] streamIdValues = request.StreamIds.Select(x => new StreamId(x)).ToArray();
 
 			List<StreamId> queryStreams = new List<StreamId>();
