Horde: Updating schema documentation

#rnx

diff --git a/Engine/Source/Programs/Horde/Docs/Config/Schema/AclActions.md b/Engine/Source/Programs/Horde/Docs/Config/Schema/AclActions.md
--- a/Engine/Source/Programs/Horde/Docs/Config/Schema/AclActions.md
+++ b/Engine/Source/Programs/Horde/Docs/Config/Schema/AclActions.md
@@ -37,6 +37,32 @@
 | `DeleteAccount` |  |
 | `ViewAccount` |  |
 
+## Notices
+
+| Name | Description |
+| ---- | ----------- |
+| `CreateNotice` |  |
+| `UpdateNotice` |  |
+| `DeleteNotice` |  |
+
+## ServiceAccounts
+
+| Name | Description |
+| ---- | ----------- |
+| `CreateAccount` |  |
+| `UpdateAccount` |  |
+| `DeleteAccount` |  |
+| `ViewAccount` |  |
+
+## Accounts
+
+| Name | Description |
+| ---- | ----------- |
+| `CreateAccount` |  |
+| `UpdateAccount` |  |
+| `DeleteAccount` |  |
+| `ViewAccount` |  |
+
 ## Artifacts
 
 | Name | Description |
@@ -129,106 +155,12 @@
 | `ViewChanges` | View changes submitted to a stream. NOTE: this returns responses from the server's Perforce account, which may be a priviledged user. |
 | `ViewTemplate` | View template associated with a stream |
 
-## Accounts
+## TestData
 
 | Name | Description |
 | ---- | ----------- |
-| `CreateAccount` |  |
-| `UpdateAccount` |  |
-| `DeleteAccount` |  |
-| `ViewAccount` |  |
-
-## Artifacts
-
-| Name | Description |
-| ---- | ----------- |
-| `ReadArtifact` |  |
-| `WriteArtifact` |  |
-| `DeleteArtifact` |  |
-| `UploadArtifact` |  |
-| `DownloadArtifact` |  |
-
-## Bisect
-
-| Name | Description |
-| ---- | ----------- |
-| `CreateBisectTask` |  |
-| `UpdateBisectTask` |  |
-| `ViewBisectTask` |  |
-
-## BuildHealth
-
-| Name | Description |
-| ---- | ----------- |
-| `AddUpdateDeleteBuildHealthFilter` |  |
-
-## Devices
-
-| Name | Description |
-| ---- | ----------- |
-| `DeviceRead` |  |
-| `DeviceWrite` |  |
-
-## Jobs
-
-| Name | Description |
-| ---- | ----------- |
-| `CreateJob` |  |
-| `UpdateJob` |  |
-| `DeleteJob` |  |
-| `ExecuteJob` |  |
-| `RetryJobStep` |  |
-| `ViewJob` |  |
-
-## Notices
-
-| Name | Description |
-| ---- | ----------- |
-| `CreateNotice` |  |
-| `UpdateNotice` |  |
-| `DeleteNotice` |  |
-
-## Notifications
-
-| Name | Description |
-| ---- | ----------- |
-| `CreateSubscription` |  |
-
-## Projects
-
-| Name | Description |
-| ---- | ----------- |
-| `CreateProject` |  |
-| `DeleteProject` |  |
-| `UpdateProject` |  |
-| `ViewProject` |  |
-
-## Replicators
-
-| Name | Description |
-| ---- | ----------- |
-| `UpdateReplicator` |  |
-| `ViewReplicator` |  |
-
-## ServiceAccounts
-
-| Name | Description |
-| ---- | ----------- |
-| `CreateAccount` |  |
-| `UpdateAccount` |  |
-| `DeleteAccount` |  |
-| `ViewAccount` |  |
-
-## Streams
-
-| Name | Description |
-| ---- | ----------- |
-| `CreateStream` |  |
-| `UpdateStream` |  |
-| `DeleteStream` |  |
-| `ViewStream` |  |
-| `ViewChanges` |  |
-| `ViewTemplate` |  |
+| `TestHealthRead` | Ability to read test health |
+| `TestHealthWrite` | Ability to write test health |
 
 ## Accounts
 
@@ -394,6 +326,13 @@
 | `ViewChanges` |  |
 | `ViewTemplate` |  |
 
+## TestData
+
+| Name | Description |
+| ---- | ----------- |
+| `TestHealthRead` |  |
+| `TestHealthWrite` |  |
+
 ## Accounts
 
 | Name | Description |
@@ -558,6 +497,13 @@
 | `ViewChanges` |  |
 | `ViewTemplate` |  |
 
+## TestData
+
+| Name | Description |
+| ---- | ----------- |
+| `TestHealthRead` |  |
+| `TestHealthWrite` |  |
+
 ## Accounts
 
 | Name | Description |
@@ -687,12 +633,6 @@
 | `UpdateReplicator` |  |
 | `ViewReplicator` |  |
 
-## Secrets
-
-| Name | Description |
-| ---- | ----------- |
-| `ViewSecret` | View a credential |
-
 ## ServiceAccounts
 
 | Name | Description |
@@ -728,6 +668,13 @@
 | `ViewChanges` |  |
 | `ViewTemplate` |  |
 
+## TestData
+
+| Name | Description |
+| ---- | ----------- |
+| `TestHealthRead` |  |
+| `TestHealthWrite` |  |
+
 ## Accounts
 
 | Name | Description |
@@ -861,7 +808,7 @@
 
 | Name | Description |
 | ---- | ----------- |
-| `ViewSecret` |  |
+| `ViewSecret` | View a credential |
 
 ## ServiceAccounts
 
@@ -898,6 +845,13 @@
 | `ViewChanges` |  |
 | `ViewTemplate` |  |
 
+## TestData
+
+| Name | Description |
+| ---- | ----------- |
+| `TestHealthRead` |  |
+| `TestHealthWrite` |  |
+
 ## Accounts
 
 | Name | Description |
@@ -1057,19 +1011,6 @@
 | `DownloadSoftware` |  |
 | `DeleteSoftware` |  |
 
-## Storage
-
-| Name | Description |
-| ---- | ----------- |
-| `ReadBlobs` | Ability to read blobs from the storage service |
-| `WriteBlobs` | Ability to write blobs to the storage service |
-| `ReadAliases` | Ability to read aliases from the storage service |
-| `WriteAliases` | Ability to write aliases to the storage service |
-| `DeleteAliases` | Ability to write aliases to the storage service |
-| `ReadRefs` | Ability to read refs from the storage service |
-| `WriteRefs` | Ability to write refs to the storage service |
-| `DeleteRefs` | Ability to delete refs |
-
 ## Streams
 
 | Name | Description |
@@ -1081,6 +1022,19 @@
 | `ViewChanges` |  |
 | `ViewTemplate` |  |
 
+## Symbols
+
+| Name | Description |
+| ---- | ----------- |
+| `ReadSymbols` | Ability to download symbols |
+
+## TestData
+
+| Name | Description |
+| ---- | ----------- |
+| `TestHealthRead` |  |
+| `TestHealthWrite` |  |
+
 ## Accounts
 
 | Name | Description |
@@ -1244,14 +1198,14 @@
 
 | Name | Description |
 | ---- | ----------- |
-| `ReadBlobs` |  |
-| `WriteBlobs` |  |
-| `ReadAliases` |  |
-| `WriteAliases` |  |
-| `DeleteAliases` |  |
-| `ReadRefs` |  |
-| `WriteRefs` |  |
-| `DeleteRefs` |  |
+| `ReadBlobs` | Ability to read blobs from the storage service |
+| `WriteBlobs` | Ability to write blobs to the storage service |
+| `ReadAliases` | Ability to read aliases from the storage service |
+| `WriteAliases` | Ability to write aliases to the storage service |
+| `DeleteAliases` | Ability to write aliases to the storage service |
+| `ReadRefs` | Ability to read refs from the storage service |
+| `WriteRefs` | Ability to write refs to the storage service |
+| `DeleteRefs` | Ability to delete refs |
 
 ## Streams
 
@@ -1268,7 +1222,14 @@
 
 | Name | Description |
 | ---- | ----------- |
-| `ReadSymbols` | Ability to download symbols |
+| `ReadSymbols` |  |
 
+## TestData
+
+| Name | Description |
+| ---- | ----------- |
+| `TestHealthRead` |  |
+| `TestHealthWrite` |  |
+
 ## Accounts
 
@@ -1459,6 +1420,13 @@
 | ---- | ----------- |
 | `ReadSymbols` |  |
 
+## TestData
+
+| Name | Description |
+| ---- | ----------- |
+| `TestHealthRead` |  |
+| `TestHealthWrite` |  |
+
 ## Tools
 
 | Name | Description |
@@ -1655,6 +1623,13 @@
 | ---- | ----------- |
 | `ReadSymbols` |  |
 
+## TestData
+
+| Name | Description |
+| ---- | ----------- |
+| `TestHealthRead` |  |
+| `TestHealthWrite` |  |
+
 ## Tools
 
 | Name | Description |
diff --git a/Engine/Source/Programs/Horde/Docs/Config/Schema/Globals.md b/Engine/Source/Programs/Horde/Docs/Config/Schema/Globals.md
--- a/Engine/Source/Programs/Horde/Docs/Config/Schema/Globals.md
+++ b/Engine/Source/Programs/Horde/Docs/Config/Schema/Globals.md
@@ -59,16 +59,24 @@
 
 Name | Description
 ---- | -----------
-`build` | [BuildConfig](#buildconfig)<br>Configuration for the build plugin
 `analytics` | [AnalyticsConfig](#analyticsconfig)<br>Configuration for the analytics plugin
+`build` | [BuildConfig](#buildconfig)<br>Configuration for the build plugin
 `compute` | [ComputeConfig](#computeconfig)<br>Configuration for the compute plugin
 `exampleplugin` | [ExamplePluginConfig](#examplepluginconfig)<br>Configuration for the exampleplugin plugin
-`secrets` | [SecretsConfig](#secretsconfig)<br>Configuration for the secrets plugin
 `experimental` | [ExperimentalConfig](#experimentalconfig)<br>Configuration for the experimental plugin
-`storage` | [StorageConfig](#storageconfig)<br>Configuration for the storage plugin
+`secrets` | [SecretsConfig](#secretsconfig)<br>Configuration for the secrets plugin
 `symbols` | [SymbolsConfig](#symbolsconfig)<br>Configuration for the symbols plugin
+`storage` | [StorageConfig](#storageconfig)<br>Configuration for the storage plugin
 `tools` | [ToolsConfig](#toolsconfig)<br>Configuration for the tools plugin
 
+## AnalyticsConfig
+
+Config settings for analytics
+
+Name | Description
+---- | -----------
+`stores` | [TelemetryStoreConfig](Telemetry.md)`[]`<br>Metrics to aggregate on the Horde server
+
 ## BuildConfig
 
 Configuration for the build plugin
@@ -232,14 +240,6 @@
 `excludeActions` | `string[]`<br>Actions to exclude from the inherited actions
 `extends` | `string[]`<br>Other profiles to extend from
 
-## AnalyticsConfig
-
-Config settings for analytics
-
-Name | Description
----- | -----------
-`stores` | [TelemetryStoreConfig](Telemetry.md)`[]`<br>Metrics to aggregate on the Horde server
-
 ## ComputeConfig
 
 Configuration for the compute system
@@ -443,6 +443,14 @@
 ---- | -----------
 `greetingMessage` | `string`<br>The Greeting provided by the example plugin
 
+## ExperimentalConfig
+
+Configuration for the experimental plugin
+
+Name | Description
+---- | -----------
+`notifications` | [NotificationConfig](Notification.md)`[]`<br>Notifications to be sent on the Horde server
+
 ## SecretsConfig
 
 Configuration for the secrets system
@@ -525,13 +533,24 @@
 ---- | -----------
 `vaultUri` | `string`<br>The Uri to the Azure Key Vault.
 
-## ExperimentalConfig
+## SymbolsConfig
 
-Configuration for the experimental plugin
+Configuration for the tools system
 
 Name | Description
 ---- | -----------
-`notifications` | [NotificationConfig](Notification.md)`[]`<br>Notifications to be sent on the Horde server
+`stores` | [SymbolStoreConfig](#symbolstoreconfig)`[]`<br>List of symbol stores
 
+## SymbolStoreConfig
+
+Configuration for a symbol store
+
+Name | Description
+---- | -----------
+`id` | `string`<br>Identifier for this store
+`namespaceId` | `string`<br>Configuration for the symbol store backend
+`public` | `boolean`<br>Whether to make this store available without auth
+`acl` | [AclConfig](#aclconfig)<br>Access to the symbol store
+
 ## StorageConfig
 
@@ -605,25 +624,6 @@
 `enableAliases` | `boolean`<br>Support querying exports by their aliases
 `acl` | [AclConfig](#aclconfig)<br>Access list for this namespace
 
-## SymbolsConfig
-
-Configuration for the tools system
-
-Name | Description
----- | -----------
-`stores` | [SymbolStoreConfig](#symbolstoreconfig)`[]`<br>List of symbol stores
-
-## SymbolStoreConfig
-
-Configuration for a symbol store
-
-Name | Description
----- | -----------
-`id` | `string`<br>Identifier for this store
-`namespaceId` | `string`<br>Configuration for the symbol store backend
-`public` | `boolean`<br>Whether to make this store available without auth
-`acl` | [AclConfig](#aclconfig)<br>Access to the symbol store
-
 ## ToolsConfig
 
 Configuration for the tools system
diff --git a/Engine/Source/Programs/Horde/Docs/Deployment/ServerSettings.md b/Engine/Source/Programs/Horde/Docs/Deployment/ServerSettings.md
--- a/Engine/Source/Programs/Horde/Docs/Deployment/ServerSettings.md
+++ b/Engine/Source/Programs/Horde/Docs/Deployment/ServerSettings.md
@@ -127,16 +127,53 @@
 
 Name | Description
 ---- | -----------
-`build` | [BuildServerConfig](#buildserverconfig)<br>Configuration for the build plugin
 `analytics` | [AnalyticsServerConfig](#analyticsserverconfig)<br>Configuration for the analytics plugin
+`build` | [BuildServerConfig](#buildserverconfig)<br>Configuration for the build plugin
 `compute` | [ComputeServerConfig](#computeserverconfig)<br>Configuration for the compute plugin
 `exampleplugin` | [ExamplePluginServerConfig](#examplepluginserverconfig)<br>Configuration for the exampleplugin plugin
-`secrets` | [SecretsServerConfig](#secretsserverconfig)<br>Configuration for the secrets plugin
 `experimental` | [ExperimentalServerConfig](#experimentalserverconfig)<br>Configuration for the experimental plugin
-`storage` | [StorageServerConfig](#storageserverconfig)<br>Configuration for the storage plugin
+`secrets` | [SecretsServerConfig](#secretsserverconfig)<br>Configuration for the secrets plugin
 `symbols` | [PluginServerConfig](#pluginserverconfig)<br>Configuration for the symbols plugin
+`storage` | [StorageServerConfig](#storageserverconfig)<br>Configuration for the storage plugin
 `tools` | [ToolsServerConfig](#toolsserverconfig)<br>Configuration for the tools plugin
 
+## AnalyticsServerConfig
+
+Server configuration for the analytics system
+
+Name | Description
+---- | -----------
+`sinks` | [TelemetrySinkConfig](#telemetrysinkconfig)<br>Settings for the various telemetry sinks
+`enabled` | `boolean`<br>Whether the plugin should be enabled or not
+
+## TelemetrySinkConfig
+
+Telemetry sinks
+
+Name | Description
+---- | -----------
+`epic` | [EpicTelemetryConfig](#epictelemetryconfig)<br>Settings for the Epic telemetry sink
+`mongo` | [MongoTelemetryConfig](#mongotelemetryconfig)<br>Settings for the MongoDB telemetry sink
+
+## EpicTelemetryConfig
+
+Configuration for the telemetry sink
+
+Name | Description
+---- | -----------
+`url` | `string`<br>Base URL for the telemetry server
+`appId` | `string`<br>Application name to send in the event messages
+`enabled` | `boolean`<br>Whether to enable this sink
+
+## MongoTelemetryConfig
+
+Configuration for the telemetry sink
+
+Name | Description
+---- | -----------
+`retainDays` | `number`<br>Number of days worth of telmetry events to keep
+`enabled` | `boolean`<br>Whether to enable this sink
+
 ## BuildServerConfig
 
 Static configuration for the build plugin
@@ -168,6 +205,8 @@
 `jobNotificationChannel` | `string`<br>Slack channel to send job related notifications to. Multiple channels can be specified, separated by ;
 `agentNotificationChannel` | `string`<br>Slack channel to send agent related notifications to.
 `testDataRetainMonths` | `integer`<br>The number of months to retain test data
+`testHealthHistoryEvaluationDays` | `integer`<br>The number of days to go back in historical data to evaluate test health
+`testHealthReportRemindDays` | `integer`<br>The number of days after which a test health report remind for improvement
 `blockCacheDir` | `string`<br>Directory to store the fine-grained block cache. This caches individual exports embedded in bundles.
 `blockCacheSize` | `string`<br>Maximum size of the block cache. Accepts standard binary suffixes. Currently only allocates in multiples of 1024mb.
 `blockCacheSizeBytes` | `integer`<br>Accessor for the block cache size in bytes
@@ -279,43 +318,6 @@
 `maxChildCount` | `integer`<br>Maximum number of children in each node
 `sliceThreshold` | `integer`<br>Threshold hash value for splitting interior nodes
 
-## AnalyticsServerConfig
-
-Server configuration for the analytics system
-
-Name | Description
----- | -----------
-`sinks` | [TelemetrySinkConfig](#telemetrysinkconfig)<br>Settings for the various telemetry sinks
-`enabled` | `boolean`<br>Whether the plugin should be enabled or not
-
-## TelemetrySinkConfig
-
-Telemetry sinks
-
-Name | Description
----- | -----------
-`epic` | [EpicTelemetryConfig](#epictelemetryconfig)<br>Settings for the Epic telemetry sink
-`mongo` | [MongoTelemetryConfig](#mongotelemetryconfig)<br>Settings for the MongoDB telemetry sink
-
-## EpicTelemetryConfig
-
-Configuration for the telemetry sink
-
-Name | Description
----- | -----------
-`url` | `string`<br>Base URL for the telemetry server
-`appId` | `string`<br>Application name to send in the event messages
-`enabled` | `boolean`<br>Whether to enable this sink
-
-## MongoTelemetryConfig
-
-Configuration for the telemetry sink
-
-Name | Description
----- | -----------
-`retainDays` | `number`<br>Number of days worth of telmetry events to keep
-`enabled` | `boolean`<br>Whether to enable this sink
-
 ## ComputeServerConfig
 
 Static configuration for the compute plugin
@@ -369,6 +371,14 @@
 ---- | -----------
 `enabled` | `boolean`<br>Whether the plugin should be enabled or not
 
+## ExperimentalServerConfig
+
+Server configuration for the experimental plugin
+
+Name | Description
+---- | -----------
+`enabled` | `boolean`<br>Whether the plugin should be enabled or not
+
 ## SecretsServerConfig
 
 Static configuration for the secrets plugin
@@ -378,9 +388,9 @@
 `withAws` | `boolean`<br>Whether to enable Amazon Web Services (AWS) specific features
 `enabled` | `boolean`<br>Whether the plugin should be enabled or not
 
-## ExperimentalServerConfig
+## PluginServerConfig
 
-Server configuration for the experimental plugin
+Base class for plugin server config objects
 
 Name | Description
 ---- | -----------
@@ -445,14 +455,6 @@
 `AssumeRole` | Assume a particular role. Should specify ARN in
 `AssumeRoleWebIdentity` | Assume a particular role using the current environment variables.
 
-## PluginServerConfig
-
-Base class for plugin server config objects
-
-Name | Description
----- | -----------
-`enabled` | `boolean`<br>Whether the plugin should be enabled or not
-
 ## ToolsServerConfig
 
 Server configuration for bundled tools
