Fix situation where credentials would be printed to log on Mac/Linux because we compared conformed environment variable names against non-conformed credential key names.  Now doing the comparison with conformed values on both sides of the comparison.

#Synced-CL 48523383

diff --git a/Engine/Source/Programs/Horde/Drivers/JobDriver/Execution/JobExecutor.cs b/Engine/Source/Programs/Horde/Drivers/JobDriver/Execution/JobExecutor.cs
--- a/Engine/Source/Programs/Horde/Drivers/JobDriver/Execution/JobExecutor.cs
+++ b/Engine/Source/Programs/Horde/Drivers/JobDriver/Execution/JobExecutor.cs
@@ -2021,8 +2021,8 @@ namespace JobDriver.Execution
 			}
 
 			// Log all the environment variables to the log
-			HashSet<string> credentialKeys = new HashSet<string>(step.Credentials.Keys, StringComparer.OrdinalIgnoreCase);
-			credentialKeys.Add(HordeHttpClient.HordeTokenEnvVarName);
+			HashSet<string> credentialKeys = new HashSet<string>(step.Credentials.Keys.Select(credentialKey => ConformEnvironmentVariableName(credentialKey)), StringComparer.OrdinalIgnoreCase);
+			credentialKeys.Add(ConformEnvironmentVariableName(HordeHttpClient.HordeTokenEnvVarName));
 
 			foreach (KeyValuePair<string, string> envVar in newEnvVars.OrderBy(x => x.Key))
 			{
