Horde: Suppress telemetry events with duplicate keys

diff --git a/Engine/Source/Programs/Horde/Plugins/Analytics/HordeServer.Analytics/Telemetry/TelemetryController.cs b/Engine/Source/Programs/Horde/Plugins/Analytics/HordeServer.Analytics/Telemetry/TelemetryController.cs
--- a/Engine/Source/Programs/Horde/Plugins/Analytics/HordeServer.Analytics/Telemetry/TelemetryController.cs
+++ b/Engine/Source/Programs/Horde/Plugins/Analytics/HordeServer.Analytics/Telemetry/TelemetryController.cs
@@ -8,6 +8,7 @@ using HordeServer.Telemetry.Metrics;
 using HordeServer.Utilities;
 using Microsoft.AspNetCore.Authorization;
 using Microsoft.AspNetCore.Mvc;
+using Microsoft.AspNetCore.Mvc.ModelBinding.Validation;
 using Microsoft.Extensions.Options;
 
 namespace HordeServer.Telemetry
@@ -178,7 +179,7 @@ namespace HordeServer.Telemetry
 		/// <param name="uploadType">Type of data being uploaded</param>
 		[HttpPost]
 		[Route("/api/v1/telemetry/{storeId}")]
-		public ActionResult PostEvent(TelemetryStoreId storeId, [FromBody] PostTelemetryEventStreamRequest request, [FromQuery][Required] string appId, [FromQuery][Required] string appVersion, [FromQuery][Required] string appEnvironment, [FromQuery][Required] TelemetryUploadType uploadType)
+		public ActionResult PostEvent(TelemetryStoreId storeId, [FromBody][ValidateNever] PostTelemetryEventStreamRequest request, [FromQuery][Required] string appId, [FromQuery][Required] string appVersion, [FromQuery][Required] string appEnvironment, [FromQuery][Required] TelemetryUploadType uploadType)
 		{
 			if (uploadType != TelemetryUploadType.EtEventStream)
 			{
