Add TestHealthService to Hosted Service a second time

#rnx
#Synced-CL 49069065

diff --git a/Engine/Source/Programs/Horde/Plugins/Build/HordeServer.Build/BuildPlugin.cs b/Engine/Source/Programs/Horde/Plugins/Build/HordeServer.Build/BuildPlugin.cs
--- a/Engine/Source/Programs/Horde/Plugins/Build/HordeServer.Build/BuildPlugin.cs
+++ b/Engine/Source/Programs/Horde/Plugins/Build/HordeServer.Build/BuildPlugin.cs
@@ -217,6 +217,7 @@ namespace HordeServer
 				services.AddHostedService(provider => provider.GetRequiredService<DeviceService>());
 				services.AddHostedService<DeviceReportService>();
 				services.AddHostedService(provider => provider.GetRequiredService<TestDataService>());
+				services.AddHostedService<TestHealthService>();
 
 				if (_staticConfig.Commits.ReplicateMetadata)
 				{
