Horde: Adding configurable User1, 2, 3 status colors

diff --git a/Engine/Source/Programs/Horde/HordeDashboard/src/backend/Api.ts b/Engine/Source/Programs/Horde/HordeDashboard/src/backend/Api.ts
--- a/Engine/Source/Programs/Horde/HordeDashboard/src/backend/Api.ts
+++ b/Engine/Source/Programs/Horde/HordeDashboard/src/backend/Api.ts
@@ -1235,6 +1235,8 @@ export type GetReportResponse = {
 	/** The job or step severity */
 	severity?: string;
 
+	/** Color to display for message severity */
+	color?: string;
 }
 
 export type GetJobArtifactResponse = {
@@ -3403,6 +3405,9 @@ export enum DashboardPreference {
 	ColorWarning = "ColorWarning",
 	ColorError = "ColorError",
 	ColorRunning = "ColorRunning",
+	ColorUser1 = "ColorUser1",
+	ColorUser2 = "ColorUser2",
+	ColorUser3 = "ColorUser3",
 	LocalCache = "LocalCache",
 	LeftAlignLog = "LeftAlignLog",
 	CompactViews = "CompactViews",
diff --git a/Engine/Source/Programs/Horde/HordeDashboard/src/backend/Dashboard.ts b/Engine/Source/Programs/Horde/HordeDashboard/src/backend/Dashboard.ts
--- a/Engine/Source/Programs/Horde/HordeDashboard/src/backend/Dashboard.ts
+++ b/Engine/Source/Programs/Horde/HordeDashboard/src/backend/Dashboard.ts
@@ -14,6 +14,9 @@ export enum StatusColor {
     Skipped,
     Aborted,
     Running,
+    User1,
+    User2,
+    User3,
     Unspecified
 }
 
@@ -499,6 +502,9 @@ export class Dashboard {
             [StatusColor.Waiting, dark ? "#474542" : "#A19F9D"],
             [StatusColor.Ready, dark ? "#474542" : "#A19F9D"],
             [StatusColor.Skipped, dark ? "#63625c" : "#C3C2C1"],
+            [StatusColor.User1, dark ? "#8e4747" : "#8e4747"],
+            [StatusColor.User2, dark ? "#785EF0" : "#785EF0"],
+            [StatusColor.User3, dark ? "#FE6100" : "#FE6100"],            
             [StatusColor.Unspecified, "#646464"]
         ]);
 
@@ -518,12 +524,18 @@ export class Dashboard {
         const warning = this.getPreference(DashboardPreference.ColorWarning);
         const error = this.getPreference(DashboardPreference.ColorError);
         const running = this.getPreference(DashboardPreference.ColorRunning);
+        const user1 = this.getPreference(DashboardPreference.ColorUser1);
+        const user2 = this.getPreference(DashboardPreference.ColorUser2);
+        const user3 = this.getPreference(DashboardPreference.ColorUser3 );
 
         const colors = new Map<StatusColor, string>([
             [StatusColor.Success, success ? success : defaultStatusColors.get(StatusColor.Success)!],
             [StatusColor.Warnings, warning ? warning : defaultStatusColors.get(StatusColor.Warnings)!],
             [StatusColor.Failure, error ? error : defaultStatusColors.get(StatusColor.Failure)!],
             [StatusColor.Running, running ? running : defaultStatusColors.get(StatusColor.Running)!],
+            [StatusColor.User1, user1 ? user1 : defaultStatusColors.get(StatusColor.User1)!],
+            [StatusColor.User2, user2 ? user2 : defaultStatusColors.get(StatusColor.User2)!],
+            [StatusColor.User3, user3 ? user3 : defaultStatusColors.get(StatusColor.User3)!],
             [StatusColor.Waiting, defaultStatusColors.get(StatusColor.Waiting)!],
             [StatusColor.Ready, defaultStatusColors.get(StatusColor.Ready)!],
             [StatusColor.Skipped, defaultStatusColors.get(StatusColor.Skipped)!],
diff --git a/Engine/Source/Programs/Horde/HordeDashboard/src/components/DashboardView.tsx b/Engine/Source/Programs/Horde/HordeDashboard/src/components/DashboardView.tsx
--- a/Engine/Source/Programs/Horde/HordeDashboard/src/components/DashboardView.tsx
+++ b/Engine/Source/Programs/Horde/HordeDashboard/src/components/DashboardView.tsx
@@ -16,18 +16,23 @@ const colorBlind1 = new Map<StatusColor, string>([
    [StatusColor.Success, "#37A862"],
    [StatusColor.Warnings, "#EECC17"],
    [StatusColor.Failure, "#FD827D"],
-   [StatusColor.Running, "#7A6FF1"]
-]);
+   [StatusColor.Running, "#7A6FF1"],
+   [StatusColor.User1, "#8e4747"],
+   [StatusColor.User2, "#785EF0"],
+   [StatusColor.User3, "#FE6100"]
 
+]);
 
 const colorBlind2 = new Map<StatusColor, string>([
    [StatusColor.Success, "#208FA3"],
    [StatusColor.Warnings, "#FFB000"],
    [StatusColor.Failure, "#E8384F"],
-   [StatusColor.Running, "#AA71FF"]
+   [StatusColor.Running, "#AA71FF"],
+   [StatusColor.User1, "#8e4747"],
+   [StatusColor.User2, "#785EF0"],
+   [StatusColor.User3, "#FE6100"],            
 ]);
 
-
 const statusText = new Map<StatusColor, string>([
    [StatusColor.Success, "Success"],
    [StatusColor.Warnings, "Warnings"],
@@ -36,21 +41,30 @@ const statusText = new Map<StatusColor, string>([
    [StatusColor.Ready, "Ready"],
    [StatusColor.Skipped, "Skipped"],
    [StatusColor.Aborted, "Canceled"],
-   [StatusColor.Running, "Running"]
+   [StatusColor.Running, "Running"],
+   [StatusColor.User1, "User1"],
+   [StatusColor.User2, "User2"],
+   [StatusColor.User3, "User3"]
 ]);
 
 const orderedStatus = [
    StatusColor.Success,
    StatusColor.Warnings,
    StatusColor.Failure,
-   StatusColor.Running
+   StatusColor.Running,
+   StatusColor.User1,
+   StatusColor.User2,
+   StatusColor.User3
 ];
 
 const statusPrefs = new Map<StatusColor, DashboardPreference>([
    [StatusColor.Success, DashboardPreference.ColorSuccess],
    [StatusColor.Warnings, DashboardPreference.ColorWarning],
    [StatusColor.Failure, DashboardPreference.ColorError],
-   [StatusColor.Running, DashboardPreference.ColorRunning]
+   [StatusColor.Running, DashboardPreference.ColorRunning],
+   [StatusColor.User1, DashboardPreference.ColorUser1],
+   [StatusColor.User2, DashboardPreference.ColorUser2],
+   [StatusColor.User3, DashboardPreference.ColorUser3]
 ]);
 
 
@@ -128,8 +142,8 @@ const ColorPreferenceDialog: React.FC<{ shown: boolean, statusIn: StatusColor, o
       "#FFB000"
    ];
 
-   const suggestions = tolColors.map(c => {
-      return <DefaultButton key={`color_${c}`} style={{ backgroundColor: c, width: 80 }} text={statusText.get(status)} onClick={() => { setColor(c) }} />
+   const suggestions = tolColors.map((c, index) => {
+      return <DefaultButton key={`color_${index}`} style={{ backgroundColor: c, width: 80 }} text={statusText.get(status)} onClick={() => { setColor(c) }} />
    })
 
    const colorPicker = () => {
@@ -160,7 +174,7 @@ const ColorPreferenceDialog: React.FC<{ shown: boolean, statusIn: StatusColor, o
       const buttons = orderedStatus.map(status => {
 
          return <div className="horde-no-darktheme">
-            <Stack className={hordeClasses.badgeNoIcon} key={`badge_${status}`}> <DefaultButton style={{ backgroundColor: state.colors.get(status) }} text={statusText.get(status)} onClick={() => { setStatus(status) }} />
+            <Stack className={hordeClasses.badgeNoIcon} key={`badge_${status}`}> <DefaultButton style={{ backgroundColor: state.colors.get(status) ?? defaultStatusColors.get(status) }} text={statusText.get(status)} onClick={() => { setStatus(status) }} />
             </Stack>
          </div>
       });
diff --git a/Engine/Source/Programs/Horde/HordeDashboard/src/components/jobDetailsV2/JobDetailStepHistory.tsx b/Engine/Source/Programs/Horde/HordeDashboard/src/components/jobDetailsV2/JobDetailStepHistory.tsx
--- a/Engine/Source/Programs/Horde/HordeDashboard/src/components/jobDetailsV2/JobDetailStepHistory.tsx
+++ b/Engine/Source/Programs/Horde/HordeDashboard/src/components/jobDetailsV2/JobDetailStepHistory.tsx
@@ -7,7 +7,7 @@ import { useEffect } from "react";
 import { Link, useNavigate, useLocation } from "react-router-dom";
 import backend from "../../backend";
 import { GetJobStepRefResponse } from "../../backend/Api";
-import dashboard, { StatusColor } from "../../backend/Dashboard";
+import dashboard from "../../backend/Dashboard";
 import { ISideRailLink } from "../../base/components/SideRail";
 import { displayTimeZone, getElapsedString } from "../../base/utilities/timeUtils";
 import { ChangeButton } from "../ChangeButton";
@@ -15,6 +15,7 @@ import { StepRefStatusIcon } from "../StatusIcon";
 import { JobDataView, JobDetailsV2 } from "./JobDetailsViewCommon";
 import { getHordeTheme } from "../../styles/theme";
 import { getHordeStyling } from "../../styles/Styles";
+import { getStepRefMessageInfo } from "./JobDetailsUtils";
 
 const sideRail: ISideRailLink = { text: "History", url: "rail_step_history" };
 
@@ -177,29 +178,11 @@ export const StepHistoryPanel: React.FC<{ jobDetails: JobDetailsV2; stepId: stri
       }
 
       if (column.name === "Meta") {
-
-         let stepMessage = ref.metadata?.find(m => m.startsWith("ReportStepMessage="));
-         let stepSeverity = ref.metadata?.find(m => m.startsWith("ReportStepSeverity="));
-
-         stepMessage = stepMessage?.split("=")[1]?.trim();
-         stepSeverity = stepSeverity?.split("=")[1]?.trim();
-
-         if (!stepMessage) {
-            return null;
-         }
-
-         let stepMessageColor = "";
-         if (stepSeverity === "Error" || stepSeverity === "Errors") {
-            stepMessageColor = dashboard.getStatusColors().get(StatusColor.Failure)!;
-         }
-         if (stepSeverity === "Warning" || stepSeverity === "Warnings") {
-            stepMessageColor = dashboard.getStatusColors().get(StatusColor.Warnings)!;
-         }
-
+         let { message: stepMessage, color: stepMessageColor } = getStepRefMessageInfo(ref);
          return <Stack verticalAlign="center" horizontal tokens={{ childrenGap: 6 }} horizontalAlign="center">
                {!!stepMessageColor && <FontIcon style={{ fontSize: 13, color: stepMessageColor, paddingTop: 1 }} iconName="Square" />}
                <Text style={{fontSize: 12, overflow: "hidden", textOverflow: "ellipsis"}}>{stepMessage}</Text>
-            </Stack>         
+            </Stack>
       }
 
 
diff --git a/Engine/Source/Programs/Horde/HordeDashboard/src/components/jobDetailsV2/JobDetailViewSteps.tsx b/Engine/Source/Programs/Horde/HordeDashboard/src/components/jobDetailsV2/JobDetailViewSteps.tsx
--- a/Engine/Source/Programs/Horde/HordeDashboard/src/components/jobDetailsV2/JobDetailViewSteps.tsx
+++ b/Engine/Source/Programs/Horde/HordeDashboard/src/components/jobDetailsV2/JobDetailViewSteps.tsx
@@ -15,6 +15,7 @@ import { LabelsPanelV2 } from "./JobDetailLabels";
 import { JobDataView, JobDetailsV2, JobFilterBar, StateFilter } from "./JobDetailsViewCommon";
 import { getHordeTheme } from "../../styles/theme";
 import { getHordeStyling } from "../../styles/Styles";
+import { getStepMessageInfo } from "./JobDetailsUtils";
 
 const stepsSideRail: ISideRailLink = { text: "Steps", url: "rail_steps" };
 const depSideRail: ISideRailLink = { text: "Dependencies", url: "rail_dependencies" };
@@ -87,23 +88,11 @@ const RenderDynamic: React.FC<{ jobDetails: JobDetailsV2, dataView: StepsDataVie
 
       // Handle step message
       const report = step?.reports?.find(r => !!r.message);
-      if (report) {
-         const statusColors = dashboard.getStatusColors();
-         let color = "";
-         switch (report.severity) {
-            case "Error":
-            case "Errors":
-               color = statusColors.get(StatusColor.Failure)!;
-               break;
-            case "Warning":
-            case "Warnings":
-               color = statusColors.get(StatusColor.Warnings)!;
-               break;
-         }
-
+      const { message: stepMessage, color: stepMessageColor } = getStepMessageInfo(report);
+      if (stepMessage) {
          return <Stack horizontal verticalFill verticalAlign="center" tokens={{ childrenGap: 6 }} horizontalAlign="center">
-            {!!color && <FontIcon style={{ fontSize: 13, color: color }} iconName="Square" />}
-            <Text variant="tiny" style={{ overflow: "hidden", textOverflow: "ellipsis", wordBreak: "normal" }}>{report.message}</Text>
+            {!!stepMessageColor && <FontIcon style={{ fontSize: 13, color: stepMessageColor }} iconName="Square" />}
+            <Text variant="tiny" style={{ overflow: "hidden", textOverflow: "ellipsis", wordBreak: "normal" }}>{stepMessage}</Text>
          </Stack>
       }
 
diff --git a/Engine/Source/Programs/Horde/HordeDashboard/src/components/jobDetailsV2/JobDetailsUtils.ts b/Engine/Source/Programs/Horde/HordeDashboard/src/components/jobDetailsV2/JobDetailsUtils.ts
new file mode 100644
--- /dev/null
+++ b/Engine/Source/Programs/Horde/HordeDashboard/src/components/jobDetailsV2/JobDetailsUtils.ts
@@ -0,0 +1,40 @@
+import { GetJobStepRefResponse, GetReportResponse, ReportPlacement } from "horde/backend/Api";
+import dashboard, { StatusColor } from "horde/backend/Dashboard";
+
+export function getStepRefMessageInfo(ref: GetJobStepRefResponse | undefined): { message: string | undefined; color: string | undefined } {
+    
+    if (!ref?.metadata) {
+        return { message: ref?.outcome ?? "", color: undefined };
+    }
+
+    const report: GetReportResponse = {
+        placement: ReportPlacement.Summary,
+        name: ref.stepName || "",
+        message: ref.metadata.find(m => m.startsWith("ReportStepMessage="))?.split("=")[1]?.trim() ?? ref.outcome ?? "",
+        color: ref.metadata.find(m => m.startsWith("ReportStepColor="))?.split("=")[1]?.trim(),
+        severity: ref.metadata.find(m => m.startsWith("ReportStepSeverity="))?.split("=")[1]?.trim(),
+    };
+
+    return getStepMessageInfo(report);
+}
+
+export function getStepMessageInfo(report: GetReportResponse | undefined): { message: string | undefined; color: string | undefined } {
+    
+    if (!report?.message) {
+        return { message: undefined, color: undefined };
+    }
+
+    const statusColors = dashboard.getStatusColors();
+    const colorLower = report.color?.toLowerCase();
+    const severityLower = report.severity?.toLowerCase();
+
+    const userColors: Record<string, StatusColor> = { user1: StatusColor.User1, user2: StatusColor.User2, user3: StatusColor.User3 };
+    const severityColors: Record<string, StatusColor> = { error: StatusColor.Failure, errors: StatusColor.Failure, warning: StatusColor.Warnings, warnings: StatusColor.Warnings };
+
+    let color = colorLower && userColors[colorLower] ? statusColors.get(userColors[colorLower])
+        : report.color?.match(/^#?([0-9a-fA-F]{6}|[0-9a-fA-F]{3})$/) ? (report.color.startsWith('#') ? report.color : `#${report.color}`)
+            : severityLower && severityColors[severityLower] ? statusColors.get(severityColors[severityLower])
+                : undefined;
+
+    return { message: report.message, color };
+}
diff --git a/Engine/Source/Programs/Horde/HordeDashboard/src/components/jobDetailsV2/StepDetailSummary.tsx b/Engine/Source/Programs/Horde/HordeDashboard/src/components/jobDetailsV2/StepDetailSummary.tsx
--- a/Engine/Source/Programs/Horde/HordeDashboard/src/components/jobDetailsV2/StepDetailSummary.tsx
+++ b/Engine/Source/Programs/Horde/HordeDashboard/src/components/jobDetailsV2/StepDetailSummary.tsx
@@ -3,7 +3,7 @@
 import { FontIcon, Label, Stack, Text } from '@fluentui/react';
 import { observer } from 'mobx-react-lite';
 import React, { useEffect } from 'react';
-import { JobStepError, JobStepOutcome, JobStepState, ReportPlacement } from '../../backend/Api';
+import { GetReportResponse, JobStepError, JobStepOutcome, JobStepState } from '../../backend/Api';
 import { Markdown } from '../../base/components/Markdown';
 import { ISideRailLink } from '../../base/components/SideRail';
 import { getNiceTime, getStepElapsed, getStepETA, getStepFinishTime, getStepTimingDelta } from '../../base/utilities/timeUtils';
@@ -11,6 +11,7 @@ import { getHordeStyling } from '../../styles/Styles';
 import { getBatchText } from '../JobDetailCommon';
 import { JobDataView, JobDetailsV2 } from './JobDetailsViewCommon';
 import dashboard, { StatusColor } from 'horde/backend/Dashboard';
+import { getStepMessageInfo } from './JobDetailsUtils';
 
 const sideRail: ISideRailLink = { text: "Summary", url: "rail_step_summary" };
 
@@ -241,26 +242,8 @@ export const StepSummaryPanel: React.FC<{ jobDetails: JobDetailsV2; stepId: stri
    const spawnedJobUrl = spawnedJobs.length > 0 ? `/job/${spawnedJobs[0]}` : undefined;
    const parentJobUrl = jobData.parentJobId && jobData.parentJobId ? `/job/${jobData.parentJobId}?step=${jobData.parentJobStepId}` : undefined;
 
-   let stepMessage: string | undefined;
-   let stepMessageColor: string | undefined;
-
    const report = step?.reports?.find(r => !!r.message);
-   if (report) {
-
-      const statusColors = dashboard.getStatusColors();
-      stepMessage = report.message
-      stepMessageColor = "";
-      switch (report.severity) {
-         case "Error":
-         case "Errors":
-            stepMessageColor = statusColors.get(StatusColor.Failure)!;
-            break;
-         case "Warning":
-         case "Warnings":
-            stepMessageColor = statusColors.get(StatusColor.Warnings)!;
-            break;
-      }
-   }
+   const { message: stepMessage, color: stepMessageColor } = getStepMessageInfo(report);
 
    return (<Stack id={sideRail.url} styles={{ root: { paddingTop: 0, paddingRight: 12 } }}>
       <Stack className={hordeClasses.raised}>
diff --git a/Engine/Source/Programs/Horde/HordeDashboard/src/components/logs/LogView.tsx b/Engine/Source/Programs/Horde/HordeDashboard/src/components/logs/LogView.tsx
--- a/Engine/Source/Programs/Horde/HordeDashboard/src/components/logs/LogView.tsx
+++ b/Engine/Source/Programs/Horde/HordeDashboard/src/components/logs/LogView.tsx
@@ -8,7 +8,7 @@ import { Link, useLocation, useNavigate, useParams } from 'react-router-dom';
 import backend from '../../backend';
 import { ArtifactContextType, DashboardPreference, EventSeverity, GetChangeSummaryResponse, GetJobStepRefResponse, GetLogEventResponse, LogLevel } from '../../backend/Api';
 import { CommitCache } from '../../backend/CommitCache';
-import dashboard, { StatusColor } from '../../backend/Dashboard';
+import dashboard from '../../backend/Dashboard';
 import { Markdown } from '../../base/components/Markdown';
 import { useWindowSize } from '../../base/utilities/hooks';
 import { displayTimeZone, getElapsedString } from '../../base/utilities/timeUtils';
@@ -27,6 +27,7 @@ import { TopNav } from '../TopNav';
 import { isBright } from 'horde/base/utilities/colors';
 import { JobLogSource, LogSource } from './LogSource';
 import { JobDetailsV2 } from '../jobDetailsV2/JobDetailsViewCommon';
+import { getStepMessageInfo } from '../jobDetailsV2/JobDetailsUtils';
 import { LogHandler } from './LogHandler';
 import { updateLineQuery, getQueryLine } from 'horde/base/utilities/logUtils';
 
@@ -1148,19 +1149,8 @@ const LogList: React.FC<{ logId: string }> = observer(({ logId }) => {
       jobDetails = (logSource as JobLogSource).jobDetails;
    }
 
-   let stepMessage = "";
-   let stepMessageColor = "";
-
    const report = jobDetails?.stepByLogId(logId)?.reports?.find(r => !!r.message);
-   if (report) {
-      stepMessage = report.message!;
-      if (report.severity === "Error" || report.severity === "Errors") {
-         stepMessageColor = dashboard.getStatusColors().get(StatusColor.Failure)!
-      }
-      else if (report.severity === "Warning" || report.severity === "Warnings") {
-         stepMessageColor = dashboard.getStatusColors().get(StatusColor.Warnings)!
-      }
-   }
+   const { message: stepMessage, color: stepMessageColor } = getStepMessageInfo(report);
 
    return <Stack>
       {logSourceIsJobLogSource && <IssueModalV2 issueId={query.get("issue")} popHistoryOnClose={issueHistory} />}
