Using IOptionals instead of IOptionalsSnapshot in TestHealthService constructor

#rnx
#Synced-CL 49069065

diff --git a/Engine/Source/Programs/Horde/Plugins/Build/HordeServer.Build/Jobs/TestData/TestHealthService.cs b/Engine/Source/Programs/Horde/Plugins/Build/HordeServer.Build/Jobs/TestData/TestHealthService.cs
--- a/Engine/Source/Programs/Horde/Plugins/Build/HordeServer.Build/Jobs/TestData/TestHealthService.cs
+++ b/Engine/Source/Programs/Horde/Plugins/Build/HordeServer.Build/Jobs/TestData/TestHealthService.cs
@@ -43,7 +43,7 @@ namespace HordeServer.Jobs.TestData
 		readonly ITestDataCollectionV2 _testData;
 		readonly SingletonDocument<TestHealthReportState> _state;
 		readonly IOptions<BuildServerConfig> _staticBuildConfig;
-		readonly IOptionsSnapshot<BuildConfig> _buildConfig;
+		readonly IOptions<BuildConfig> _buildConfig;
 		readonly INotificationService _notificationService;
 
 		readonly int _reportIntervalMinutes = 12 * 60; // every 12 hours
@@ -60,7 +60,7 @@ namespace HordeServer.Jobs.TestData
 		/// <param name="tracer"></param>
 		/// <param name="staticBuildConfig"></param>
 		/// <param name="buildConfig"></param>
-		public TestHealthService(IMongoService mongoService, ITestDataCollectionV2 testData, INotificationService notificationService, IClock clock, ILogger<TestHealthService> logger, Tracer tracer, IOptions<BuildServerConfig> staticBuildConfig, IOptionsSnapshot<BuildConfig> buildConfig)
+		public TestHealthService(IMongoService mongoService, ITestDataCollectionV2 testData, INotificationService notificationService, IClock clock, ILogger<TestHealthService> logger, Tracer tracer, IOptions<BuildServerConfig> staticBuildConfig, IOptions<BuildConfig> buildConfig)
 		{
 			_state = new SingletonDocument<TestHealthReportState>(mongoService);
 			_testData = testData;
