[Test Health] Fix Uri construction in slack notification

#rnx
#Synced-CL 49069065

diff --git a/Engine/Source/Programs/Horde/Plugins/Build/HordeServer.Build/Notifications/Sinks/SlackNotificationSink.cs b/Engine/Source/Programs/Horde/Plugins/Build/HordeServer.Build/Notifications/Sinks/SlackNotificationSink.cs
--- a/Engine/Source/Programs/Horde/Plugins/Build/HordeServer.Build/Notifications/Sinks/SlackNotificationSink.cs
+++ b/Engine/Source/Programs/Horde/Plugins/Build/HordeServer.Build/Notifications/Sinks/SlackNotificationSink.cs
@@ -2736,7 +2736,7 @@ namespace HordeServer.Notifications.Sinks
 			}
 
 			SlackMessage message = new SlackMessage();
-			string testHealthDetailsUrlLink = new Uri($"{_serverInfo.DashboardUrl}/test-automation?stream={report.StreamId}&view=1&health={report.TestId}").ToString();
+			string testHealthDetailsUrlLink = new Uri(_serverInfo.DashboardUrl, $"test-automation?stream={report.StreamId}&view=1&health={report.TestId}").ToString();
 
 			string title = GenerateTestHealthSlackMessageTitle(report, streamConfig, testHealthDetailsUrlLink);
 			message.Blocks.Add(new SectionBlock(title));
@@ -2774,7 +2774,7 @@ namespace HordeServer.Notifications.Sinks
 		static string GenerateTestHealthSlackMessageTitle(ITestHealthReport report, StreamConfig streamConfig, string testHealthDetailsUrlLink)
 		{
 			string state = report.PreviousState is not null && report.IsHealthy ? report.PreviousState : report.State;
-			string title = $"*[{streamConfig.Name}]* {RedBadge} <{testHealthDetailsUrlLink}|test *{report.TestName}* health> has degraded to *{state}*";
+			string title = $"*[{streamConfig.Name}]* {RedBadge} test <{testHealthDetailsUrlLink}|*{report.TestName}* health> has degraded to *{state}*";
 			if (report.IsHealthy)
 			{
 				// strike title
