Allow changing Test Health notification workflow
+ allow Team and Profile sorting
+ add owner column (sortable)
+ add intent display on the test health dashboard as a hover card
+ use median times instead of average

#rnx

#Synced-CL 49069065

diff --git a/Engine/Source/Programs/Horde/HordeDashboard/plugins/test-automation/testAutomationCommon.tsx b/Engine/Source/Programs/Horde/HordeDashboard/plugins/test-automation/testAutomationCommon.tsx
--- a/Engine/Source/Programs/Horde/HordeDashboard/plugins/test-automation/testAutomationCommon.tsx
+++ b/Engine/Source/Programs/Horde/HordeDashboard/plugins/test-automation/testAutomationCommon.tsx
@@ -59,6 +59,13 @@ export const styles = {
     stripes: {
         backgroundImage: 'repeating-linear-gradient(-45deg, rgba(255, 255, 255, .2) 25%, transparent 25%, transparent 50%, rgba(255, 255, 255, .2) 50%, rgba(255, 255, 255, .2) 75%, transparent 75%, transparent)',
     },
+    streamBadge: {
+        width: 200,
+        height: 24,
+        backgroundColor: "#8e8e8e48",
+        padding: '3px 16px',
+        borderRadius: 2
+    }
 }
 
 let _styles: any;
@@ -915,7 +922,7 @@ export const getTestUniqueMetaIdentifiers = (testMeta: Map<string, [TestNameRef,
 export const StreamSelector: React.FC<{ streams: string[], onClick: (stream: string) => void, selected?: string, disabled?: boolean }> = memo(({ streams, onClick, selected, disabled }) => {
     const menuProps: IContextualMenuProps = getStreamOptions(streams, onClick, selected);
     return <DefaultButton
-                style={{ height: 22, fontSize: 11, whiteSpace: 'nowrap', borderWidth: 0 }}
+                style={{ height: 24, fontSize: 11, whiteSpace: 'nowrap', borderWidth: 0 }}
                 menuProps={menuProps}
                 text={selected? projectStore.streamById(selected)?.fullname ?? selected : "Compare stream"}
                 disabled={disabled}
diff --git a/Engine/Source/Programs/Horde/HordeDashboard/plugins/test-automation/testAutomationHealth.tsx b/Engine/Source/Programs/Horde/HordeDashboard/plugins/test-automation/testAutomationHealth.tsx
--- a/Engine/Source/Programs/Horde/HordeDashboard/plugins/test-automation/testAutomationHealth.tsx
+++ b/Engine/Source/Programs/Horde/HordeDashboard/plugins/test-automation/testAutomationHealth.tsx
@@ -1,17 +1,17 @@
 // Copyright Epic Games, Inc. All Rights Reserved.
 
-import { DetailsList, FontIcon, IColumn, IconButton, Label, Modal, SelectionMode, Stack, Text, HoverCard, HoverCardType, IPlainCardProps, DirectionalHint, DetailsListLayoutMode, Spinner, SpinnerSize, Toggle, DatePicker, ComboBox, IComboBoxOption, FocusZone, FocusZoneTabbableElements, FocusTrapCallout, PrimaryButton, DefaultButton, TooltipHost, TooltipOverflowMode } from "@fluentui/react";
+import { DetailsList, FontIcon, IColumn, IconButton, Label, Modal, SelectionMode, Stack, Text, HoverCard, HoverCardType, IPlainCardProps, DirectionalHint, DetailsListLayoutMode, Spinner, SpinnerSize, Toggle, DatePicker, ComboBox, IComboBoxOption, FocusZone, FocusZoneTabbableElements, FocusTrapCallout, PrimaryButton, DefaultButton, TooltipHost, TooltipOverflowMode, SuggestionItemType, TooltipDelay } from "@fluentui/react";
 import { MetadataRef, TestDataHandler, TestNameRef, TestSessionResult } from "./testData";
-import { useCallback, useEffect, useMemo, useRef, useState } from "react";
+import { useCallback, useEffect, useMemo, useState } from "react";
 import { observer } from "mobx-react-lite";
-import { TestOutcome, getUsersInfo } from "./api";
-import { DomainNameColors, EditableTextField, getTestUniqueMetaIdentifiers, StreamSelector, tagColors, UserPickerModal } from "./testAutomationCommon";
+import { TestOutcome } from "./api";
+import { DomainNameColors, EditableTextField, getTestUniqueMetaIdentifiers, StreamSelector, tagColors, UserPickerModal, styles, getStatusColors } from "./testAutomationCommon";
 import { TestHistoryGraphWidget } from "./testAutomationSummary";
 import { StoreKey, userStorage } from "./testAutomationUserStorage";
 import { projectStore } from "horde/backend/ProjectStore";
-import dashboard from "horde/backend/Dashboard";
+import dashboard, { StatusColor } from "horde/backend/Dashboard";
 import { getShortNiceTime, msecToElapsed } from "horde/base/utilities/timeUtils";
-import { booleanColors, cadenceColors, cadenceOrders, healthColors, healthStars, TestCadence, TestHealth, TestHealthExplain, TestHealthItem, TestHealthTiming } from "./testHealthData";
+import { booleanColors, cadenceColors, cadenceOrders, healthColors, healthStars, TestCadence, TestHealth, TestHealthExplain, TestHealthItem, TestHealthTiming, UsersInfoCache } from "./testHealthData";
 import { GetUserResponse } from "horde/backend/Api";
 
 /// Colors
@@ -23,12 +23,47 @@ const harnessColors = new DomainNameColors();
 
 const healthBadgeBackground = () => dashboard.darktheme? "#67676748" : "#8f8f8f9e";
 const textEmphasisColor = () => dashboard.darktheme? "#ffffff75" : "#00000075";
+const borderCardColor = () => dashboard.darktheme ? "#2D2B29" : "#EDEBE9";
+
+const healthStyles = {
+    userTag: {
+        backgroundColor: booleanColors.get(false),
+        padding: '2px 9px',
+        borderRadius: 11
+    },
+    badge: {
+        padding: '3px 12px',
+        borderRadius: 4,
+        cursor: 'default'
+    }
+}
 
 enum UserType {
     Owner = "Owner",
     Customers = "Customers",
 }
 
+const onRenderIntentCard = (item: TestHealthItem): JSX.Element => {
+    return <Stack style={{pointerEvents: 'none', padding: 10, maxWidth: 300, border: `1px solid ${borderCardColor()}`}} tokens={{childrenGap: 5}}>{item.intent}</Stack>
+}
+
+const renderTestTitle = (item: TestHealthItem): JSX.Element => {
+    const maskMetaString = item.displayMeta;
+    let title = item.name;
+    if (maskMetaString) title += ` - ${maskMetaString}`;
+    const DetailsCardProps: IPlainCardProps | undefined = item.intent? {
+        onRenderPlainCard: onRenderIntentCard,
+        renderData: item,
+        directionalHint: DirectionalHint.topLeftEdge,
+        calloutProps: { isBeakVisible: true }
+    } : { onRenderPlainCard: () => null /** empty card */ };
+    return <HoverCard plainCardProps={DetailsCardProps} type={HoverCardType.plain} cardOpenDelay={300} sticky={false} styles={{host: {width: '100%'}}}>
+                <Text title={title} style={{textOverflow: 'ellipsis', overflow: 'hidden', whiteSpace: 'nowrap'}}>
+                    {item.name}{maskMetaString && <span style={{fontSize: 11}}> - {maskMetaString}</span>}
+                </Text>
+            </HoverCard>
+}
+
 const onRenderExplainCard = (explainData: TestHealthExplain): JSX.Element => {
     return <Stack style={{pointerEvents: 'none', padding: 10}} tokens={{childrenGap: 5}}>
                 {!!explainData.successRate && <Text>{explainData.successRate}% Success</Text>}
@@ -105,11 +140,7 @@ const getHealthColumns = (items: TestHealthItem[], columns: TestHealthColumn[],
                 const maskMetaString = item.displayMeta;
                 let title = item.name;
                 if (maskMetaString) title += ` - ${maskMetaString}`;
-                return <Stack onClick={() => onClickItem(item)} style={{cursor: 'pointer'}} horizontal>
-                            <Text title={title} style={{textOverflow: 'ellipsis', overflow: 'hidden', whiteSpace: 'nowrap'}}>
-                                {item.name}{maskMetaString && <span style={{fontSize: 11}}> - {maskMetaString}</span>}
-                            </Text>
-                        </Stack>
+                return <Stack onClick={() => onClickItem(item)} style={{cursor: 'pointer'}} horizontal>{renderTestTitle(item)}</Stack>
             },
             onColumnClick: (_, column) => onClickColumn(column),
             compare: (a: TestHealthItem, b: TestHealthItem) => a.name.localeCompare(b.name)
@@ -138,29 +169,51 @@ const getHealthColumns = (items: TestHealthItem[], columns: TestHealthColumn[],
             isSortedDescending: true,
             onRender: (item: TestHealthItem) => {
                 const profile = item.profile ?? "N/A";
-                return <Stack style={{backgroundColor: profileColors.get(item.profile), padding: 3, borderRadius: 4, cursor: 'default'}} horizontalAlign="center" horizontal>
+                return <Stack style={{...healthStyles.badge, backgroundColor: profileColors.get(item.profile)}} horizontalAlign="center" horizontal>
                             <Text title={profile} style={{textOverflow: 'ellipsis', overflow: 'hidden', whiteSpace: 'nowrap'}}>{profile}</Text>
                         </Stack>
             },
+            onColumnClick: (_, column) => onClickColumn(column),
             compare: (a: TestHealthItem, b: TestHealthItem) => a.profile?.localeCompare(b.profile ?? "") ?? -1
         },
         {
             key: 'team',
             name: 'Team',
-            minWidth: 70,
-            maxWidth: 70,
+            minWidth: 90,
+            maxWidth: 90,
             isResizable: true,
             isCollapsible: true,
             isSorted: false,
             isSortedDescending: true,
             onRender: (item: TestHealthItem) => {
                 const team = item.team  ?? "N/A";
-                return <Stack style={{backgroundColor: teamColors.get(item.team), padding: 3, borderRadius: 4, cursor: 'default'}} horizontalAlign="center" horizontal>
+                return <Stack style={{...healthStyles.badge, backgroundColor: teamColors.get(item.team)}} horizontalAlign="center" horizontal>
                             <Text title={team} style={{textOverflow: 'ellipsis', overflow: 'hidden', whiteSpace: 'nowrap'}}>{team}</Text>
                         </Stack>
             },
+            onColumnClick: (_, column) => onClickColumn(column),
             compare: (a: TestHealthItem, b: TestHealthItem) => a.team?.localeCompare(b.team ?? "") ?? -1
         },
+        {
+            key: 'owner',
+            name: 'Owner',
+            minWidth: 70,
+            maxWidth: 120,
+            isResizable: true,
+            isCollapsible: true,
+            isSorted: false,
+            isSortedDescending: true,
+            onRender: (item: TestHealthItem) => {
+                return <Stack style={{...healthStyles.userTag, backgroundColor: booleanColors.get(!item.owner), cursor: 'default'}} horizontalAlign="center" horizontal>
+                            <Text title={usersInfoCache.getInfo(item.owner)?.name} style={{textOverflow: 'ellipsis', overflow: 'hidden', whiteSpace: 'nowrap'}}>
+                                {!!item.owner && (usersInfoCache.getInfo(item.owner)?.name ?? <Spinner size={SpinnerSize.small}/>)}
+                                {!item.owner && "N/A"}
+                            </Text>
+                        </Stack>
+            },
+            onColumnClick: (_, column) => onClickColumn(column),
+            compare: (a: TestHealthItem, b: TestHealthItem) => a.owner?.localeCompare(b.owner ?? "") ?? -1
+        },
         {
             key: 'hasNotification',
             name: 'Notification',
@@ -190,7 +243,7 @@ const getHealthColumns = (items: TestHealthItem[], columns: TestHealthColumn[],
                 const content = <Stack horizontal verticalAlign="center" tokens={{childrenGap: 4}} wrap style={{maxWidth: 350}}>
                                     {item.nameRef?.tagRefs &&
                                         item.nameRef.tagRefs.keys().map(t => t.name).toArray().sort().map(tag =>
-                                            <Stack style={{backgroundColor: tagColors.getTagColor(tag), borderRadius: 5, padding: '1px 4px', opacity: 0.8}}>
+                                            <Stack key={`tag-${tag}`} style={{backgroundColor: tagColors.getTagColor(tag), borderRadius: 5, padding: '1px 4px', opacity: 0.8}}>
                                                 <Text>#{tag}</Text>
                                             </Stack>
                                         )
@@ -201,19 +254,18 @@ const getHealthColumns = (items: TestHealthItem[], columns: TestHealthColumn[],
                             <TooltipHost content={content} overflowMode={TooltipOverflowMode.Self}>{content}</TooltipHost>
                         </Stack>
             },
-            onColumnClick: (_, column) => onClickColumn(column),
         },
         {
             key: 'cadence',
             name: 'Cadence',
             minWidth: 75,
-            maxWidth: 100,
+            maxWidth: 120,
             isResizable: true,
             isCollapsible: true,
             isSorted: false,
             isSortedDescending: true,
             onRender: (item: TestHealthItem) => {
-                return <Stack style={{backgroundColor: cadenceColors.get(item.cadence), padding: 3, borderRadius: 4, cursor: 'default'}} horizontalAlign="center" horizontal>
+                return <Stack style={{...healthStyles.badge, backgroundColor: cadenceColors.get(item.cadence)}} horizontalAlign="center" horizontal>
                             <Text title={item.cadence} style={{textOverflow: 'ellipsis', overflow: 'hidden', whiteSpace: 'nowrap'}}>{item.cadence}</Text>
                         </Stack>
             },
@@ -404,9 +456,17 @@ const getTestHealthItems = (handler: TestDataHandler, stream?: string, test?: Te
         }
 
         // Compute timings
+        testCatastrophicDurations.sort();
+        testCompleteRunDurations.sort();
+        const midIndex1 = Math.floor(testCatastrophicDurations.length / 2);
+        const midIndex2 = Math.floor(testCompleteRunDurations.length / 2);
         const timings: TestHealthTiming = {
-            averageToInterruptSecs: testCatastrophicDurations.length === 0? 0 : testCatastrophicDurations.reduce((acc, value) => acc + value, 0) / testCatastrophicDurations.length,
-            averageToCompleteRunSecs: testCompleteRunDurations.length === 0? 0 : testCompleteRunDurations.reduce((acc, value) => acc + value, 0) / testCompleteRunDurations.length,
+            medianToInterruptSecs: testCatastrophicDurations.length === 0?
+                            0 : testCatastrophicDurations.length % 2?
+                                testCatastrophicDurations[midIndex1] : (testCatastrophicDurations[midIndex1] + testCatastrophicDurations[midIndex1 - 1]) / 2,
+            medianToCompleteRunSecs: testCompleteRunDurations.length === 0?
+                            0 : testCompleteRunDurations.length % 2?
+                                testCompleteRunDurations[midIndex2] : (testCompleteRunDurations[midIndex2] + testCompleteRunDurations[midIndex2 - 1]) / 2,
             latestCompletedRun: latestCompletedRun,
         }
 
@@ -436,33 +496,35 @@ const TestHealthTimes: React.FC<{ test: TestHealthItem }> = ({test}) => {
 
     if (!test.timings) return null;
 
-    const averageToCompleteRunString = useMemo(() => test.timings? msecToElapsed(test.timings.averageToCompleteRunSecs * 1000, true, true) : "", [test]);
-    const averageToInterruptString = useMemo(() => test.timings? msecToElapsed(test.timings.averageToInterruptSecs * 1000, true, true) : "", [test]);
+    const medianToCompleteRunString = useMemo(() => test.timings? msecToElapsed(test.timings.medianToCompleteRunSecs * 1000, true, true) : "", [test]);
+    const medianToInterruptString = useMemo(() => test.timings? msecToElapsed(test.timings.medianToInterruptSecs * 1000, true, true) : "", [test]);
     const timeoutString = useMemo(() => test.recipes? msecToElapsed(test.recipes.map(r => r.timeoutMinutes).reduce((acc, value) => value > acc? value : acc, 0 ) * 60 * 1000, true, true) : "", [test, init]);
-
-    return <Text>
-                {!!test.timings.averageToCompleteRunSecs &&
-                    <span>
-                        <span>Completed run Avg. time:</span>
-                        <span style={{color: textEmphasisColor()}}> {averageToCompleteRunString}</span>
-                    </span>
-                }                                    
-                {!test.timings.averageToCompleteRunSecs && <span>No completed run recorded</span> }
-                {!!test.timings.averageToInterruptSecs &&
-                    <span>
-                        <span> | Interrupted run Avg. time:</span>
-                        <span style={{color: textEmphasisColor()}}> {averageToInterruptString}</span>
-                    </span>
+    const colors = getStatusColors();
+
+    return <Stack horizontal tokens={{childrenGap: 15}} style={{cursor: 'default'}}>
+                <TooltipHost content="Median time to Complete a run" delay={TooltipDelay.zero}>
+                    <FontIcon iconName="Completed" style={{fontSize: 14, paddingRight: 5, color: colors.get(StatusColor.Success)}}/>
+                    {!!test.timings.medianToCompleteRunSecs && <span style={{color: textEmphasisColor()}}>{medianToCompleteRunString}</span>}
+                    {!test.timings.medianToCompleteRunSecs && <span style={{color: textEmphasisColor()}}>n/a</span> }
+                </TooltipHost>
+                {!!test.timings.medianToInterruptSecs &&
+                    <TooltipHost content="Median time to Interrupted run" delay={TooltipDelay.zero}>
+                        <FontIcon iconName="ErrorBadge" style={{fontSize: 14, paddingRight: 5, color: colors.get(StatusColor.Unspecified)}} />
+                        <span style={{color: textEmphasisColor()}}>{medianToInterruptString}</span>
+                    </TooltipHost>
                 }
                 {!!test.recipes &&
-                    <span>
-                        <span> | Max Timeout:</span>
-                        <span  style={{color: textEmphasisColor()}}> {timeoutString}</span>
-                    </span>
+                    <TooltipHost content="Max timeout" delay={TooltipDelay.zero}>
+                        <FontIcon iconName="Recent" style={{fontSize: 14, paddingRight: 5, color: colors.get(StatusColor.Waiting)}} />
+                        <span  style={{color: textEmphasisColor()}}>{timeoutString}</span>
+                    </TooltipHost>
                 }
-            </Text> 
+            </Stack> 
 }
 
+
+const usersInfoCache = new UsersInfoCache();
+
 export const TestHealthView: React.FC<{ test: TestHealthItem, handler: TestDataHandler, onUpdateItem: (item: TestHealthItem) => void }> = observer(({ test, handler, onUpdateItem }) => {
     const [streams, setStreams] = useState<{current: string, toCompare?: string}>(
         {
@@ -471,12 +533,8 @@ export const TestHealthView: React.FC<{ test: TestHealthItem, handler: TestDataH
         });
     const [testHealthToCompare, setTestHealthToCompare] = useState<TestHealthItem | undefined>(getTestHealthItems(handler, userStorage.getItem(StoreKey.HistoryToCompare), test.nameRef).find(() => true));
     const [userPicker, setUserPicker] = useState<{target?: UserType, defaultUser?: GetUserResponse}>({});
-    const userCache = useRef<Map<string | undefined, GetUserResponse> | undefined>(undefined);
 
     useEffect(() => {
-        if (!userCache.current) {
-            userCache.current = new Map<string, GetUserResponse>();
-        }
         let userIds: string[] = [];
         if (test.owner) {
             userIds.push(test.owner);
@@ -484,14 +542,8 @@ export const TestHealthView: React.FC<{ test: TestHealthItem, handler: TestDataH
         if (test.customers) {
             userIds = userIds.concat(test.customers);
         }
-        userIds = userIds.filter(id => !userCache.current!.has(id));
         if (userIds.length > 0) {
-            getUsersInfo(userIds).then((users) => {
-                users.forEach(user => {
-                    userCache.current!.set(user.id, user);
-                });
-                test.setUpdated();
-            });
+            usersInfoCache.fetchUsersInfo(userIds).then((fetchedUsers) => fetchedUsers.length > 0 && test.setUpdated());
         }
     }, [test.owner, test.customers]);
 
@@ -534,6 +586,10 @@ export const TestHealthView: React.FC<{ test: TestHealthItem, handler: TestDataH
     const profileOptions: IComboBoxOption[] = profileColors.getDomain().map(p => ({key: p, text: p, styles: {optionText: {backgroundColor: profileColors.get(p), padding: 3, borderRadius: 4, minWidth: 72}}}));
     const teamOptions: IComboBoxOption[] = teamColors.getDomain().map(p => ({key: p, text: p, styles: {optionText: {backgroundColor: teamColors.get(p), padding: 3, borderRadius: 4, minWidth: 72}}}));
     const harnessOptions: IComboBoxOption[] = harnessColors.getDomain().map(p => ({key: p, text: p}));
+    const notificationOptions: IComboBoxOption[] | undefined = projectStore.streamById(test.stream)?.workflows?.filter(w => !!w.reportChannel).map(w => ({key: w.id, text: w.id, data: w.id}));
+    if (!!notificationOptions) {
+        notificationOptions.unshift({key: 'disable', text: '- disable notification -', data: undefined});
+    }
 
     const tags = useMemo(() => test.nameRef?.tagRefs?.keys().map(t => t.name).toArray().sort(), [test]);
 
@@ -559,108 +615,124 @@ export const TestHealthView: React.FC<{ test: TestHealthItem, handler: TestDataH
                     />
                 }
                 <Stack tokens={{childrenGap: 16}} style={{position: 'relative'}}>
-                    <Stack tokens={{childrenGap: 8}} style={{padding: 14, borderRadius: 4}} >
-                        <Stack horizontal tokens={{childrenGap: 12}} verticalAlign="center">
-                            <ComboBox
-                                key={`profile-${test.updated}`}
-                                styles={{input: {backgroundColor: profileColors.get(test.profile)}, root: {backgroundColor: profileColors.get(test.profile), width: 150}}}
-                                placeholder="Profile"
-                                allowFreeform
-                                autoComplete="on"
-                                options={profileOptions}
-                                text={test.profile}
-                                onChange={(_1, _2, _3, value) => { test.profile = value; onUpdateItem(test); }}
-                            />
-                            <ComboBox
-                                key={`team-${test.updated}`}
-                                styles={{input: {backgroundColor: teamColors.get(test.team)}, root: {backgroundColor: teamColors.get(test.team), width: 150}}}
-                                placeholder="Team"
-                                allowFreeform
-                                autoComplete="on"
-                                options={teamOptions}
-                                text={test.team}
-                                onChange={(_1, _2, _3, value) => { test.team = value; onUpdateItem(test); }}
-                            />
-                            <ComboBox
-                                key={`harness-${test.updated}`}
-                                styles={{root: {width: 170}}}
-                                placeholder="Harness"
-                                allowFreeform
-                                autoComplete="on"
-                                options={harnessOptions}
-                                text={test.harness}
-                                onChange={(_1, _2, _3, value) => { test.harness = value; }}
-                            />
-                        </Stack>
-                        <Stack horizontal tokens={{childrenGap: 12}} style={{paddingLeft: 12, paddingTop: 6}} grow>
-                            <EditableTextField
-                                key={`intent-${test.updated}`}
-                                label="Intent"
-                                text={test.intent}
-                                onChange={(value) => test.intent = value}
-                                placeholder="a short description"
-                                loaded={handler.queryAuditsLoading}
-                            />
-                        </Stack>
+                    <Stack tokens={{childrenGap: 16}} style={{padding: 14}}>
+                        <Stack horizontal tokens={{childrenGap: 12}}>
+                            <Stack tokens={{childrenGap: 12}} grow style={{width: '100%'}}>
+
+                                <Stack horizontal tokens={{childrenGap: 12}} style={{paddingLeft: 8, backgroundColor: booleanColors.get(false)}}>
+                                    <EditableTextField
+                                        key={`intent-${test.updated}`}
+                                        label="Intent"
+                                        text={test.intent}
+                                        onChange={(value) => test.intent = value}
+                                        placeholder="a short description"
+                                        loaded={handler.queryAuditsLoading}
+                                    />
+                                </Stack>
 
-                        {tags &&
-                            <Stack horizontal verticalAlign="center" tokens={{childrenGap: 4}} wrap grow style={{paddingLeft: 12, paddingTop: 6}}>
-                                {tags.map(tag =>
-                                        <Stack style={{backgroundColor: tagColors.getTagColor(tag), borderRadius: 5, padding: '1px 4px', opacity: 0.8}}>
-                                            <Text>#{tag}</Text>
-                                        </Stack>
-                                    )
+                                <Stack horizontal tokens={{childrenGap: 8}}>
+                                    <ComboBox
+                                        key={`profile-${test.updated}`}
+                                        styles={{input: {backgroundColor: profileColors.get(test.profile)}, root: {backgroundColor: profileColors.get(test.profile), width: 170, height: 24}}}
+                                        placeholder="Profile"
+                                        allowFreeform
+                                        autoComplete="on"
+                                        options={profileOptions}
+                                        text={test.profile}
+                                        onChange={(_1, _2, _3, value) => { test.profile = value; onUpdateItem(test); }}
+                                    />
+                                    <ComboBox
+                                        key={`team-${test.updated}`}
+                                        styles={{input: {backgroundColor: teamColors.get(test.team)}, root: {backgroundColor: teamColors.get(test.team), width: 170, height: 24}}}
+                                        placeholder="Team"
+                                        allowFreeform
+                                        autoComplete="on"
+                                        options={teamOptions}
+                                        text={test.team}
+                                        onChange={(_1, _2, _3, value) => { test.team = value; onUpdateItem(test); }}
+                                    />
+                                    <ComboBox
+                                        key={`harness-${test.updated}`}
+                                        styles={{root: {width: 170, height: 24}}}
+                                        placeholder="Harness"
+                                        allowFreeform
+                                        autoComplete="on"
+                                        options={harnessOptions}
+                                        text={test.harness}
+                                        onChange={(_1, _2, _3, value) => { test.harness = value; }}
+                                    />
+                                </Stack>
+
+                                {tags &&
+                                    <Stack horizontal verticalAlign="center" tokens={{childrenGap: 4}} wrap grow style={{paddingLeft: 8}}>
+                                        {tags.map(tag =>
+                                                <Stack key={`tag-${tag}`} style={{backgroundColor: tagColors.getTagColor(tag), borderRadius: 5, padding: '1px 4px', opacity: 0.8}}>
+                                                    <Text>#{tag}</Text>
+                                                </Stack>
+                                            )
+                                        }
+                                    </Stack>
                                 }
-                            </Stack>
-                        }
 
-                        <Stack horizontal style={{paddingLeft: 12}}>
-                            <Stack horizontal verticalAlign="center" tokens={{childrenGap: 6}} style={{cursor: 'default'}} title="edit disabled for now">
-                                <Label>Workflow Notification: </Label>
-                                <Text style={{color:textEmphasisColor()}}>{!!test.notification? test.notification : "none"}</Text>
                             </Stack>
-                        </Stack>
-
-                        <Stack horizontal tokens={{childrenGap: 12}} style={{paddingLeft: 12}}>
-                            <Stack style={{borderRadius: 4}} verticalAlign="center">
-                                <Stack  horizontal tokens={{childrenGap: 6}} verticalAlign="center">
-                                    <Label>Owner: </Label>
+                            <Stack tokens={{childrenGap: 8}} style={{width: '100%'}}>
+
+                                <Stack horizontal>
+                                    <Stack horizontal verticalAlign="center" tokens={{childrenGap: 6}} style={{cursor: 'default'}}>
+                                        <Label style={{whiteSpace: 'nowrap'}}>Workflow Notification: </Label>
+                                        {!notificationOptions && <Text style={{color:textEmphasisColor(), cursor: 'default'}}>no workflow available in this stream</Text>}
+                                        {!!notificationOptions && <ComboBox
+                                            key={`notification-${test.updated}`}
+                                            styles={{root: {width: 150, height: 24}}}
+                                            placeholder="none"
+                                            autoComplete="on"
+                                            options={notificationOptions}
+                                            text={test.notification}
+                                            onChange={(_1, option) => { test.notification = option?.data as string | undefined; onUpdateItem(test); }}
+                                            title={test.notification}
+                                        />}
+                                    </Stack>
                                 </Stack>
-                                <Stack style={{backgroundColor: booleanColors.get(false), padding: '2px 9px', borderRadius: 11}} horizontalAlign="center" verticalAlign="center">
-                                    {!!userCache.current?.get(test.owner) &&
-                                        <Stack horizontal verticalAlign="center">
-                                            <Text style={{cursor: 'pointer', color:textEmphasisColor()}} title="change owner"
-                                                onClick={() => setUserPicker({target: UserType.Owner, defaultUser: userCache.current?.get(test.owner)})}
-                                            >{userCache.current?.get(test.owner)?.name ?? <Spinner size={SpinnerSize.xSmall} />}</Text>
-                                            <FontIcon iconName="Cancel" style={{cursor: 'pointer', fontSize: 11, paddingLeft: 6}} title="remove" onClick={() => {test.owner = undefined}}/>
+
+                                <Stack horizontal tokens={{childrenGap: 12}}>
+                                    <Stack tokens={{childrenGap: 6}} style={{borderRadius: 4}} verticalAlign="start">
+                                        <Label>Owner: </Label>
+                                        <Stack style={{...healthStyles.userTag, width: !test.owner? 46:undefined}} horizontal verticalAlign="center" horizontalAlign="center">
+                                            {!!test.owner &&
+                                                <div style={{whiteSpace: 'nowrap'}}>
+                                                    <Text style={{cursor: 'pointer', color:textEmphasisColor()}} title="change owner"
+                                                        onClick={() => setUserPicker({target: UserType.Owner, defaultUser: usersInfoCache.getInfo(test.owner)})}
+                                                    >{usersInfoCache.getInfo(test.owner)?.name ?? <Spinner size={SpinnerSize.xSmall} />}</Text>
+                                                    <FontIcon iconName="Cancel" style={{cursor: 'pointer', fontSize: 11, paddingLeft: 6}} title="remove" onClick={() => {test.owner = undefined}}/>
+                                                </div>
+                                            }
+                                            {!test.owner &&
+                                                <FontIcon iconName="AddTo" style={{cursor: 'pointer', fontSize: 15, padding: 2}} title="add an owner"
+                                                    onClick={() => setUserPicker({target: UserType.Owner})}
+                                                />
+                                            }
+                                        </Stack>
+                                    </Stack>
+                                    <Stack tokens={{childrenGap: 6}} style={{borderRadius: 4}} verticalAlign="start">
+                                        <Label>Customers: </Label>
+                                        <Stack horizontal tokens={{childrenGap: 6}} wrap>
+                                            {!!test.customers &&
+                                                test.customers.map(user => {
+                                                    return <Stack key={`c-${user}`} style={healthStyles.userTag} verticalAlign="center" horizontal>
+                                                                <Text style={{color:textEmphasisColor(), whiteSpace: 'nowrap'}}>{usersInfoCache.getInfo(user)?.name ?? <Spinner size={SpinnerSize.xSmall} />}</Text>
+                                                                <FontIcon iconName="Cancel" style={{cursor: 'pointer', fontSize: 11, paddingLeft: 6}} title="remove" onClick={() => {test.customers = test.customers?.filter(c => c !== user)}}/>
+                                                            </Stack>
+                                                })
+                                            }
+                                            <Stack style={{...healthStyles.userTag, width: 46}} verticalAlign="center" horizontalAlign="center">
+                                                <FontIcon iconName="AddTo" style={{cursor: 'pointer', fontSize: 15, padding: 2}} title="add a customer"
+                                                    onClick={() => setUserPicker({target: UserType.Customers})}
+                                                />
+                                            </Stack>
                                         </Stack>
-                                    }
-                                    {!userCache.current?.get(test.owner) &&
-                                        <FontIcon iconName="AddTo" style={{cursor: 'pointer', fontSize: 15, padding: 2}} title="add an owner"
-                                            onClick={() => setUserPicker({target: UserType.Owner})}
-                                        />
-                                    }
-                                </Stack>
-                            </Stack>
-                            <Stack style={{borderRadius: 4}} verticalAlign="center">
-                                <Stack  horizontal tokens={{childrenGap: 6}} verticalAlign="center">
-                                    <Label>Customers: </Label>
-                                </Stack>
-                                <Stack horizontal tokens={{childrenGap: 6}}>
-                                    {!!test.customers &&
-                                        test.customers.map(user => {
-                                            return <Stack style={{backgroundColor: booleanColors.get(false), padding: '2px 9px', borderRadius: 11}} verticalAlign="center" horizontal>
-                                                        <Text style={{color:textEmphasisColor()}}>{userCache.current?.get(user)?.name ?? <Spinner size={SpinnerSize.xSmall} />}</Text>
-                                                        <FontIcon iconName="Cancel" style={{cursor: 'pointer', fontSize: 11, paddingLeft: 6}} title="remove" onClick={() => {test.customers = test.customers?.filter(c => c !== user)}}/>
-                                                    </Stack>
-                                        })
-                                    }
-                                    <Stack style={{backgroundColor: booleanColors.get(false), padding: '2px 9px', borderRadius: 11, width: 46}} verticalAlign="center" horizontalAlign="center">
-                                        <FontIcon iconName="AddTo" style={{cursor: 'pointer', fontSize: 15, padding: 2}} title="add a customer"
-                                            onClick={() => setUserPicker({target: UserType.Customers})}
-                                        />
                                     </Stack>
                                 </Stack>
+
                             </Stack>
                         </Stack>
 
@@ -700,11 +772,11 @@ export const TestHealthView: React.FC<{ test: TestHealthItem, handler: TestDataH
                     <Stack style={{paddingLeft: 16}} tokens={{childrenGap: 16}}>
                         <Stack style={{paddingLeft: 12}} tokens={{childrenGap: 16}}>
                             <Stack horizontal tokens={{childrenGap: 12}} verticalAlign="center">
-                                <Stack style={{width: 200, backgroundColor: "#8e8e8e48", padding: '3px 12px', borderRadius: 3}} horizontalAlign="center">
+                                <Stack style={styles.streamBadge} horizontalAlign="center">
                                     <Text variant="small">{projectStore.streamById(streams.current)?.fullname ?? streams.current}</Text>
                                 </Stack>
                                 <Stack style={{cursor: 'default'}}>{renderHealthBadge(test, 75)}</Stack>
-                                <Stack style={{backgroundColor: cadenceColors.get(test.cadence), padding:  '3px 12px', borderRadius: 4, cursor: 'default'}} verticalAlign="center">
+                                <Stack style={{...healthStyles.badge, backgroundColor: cadenceColors.get(test.cadence)}} verticalAlign="center">
                                     <Text title="Cadence">{test.cadence}</Text>
                                 </Stack>
                                 {test.timings && <TestHealthTimes test={test} />}
@@ -721,7 +793,7 @@ export const TestHealthView: React.FC<{ test: TestHealthItem, handler: TestDataH
                                 </Stack>
                                 {!!testHealthToCompare && <Stack style={{cursor: 'default'}}>{renderHealthBadge(testHealthToCompare, 75)}</Stack>}
                                 {!!testHealthToCompare && 
-                                    <Stack style={{backgroundColor: cadenceColors.get(testHealthToCompare.cadence), padding:  '3px 12px', borderRadius: 4, cursor: 'default'}} verticalAlign="center">
+                                    <Stack style={{...healthStyles.badge, backgroundColor: cadenceColors.get(testHealthToCompare.cadence)}} verticalAlign="center">
                                         <Text title="Cadence">{testHealthToCompare.cadence}</Text>
                                     </Stack>
                                 }
@@ -852,6 +924,8 @@ const TestHealthDetailsModal: React.FC<{ test: TestHealthItem, handler: TestData
         </Modal>
 });
 
+const getTestHealthItemKey = (item: TestHealthItem) => `test-health-${item.nameRef.id}`;
+
 export const TestAutomationHealth: React.FC<{ handler: TestDataHandler}> = observer(({ handler }) => {
     const [healthDetailsState, setHealthDetailsState] = useState<{test?: TestHealthItem}>({});
     const [healthListState, setHealthListState] = useState<TestHealthListState>({items: [], columns: []});
@@ -865,7 +939,7 @@ export const TestAutomationHealth: React.FC<{ handler: TestDataHandler}> = obser
         harnessColors.setDomain(Array.from(harnesses).sort())
     }, []);
 
-    const initHealthListData = useCallback(() => {
+    const initHealthListData = useCallback((): [TestHealthItem[], TestHealthColumn[]] => {
         const initItems = getTestHealthItems(handler, handler.stream);
         const initColumns = getHealthColumns(initItems, healthListState.columns, (item) => {
             setHealthDetailsState({test: item});
@@ -879,7 +953,7 @@ export const TestAutomationHealth: React.FC<{ handler: TestDataHandler}> = obser
             columns: initColumns
         });
 
-        return initItems;
+        return [initItems, initColumns];
     }, []);
 
     const forceListRefresh = () => {
@@ -893,10 +967,12 @@ export const TestAutomationHealth: React.FC<{ handler: TestDataHandler}> = obser
     // initiate data on filtering change
     useEffect(() => {
 
-        const items = initHealthListData();
+        const [items, columns] = initHealthListData();
 
         // load audit async
-        handler.queryAudits(items.map(i => i.nameRef));
+        handler.queryAudits(items.map(i => i.nameRef))
+                    .then(() => usersInfoCache.fetchUsersInfo(items.map(i => i.owner).filter(i => !!i) as string[]))
+                    .then((fetchedUsers) => fetchedUsers.length > 0 && setHealthListState({items: items.slice(), columns: columns}));
 
     }, [handler.filterUpdated]);
 
@@ -924,16 +1000,20 @@ export const TestAutomationHealth: React.FC<{ handler: TestDataHandler}> = obser
                         onUpdateItem={(item: TestHealthItem) => { forceListRefresh(); }}
                     />
                 }
-                <DetailsList
-                    items={healthListState.items}
-                    compact={true}
-                    columns={healthListState.columns}
-                    isHeaderVisible={true}
-                    selectionMode={SelectionMode.none}
-                    layoutMode={DetailsListLayoutMode.justified}
-                />
+                { !handler.queryAuditsLoading &&
+                    <DetailsList
+                        items={healthListState.items}
+                        getKey={getTestHealthItemKey}
+                        compact={true}
+                        columns={healthListState.columns}
+                        isHeaderVisible={true}
+                        selectionMode={SelectionMode.none}
+                        layoutMode={DetailsListLayoutMode.justified}
+                    />
+                }
                 { handler.queryAuditsLoading &&
-                    <Stack horizontalAlign='center' style={{ width: "100%", position: "absolute", top: 0 }} tokens={{childrenGap: 8}}>
+                    <Stack horizontalAlign='center' style={{ width: "100%", position: "absolute", top: 24 }} tokens={{childrenGap: 8}}>
+                        <Text style={{ fontSize: 24 }}>Loading Data</Text>
                         <Spinner size={SpinnerSize.large} />
                     </Stack>
                 }
diff --git a/Engine/Source/Programs/Horde/HordeDashboard/plugins/test-automation/testAutomationPhaseView.tsx b/Engine/Source/Programs/Horde/HordeDashboard/plugins/test-automation/testAutomationPhaseView.tsx
--- a/Engine/Source/Programs/Horde/HordeDashboard/plugins/test-automation/testAutomationPhaseView.tsx
+++ b/Engine/Source/Programs/Horde/HordeDashboard/plugins/test-automation/testAutomationPhaseView.tsx
@@ -131,7 +131,7 @@ const TestPhaseHeader: React.FC<{ phase: TestPhaseStatus, expanded: boolean, tog
                         <Text style={{fontSize: 11}}>{phase.outcome}</Text>
                         {!!phase.tags &&
                             <Stack horizontal verticalAlign="center" tokens={{childrenGap: 6}} style={{paddingLeft: 6}}>
-                                {phase.tags.map(tag => <Stack style={{backgroundColor: tagColors.getTagColor(tag), borderRadius: 5, padding: '0px 4px', opacity: 0.8}}><Text style={{fontSize: 11, whiteSpace: 'nowrap'}}>#{tag}</Text></Stack>)}
+                                {phase.tags.map(tag => <Stack key={`tag-${tag}`} style={{backgroundColor: tagColors.getTagColor(tag), borderRadius: 5, padding: '0px 4px', opacity: 0.8}}><Text style={{fontSize: 11, whiteSpace: 'nowrap'}}>#{tag}</Text></Stack>)}
                             </Stack>
                         }
                     </Stack>
diff --git a/Engine/Source/Programs/Horde/HordeDashboard/plugins/test-automation/testAutomationSummary.tsx b/Engine/Source/Programs/Horde/HordeDashboard/plugins/test-automation/testAutomationSummary.tsx
--- a/Engine/Source/Programs/Horde/HordeDashboard/plugins/test-automation/testAutomationSummary.tsx
+++ b/Engine/Source/Programs/Horde/HordeDashboard/plugins/test-automation/testAutomationSummary.tsx
@@ -9,7 +9,7 @@ import dashboard, { StatusColor } from "horde/backend/Dashboard";
 import { getShortNiceTime } from "horde/base/utilities/timeUtils";
 import { getHordeStyling } from "horde/styles/Styles";
 import { projectStore } from 'horde/backend/ProjectStore';
-import { SessionStatusBar, SessionValues, StatusBar, StatusBarStack, MultiOptionChooser, getTestSessionStatusColor, getStatusColors, statusTexts, getTestUniqueMetaIdentifiers, StreamSelector } from "./testAutomationCommon";
+import { SessionStatusBar, SessionValues, StatusBar, StatusBarStack, MultiOptionChooser, getTestSessionStatusColor, getStatusColors, statusTexts, getTestUniqueMetaIdentifiers, StreamSelector, styles } from "./testAutomationCommon";
 import { TestHistoryGraph } from "./testHistoryGraph";
 import { StoreKey, userStorage } from "./testAutomationUserStorage";
 
@@ -71,7 +71,6 @@ export const TestHistoryView: React.FC<{ test: TestNameRef, handler: TestDataHan
     }, [streams]);
 
     const availableStreams = handler.availableStreams.filter((s) => s !== streams.current);
-    let anyFailed = false;
 
     // collecting metadata
     let metadataSet = new Set(handler.getStatusStream(streams.current)!.tests.get(test)!.getMetadata());
@@ -91,31 +90,23 @@ export const TestHistoryView: React.FC<{ test: TestNameRef, handler: TestDataHan
         .map((meta) => ({key: meta.id, text: meta.getValuesExcept(commonMetaKeys).join(' / ')}))
         .toArray().sort((a, b) => a.text.localeCompare(b.text));
 
-    const scolors = getStatusColors();
     const streamViews = (!!availableStreams.length? [streams.current, streams.toCompare] : [streams.current]).map(stream => {
 
         const streamName = projectStore.streamById(stream)?.fullname ?? stream;
 
-        let anyMetaFailed = false;
         const status = stream && handler.getStatusStream(stream)?.tests.get(test);
-        if (!!status) {
-            for (const [, lastMetaSession] of status!.getMetadataLastSessions()) {
-                if (lastMetaSession?.outcome !== TestOutcome.Success) {
-                    anyMetaFailed = true;
-                    anyFailed = true;
-                }
-            }
-        }
 
         return <Stack key={`test_history_view_${stream}_${test.id}`}>
                 { !!stream &&
-                    <Stack horizontal>
-                        <Stack className="horde-no-darktheme" style={{ paddingTop: 2, paddingRight: 3 }}>
-                            <FontIcon style={{ fontSize: 11, color: !!status? anyMetaFailed ? scolors.get(StatusColor.Failure)! : scolors.get(StatusColor.Success)! : scolors.get(StatusColor.Unspecified) }} iconName="Square" />
-                        </Stack>
-                        <Stack>
-                            <Text style={{ fontSize: 11 }}>{streamName}</Text>
-                        </Stack>
+                    <Stack style={{width: 200}}>
+                        { stream === streams.current &&
+                            <Stack style={styles.streamBadge} horizontalAlign="center">
+                                <Text variant="small">{streamName}</Text>
+                            </Stack>
+                        }
+                        { stream !== streams.current &&
+                            <StreamSelector streams={availableStreams} selected={streams.toCompare} onClick={onClickStream} />
+                        }
                     </Stack>
                 }
                 { !!status &&
@@ -140,9 +131,6 @@ export const TestHistoryView: React.FC<{ test: TestNameRef, handler: TestDataHan
             <Stack className={hordeClasses.raised}>
                 <Stack>
                     <Stack horizontal>
-                        <Stack className="horde-no-darktheme" style={{ paddingTop: 3, paddingRight: 4 }}>
-                            <FontIcon style={{ color: anyFailed ? scolors.get(StatusColor.Failure)! : scolors.get(StatusColor.Success)! }} iconName="Square" />
-                        </Stack>
                         <Stack>
                             <Text style={{ fontFamily: "Horde Open Sans Semibold" }} variant='medium'>{name}</Text>
                         </Stack>
@@ -150,19 +138,21 @@ export const TestHistoryView: React.FC<{ test: TestNameRef, handler: TestDataHan
                 </Stack>
                 <Stack horizontal tokens={{childrenGap: 15}} style={{ paddingBottom: 12, paddingTop: 5 }}>
                     <Stack style={{ width: 170 }}>
-                        <MultiOptionChooser style={{height: 22, fontSize: 11}} placeholder="Filter by metadata"
+                        <MultiOptionChooser style={{height: 24, fontSize: 11}} placeholder="Filter by metadata"
                             options={metaOptions} initialSelection={selectedMeta.keys().toArray()} disabled={metaOptions.length <= 1}
                             updateKeys={(keys) => setSelectedMeta(new Set(keys))} />
                     </Stack>
-                    <Stack>
-                        <StreamSelector streams={availableStreams} selected={streams.toCompare} disabled={availableStreams.length === 0} onClick={onClickStream} />
-                    </Stack>
                 </Stack>
                 <Stack style={{ paddingLeft: 12 }} tokens={{ childrenGap: 4 }}>
                     {streamViews}
+                    {!handler.subQueryLoading && !streams.toCompare &&
+                        <Stack style={{width: 200, margin: 0}}>
+                            <StreamSelector streams={availableStreams} selected={streams.toCompare} disabled={availableStreams.length === 0} onClick={onClickStream} />
+                        </Stack>
+                    }
                 </Stack>
                 {handler.subQueryLoading &&
-                    <Stack horizontalAlign='center' style={{ padding: 12}} tokens={{ childrenGap: 10 }}>
+                    <Stack horizontalAlign='center' horizontal style={{...styles.streamBadge, marginLeft: 12, cursor: 'default' }} tokens={{ childrenGap: 10 }}>
                         <Text style={{ fontSize: 12 }}>Loading Data</Text>
                         <Spinner size={SpinnerSize.small} />
                     </Stack>
diff --git a/Engine/Source/Programs/Horde/HordeDashboard/plugins/test-automation/testHealthData.ts b/Engine/Source/Programs/Horde/HordeDashboard/plugins/test-automation/testHealthData.ts
--- a/Engine/Source/Programs/Horde/HordeDashboard/plugins/test-automation/testHealthData.ts
+++ b/Engine/Source/Programs/Horde/HordeDashboard/plugins/test-automation/testHealthData.ts
@@ -3,6 +3,8 @@
 import { action, makeObservable, observable } from "mobx";
 import { MetadataRef, TestAudit, TestDataHandler, TestNameRef, TestRecipeRef } from "./testData";
 import { StoreKey, userStorage } from "./testAutomationUserStorage";
+import { getUsersInfo } from "./api";
+import { GetUserResponse } from "horde/backend/Api";
 
 export enum TestHealth {
     Healthy = "Healthy",
@@ -90,8 +92,8 @@ export type TestHealthExplain = {
 }
 
 export type TestHealthTiming = {
-    averageToInterruptSecs: number;
-    averageToCompleteRunSecs: number;
+    medianToInterruptSecs: number;
+    medianToCompleteRunSecs: number;
     latestCompletedRun: Date | undefined;
 }
 
@@ -282,6 +284,24 @@ export class TestHealthItem {
         return this.audit?.notifications?.get(this.stream);
     }
 
+    set notification(value: string | undefined) {
+        const audit = this.getOrStartTransaction();
+        if (!value) {
+            if (audit.notifications) {
+                audit.notifications.delete(this.stream);
+                if (audit.notifications.size === 0) {
+                    audit.notifications = undefined;
+                }
+            }
+        } else {
+            if (!audit.notifications) {
+                audit.notifications = new Map();
+            }            
+            audit.notifications.set(this.stream, value);
+        }
+        this.setUpdated()
+    }
+
     linkRecipeIds(ids: Iterable<string>) {
         this.recipeIds = Array.from(ids);
     }
@@ -308,3 +328,27 @@ export class TestHealthItem {
         return this.handler.queryRecipesLoading;
     }
 }
+
+export class UsersInfoCache {
+    constructor() {
+        this.cache = new Map();
+    }
+
+    async fetchUsersInfo(userIds: string[]) {
+        const missingUserIds = (new Set(userIds)).keys().filter(id => !this.cache.has(id)).toArray();
+        if (missingUserIds.length > 0) {
+            await getUsersInfo(missingUserIds).then((users) => {
+                users.forEach(user => this.cache.set(user.id, user));
+            });
+            missingUserIds.filter(id => !this.cache.has(id)).forEach(id => this.cache.set(id, {id: id, name: "unknown"}));
+        }
+
+        return missingUserIds.map(id => this.cache.get(id));
+    }
+
+    getInfo(userId?: string) {
+        return this.cache.get(userId);
+    } 
+
+    private cache : Map<string | undefined, GetUserResponse>;
+}
diff --git a/Engine/Source/Programs/Horde/HordeDashboard/src/backend/Api.ts b/Engine/Source/Programs/Horde/HordeDashboard/src/backend/Api.ts
--- a/Engine/Source/Programs/Horde/HordeDashboard/src/backend/Api.ts
+++ b/Engine/Source/Programs/Horde/HordeDashboard/src/backend/Api.ts
@@ -2341,6 +2341,11 @@ export type WorkflowConfig = {
 	/** Name of the tab to post summary data to */
 	summaryTab?: string;
 
+	/** Channel where reports are sent */
+	reportChannel?: string;
+
+	/** Channel where items to triage are sent */
+	triageChannel?: string;
 }
 
 /**Response describing a stream */
