[Horde][BuildHealth][Dashboard] Time formatting pass (user pref based - UTC & 24h clock)

#Synced-CL 48916881

diff --git a/Engine/Source/Programs/Horde/HordeDashboard/src/components/buildhealth/BuildHealthFilterModal.tsx b/Engine/Source/Programs/Horde/HordeDashboard/src/components/buildhealth/BuildHealthFilterModal.tsx
--- a/Engine/Source/Programs/Horde/HordeDashboard/src/components/buildhealth/BuildHealthFilterModal.tsx
+++ b/Engine/Source/Programs/Horde/HordeDashboard/src/components/buildhealth/BuildHealthFilterModal.tsx
@@ -8,6 +8,7 @@ import dashboard from "horde/backend/Dashboard";
 import { EditableBuildHealthFilter } from "./BuildHealthFilterComponent";
 import { getHordeStyling } from "horde/styles/Styles";
 import { getHordeTheme } from "horde/styles/theme";
+import { getDateTimeString } from "./stepoutcome/StepOutcomeUtilities";
 
 // #region -- Private Helpers --
 
@@ -229,7 +230,7 @@ export const BuildHealthFilterModal: React.FC<BuildHealthFilterModalProps> = ({
                         color: theme.semanticColors.bodyBackgroundHovered,
                     }
                 }}>
-                    Filter Last Updated: {buildHealthFilter.lastUpdated?.toUTCString()}
+                    Filter Last Updated: {getDateTimeString(buildHealthFilter.lastUpdated)}
                 </Label>
             )
             }
diff --git a/Engine/Source/Programs/Horde/HordeDashboard/src/components/buildhealth/stepoutcome/StepOutcome.tsx b/Engine/Source/Programs/Horde/HordeDashboard/src/components/buildhealth/stepoutcome/StepOutcome.tsx
--- a/Engine/Source/Programs/Horde/HordeDashboard/src/components/buildhealth/stepoutcome/StepOutcome.tsx
+++ b/Engine/Source/Programs/Horde/HordeDashboard/src/components/buildhealth/stepoutcome/StepOutcome.tsx
@@ -8,7 +8,7 @@ import { StepOutcomeD3Table } from "./StepOutcomeD3Component";
 import { prettyPrintStepOutcomeFilters, prettyPrintStepOutcomeFiltersSummary, StepOutcomeDataHandler, StepOutcomeFilters, StepOutcomeViewOptions } from "./StepOutcomeDataHandler";
 import { AbortedStr, OutcomeTableEntry, isLabelOutcome, isStepOutcome, StepOutcomeTableEntry } from "./StepOutcomeDataTypes";
 import { NOT_APPLICABLE_LABEL, sharedBorderStyle } from "./StepOutcomeSharedUIComponents";
-import { isTerminalState, toDate } from "./StepOutcomeUtilities";
+import { getDateTimeString, isTerminalState, toDate } from "./StepOutcomeUtilities";
 import { getHordeTheme } from "horde/styles/theme";
 
 // #region -- View Components --
@@ -67,8 +67,8 @@ const StepOutcomeFocusModal: React.FC<StepOutcomeFocusModalProps> = ({ item, onD
     }
 
     const details: (string | JSX.Element)[][] = [
-        ['Date (Job)', item.jobCreateTime ? new Date(item.jobCreateTime).toLocaleString() : NOT_APPLICABLE_LABEL],
-        ['Date (Step)', startDate ? startDate.toLocaleString() : NOT_APPLICABLE_LABEL],
+        ['Date (Job)', item.jobCreateTime ? getDateTimeString(new Date(item.jobCreateTime)): NOT_APPLICABLE_LABEL],
+        ['Date (Step)', startDate ? getDateTimeString(startDate) : NOT_APPLICABLE_LABEL],
         ['Stream', item.streamId],
         ['Code Change', item.change],
         ['Job', item.jobName],
@@ -236,7 +236,7 @@ const StepOutcomeTableComponent: React.FC<StepOutcomeTableProps> = observer(({ h
                             <Stack horizontalAlign="center" styles={{ root: { paddingTop: 8, paddingBottom: 8 } }}>
                                 <Text styles={{ root: { fontSize: "large", fontWeight: 'bold' } }}>{header}</Text>
                                 <Text styles={{ root: { fontSize: "small", fontWeight: 'bold' } }} >{subHeader} - {subSubHeader}</Text>
-                                <Text styles={{ root: { color: theme.semanticColors.bodyBackgroundHovered } }}>Last refresh: {handler.lastRefreshDate?.toLocaleTimeString([], { hour: '2-digit', minute: '2-digit', second: '2-digit', timeZoneName: 'short' })}</Text>
+                                <Text styles={{ root: { color: theme.semanticColors.bodyBackgroundHovered } }}>Last refresh: {getDateTimeString(handler.lastRefreshDate)}</Text>
                             </Stack>
                             <StepOutcomeD3Table
                                 stepOutcomeTable={handler.StepOutcomeTableData}
diff --git a/Engine/Source/Programs/Horde/HordeDashboard/src/components/buildhealth/stepoutcome/StepOutcomeD3Component.tsx b/Engine/Source/Programs/Horde/HordeDashboard/src/components/buildhealth/stepoutcome/StepOutcomeD3Component.tsx
--- a/Engine/Source/Programs/Horde/HordeDashboard/src/components/buildhealth/stepoutcome/StepOutcomeD3Component.tsx
+++ b/Engine/Source/Programs/Horde/HordeDashboard/src/components/buildhealth/stepoutcome/StepOutcomeD3Component.tsx
@@ -9,7 +9,7 @@ import { AbortedStr, ChangeHeader, ColumnHeader, encodeStepNameFromStepNameHeade
 import { getCellAccents, getCellClasses, getCellSortClasses, getCellStyle, getColorRecords, IncrementalRefreshIndicatorPanel, NOT_APPLICABLE_LABEL, resetCellClasses, sharedBorderStyle, headerStyles } from "./StepOutcomeSharedUIComponents";
 import { getHordeTheme } from "horde/styles/theme";
 import { observer } from "mobx-react-lite";
-import { isTerminalState, toLocalISOString } from "./StepOutcomeUtilities";
+import { getDateTimeString, isTerminalState } from "./StepOutcomeUtilities";
 
 // #region -- Styles --
 
@@ -460,8 +460,7 @@ function GenerateD3Table(stepOutcomeTable: StepOutcomeTable, handler: StepOutcom
                     const original = new Date(changeHeader.date!);
 
                     generateTableTooltipRow(tooltipContent, "Change", changeHeader.change.toString());
-                    generateTableTooltipRow(tooltipContent, "Date (UTC)", original.toISOString());
-                    generateTableTooltipRow(tooltipContent, "Date (Local)", toLocalISOString(original));
+                    generateTableTooltipRow(tooltipContent, "Date", getDateTimeString(original));
 
                     attachTargettedToolTip(this as HTMLElement, scaffold.current!.tooltip, tooltipContent);
                 }
@@ -663,7 +662,7 @@ function GenerateD3Table(stepOutcomeTable: StepOutcomeTable, handler: StepOutcom
                                     tooltipContent.style.paddingLeft = "4px";
 
                                     generateTableTooltipRow(tooltipContent, "Change", entry.change !== undefined ? entry.change.toString() : "-");
-                                    generateTableTooltipRow(tooltipContent, "Date", `(UTC) ${original.toISOString() ?? "-"}`);
+                                    generateTableTooltipRow(tooltipContent, "Date", `${getDateTimeString(original) ?? "-"}`);
 
                                     if (isStepOutcome(entry)) {
                                         generateTableTooltipRow(tooltipContent, "Duration", `${entry.getDurationString?.() ?? "-"}`);
diff --git a/Engine/Source/Programs/Horde/HordeDashboard/src/components/buildhealth/stepoutcome/StepOutcomeUtilities.ts b/Engine/Source/Programs/Horde/HordeDashboard/src/components/buildhealth/stepoutcome/StepOutcomeUtilities.ts
--- a/Engine/Source/Programs/Horde/HordeDashboard/src/components/buildhealth/stepoutcome/StepOutcomeUtilities.ts
+++ b/Engine/Source/Programs/Horde/HordeDashboard/src/components/buildhealth/stepoutcome/StepOutcomeUtilities.ts
@@ -2,6 +2,9 @@
 
 import { GetJobResponse, GetJobStepRefResponse, JobStepBatchError, JobStepOutcome, JobStepState } from "horde/backend/Api";
 import { OutcomeTableEntry, ReadyStr, RunningStr, WaitingStr } from "./StepOutcomeDataTypes";
+import moment from "moment";
+import { displayTimeZone } from "horde/base/utilities/timeUtils";
+import dashboard from "horde/backend/Dashboard";
 
 export const DURATION_NOT_SET: number = -1;
 
@@ -91,22 +94,22 @@ export function summarizeJobFromJobStepRef(jobStepRefs: GetJobStepRefResponse[],
 }
 
 /**
- * Produces a ISO compliant string, in the provided locale.
- * @param date The date to convert.
- * @returns The string representation of the date, in ISO standard.
+ * Helper that will return the date-time string formatted according to user preferences (UTC vs Local; 12h vs 24h).
+ * @param date THe date object to consider
+ * @returns The formatted string.
  */
-export function toLocalISOString(date: Date): string {
-    const pad = (n: number) => n.toString().padStart(2, "0");
-
-    return (
-        date.getFullYear() +
-        "-" + pad(date.getMonth() + 1) +
-        "-" + pad(date.getDate()) +
-        "T" + pad(date.getHours()) +
-        ":" + pad(date.getMinutes()) +
-        ":" + pad(date.getSeconds()) +
-        "." + date.getMilliseconds().toString().padStart(3, "0")
-    );
+export function getDateTimeString(date: Date | null | undefined): string {
+    if (date === null || date === undefined) {
+        return "-";
+    }
+
+    const displayTime = moment(date).tz(displayTimeZone());
+    const format = dashboard.display24HourClock
+        ? "MMM D HH:mm:ss z"
+        : "MMM D h:mm:ss A z";
+
+    let displayTimeStr = displayTime.format(format);
+    return displayTimeStr;
 }
 
 /**
