Horde: JobDriver Perforce updates

* Ensure PerforceMaterializer has support for and uses maxFileConcurrency and maxFileScanConcurrency as the ManagedWorkspace, ensures cleaning can run file operations in parallel
* PerforceExecutor - conforms now obey the maxFileConcurrency option to ensure deletes are run in parallel as well

diff --git a/Engine/Source/Programs/Horde/Drivers/JobDriver/Execution/PerforceExecutor.cs b/Engine/Source/Programs/Horde/Drivers/JobDriver/Execution/PerforceExecutor.cs
--- a/Engine/Source/Programs/Horde/Drivers/JobDriver/Execution/PerforceExecutor.cs
+++ b/Engine/Source/Programs/Horde/Drivers/JobDriver/Execution/PerforceExecutor.cs
@@ -6,9 +6,9 @@ using System.Web;
 using EpicGames.Core;
 using EpicGames.Perforce;
 using EpicGames.Perforce.Managed;
-using JobDriver.Utility;
 using Horde.Common.Rpc;
 using HordeCommon.Rpc.Messages;
+using JobDriver.Utility;
 using Microsoft.Extensions.Logging;
 using OpenTelemetry.Trace;
 
@@ -274,12 +274,16 @@ namespace JobDriver.Execution
 			try
 			{
 				// Set up all the workspaces
+				int maxFileConcurrency = 1;
 				foreach (RpcAgentWorkspace pendingWorkspace in pendingWorkspaces)
 				{
 					ManagedWorkspaceOptions options = WorkspaceInfo.GetMwOptions(pendingWorkspace);
 
 					WorkspaceInfo workspace = await WorkspaceInfo.CreateWorkspaceInfoAsync(pendingWorkspace, rootDir, options, tracer, logger, cancellationToken);
 					workspaces.Add(workspace);
+
+					// Also collect the max set of files to delete in parallel..
+					maxFileConcurrency = Int32.Max(maxFileConcurrency, options.MaxFileConcurrency);
 				}
 
 				// Find all the unique Perforce servers
@@ -348,14 +352,11 @@ namespace JobDriver.Execution
 				if (DirectoryReference.Exists(rootDir))
 				{
 					// Delete all the files in the root
-					foreach (FileInfo file in rootDir.ToDirectoryInfo().EnumerateFiles())
-					{
-						FileUtils.ForceDeleteFile(file);
-					}
+					logger.LogInformation("Deleting files at the root of {RootDir}", rootDir);
+					Parallel.ForEach(rootDir.ToDirectoryInfo().EnumerateFiles(), new() { MaxDegreeOfParallelism = maxFileConcurrency }, FileUtils.ForceDeleteFile);
 
 					// Build a set of directories to protect
-					HashSet<DirectoryReference> protectDirs = new HashSet<DirectoryReference>();
-					protectDirs.Add(DirectoryReference.Combine(rootDir, "Leases")); // Current lease may be writing a log here
+					HashSet<DirectoryReference> protectDirs = [DirectoryReference.Combine(rootDir, "Leases")];// Current lease may be writing a log here
 					if (!removeUntrackedFiles)
 					{
 						protectDirs.Add(DirectoryReference.Combine(rootDir, "Temp"));
@@ -374,6 +375,7 @@ namespace JobDriver.Execution
 						else
 						{
 							logger.LogInformation("Deleting directory {DeleteDir}", dir);
+							Parallel.ForEach(dir.ToDirectoryInfo().EnumerateFiles("*", SearchOption.AllDirectories), new() { MaxDegreeOfParallelism = maxFileConcurrency }, FileUtils.ForceDeleteFile);
 							FileUtils.ForceDeleteDirectory(dir);
 						}
 					}
diff --git a/Engine/Source/Programs/Horde/Drivers/JobDriver/Execution/PerforceMaterializer.cs b/Engine/Source/Programs/Horde/Drivers/JobDriver/Execution/PerforceMaterializer.cs
--- a/Engine/Source/Programs/Horde/Drivers/JobDriver/Execution/PerforceMaterializer.cs
+++ b/Engine/Source/Programs/Horde/Drivers/JobDriver/Execution/PerforceMaterializer.cs
@@ -24,17 +24,39 @@ namespace JobDriver.Execution;
 /// <param name="SyncThreads">Number of threads to use when syncing a Perforce workspace</param>
 /// <param name="SyncBatch">How many files in each sync batch</param>
 /// <param name="SyncBatchSize">Size of each sync batch</param>
+/// <param name="MaxFileConcurrency">Maximum number of concurrent file system operations (copying, moving, deleting)</param>
+/// <param name="MaxFileScanConcurrency">Maximum number of current file system scaning operations (file system discovery, traversing)</param>
 /// <param name="UseModTime">Whether to use ModTime flag on the workspace</param>
 public record PerforceMaterializerOptions(
 	string DirPath,
 	RpcAgentWorkspace AgentWorkspace,
-	string? HostName = null,
-	IReadOnlyDictionary<string, string>? EnvVars = null,
-	int SyncThreads = -1,
-	int SyncBatch = -1,
-	int SyncBatchSize = -1,
+	string? HostName,
+	IReadOnlyDictionary<string, string>? EnvVars,
+	int SyncThreads,
+	int SyncBatch,
+	int SyncBatchSize,
+	int MaxFileConcurrency,
+	int MaxFileScanConcurrency,
 	bool UseModTime = false)
 {
+	/// <inheritdoc cref="PerforceMaterializerOptions(String, RpcAgentWorkspace, String, IReadOnlyDictionary{String, String}, Int32, Int32, Int32, Int32, Int32, Boolean)"/>
+	public PerforceMaterializerOptions(string dirPath, RpcAgentWorkspace agentWorkspace, string? hostName = null, IReadOnlyDictionary<string, string>? envVars = null, int syncThreads = -1, int syncBatch = -1, int syncBatchSize = -1, int? maxFileConcurrency = null, int? maxFileScanConcurrency = null, bool bUseModTime = false)
+		: this(
+			  dirPath,
+			  agentWorkspace,
+		      hostName,
+			  envVars,
+			  syncThreads,
+			  syncBatch,
+			  syncBatchSize,
+			  maxFileConcurrency ?? GetDefaultThreadCount(4),
+			  maxFileScanConcurrency ?? GetDefaultThreadCount(4),
+			  bUseModTime)
+	{ }
+
+	static int GetDefaultThreadCount(int defaultValue)
+			=> Math.Min(defaultValue, Math.Max(Environment.ProcessorCount - 1, 1));
+
 	/// <summary>
 	/// Parse options from <see cref="RpcAgentWorkspace.Method" /> as HTTP query string
 	/// </summary>
@@ -47,6 +69,8 @@ public record PerforceMaterializerOptions(
 		const string SyncThreadsKey = "syncThreads";
 		const string SyncBatchKey = "syncBatch";
 		const string SyncBatchSizeKey = "syncBatchSize";
+		const string MaxFileConcurrencyKey = "maxFileConcurrency";
+		const string MaxFileScanConcurrencyKey = "maxFileScanConcurrency";
 		const string UseModTimeKey = "useModTime";
 		
 		PerforceMaterializerOptions options = new (dirPath, agentWorkspace);
@@ -75,6 +99,14 @@ public record PerforceMaterializerOptions(
 				{
 					options = options with { UseModTime = false };
 				}
+				if (Int32.TryParse(nameValues[MaxFileConcurrencyKey], out int maxConcurrency))
+				{
+					options = options with { MaxFileConcurrency = maxConcurrency };
+				}
+				if (Int32.TryParse(nameValues[MaxFileScanConcurrencyKey], out int maxScanConcurrency))
+				{
+					options = options with { MaxFileScanConcurrency = maxScanConcurrency };
+				}
 			}
 		}
 
@@ -272,7 +304,8 @@ public sealed class PerforceMaterializer : IWorkspaceMaterializer
 			using TelemetrySpan deleteSpan = _tracer.StartActiveSpan($"{nameof(PerforceMaterializer)}.{nameof(SyncAsync)}.Delete");
 			using ILoggerProgress status = _logger.BeginProgressScope("Deleting local workspace files...");
 			Stopwatch deleteTimer = Stopwatch.StartNew();
-			FileUtils.ForceDeleteDirectoryContents(_workspaceDir);
+			Parallel.ForEach(DirectoryReference.EnumerateFiles(_workspaceDir, "*", SearchOption.AllDirectories), new() { MaxDegreeOfParallelism = _options.MaxFileConcurrency }, FileUtils.ForceDeleteFile);
+			FileUtils.ForceDeleteDirectory(_workspaceDir);
 			status.Progress = $"({deleteTimer.Elapsed.TotalSeconds:0.0}s)";
 		}
 
