Use User1 instead of using hardcoded color for catastrophic failure

#rnx

#Synced-CL 49105898

diff --git a/Engine/Source/Programs/Horde/HordeDashboard/plugins/test-automation/testAutomationCommon.tsx b/Engine/Source/Programs/Horde/HordeDashboard/plugins/test-automation/testAutomationCommon.tsx
--- a/Engine/Source/Programs/Horde/HordeDashboard/plugins/test-automation/testAutomationCommon.tsx
+++ b/Engine/Source/Programs/Horde/HordeDashboard/plugins/test-automation/testAutomationCommon.tsx
@@ -132,7 +132,7 @@ export const getStatusColors = (): Map<StatusColor, string> => {
         [StatusColor.Warnings, dashboardStatusColors.get(StatusColor.Warnings)!],
         [StatusColor.Failure, dashboardStatusColors.get(StatusColor.Failure)!],
         [StatusColor.Skipped,  dashboard.darktheme ? "#c5c4c3ff" : "#a1a1a1ff"],
-        [StatusColor.Unspecified, "#8e4747ff"]
+        [StatusColor.Unspecified, dashboardStatusColors.get(StatusColor.User1)!]
     ]);
 
     return colors;
