Fix inconsistent height in user name widget

#rnx
#Synced-CL 49105898

diff --git a/Engine/Source/Programs/Horde/HordeDashboard/plugins/test-automation/testAutomationHealth.tsx b/Engine/Source/Programs/Horde/HordeDashboard/plugins/test-automation/testAutomationHealth.tsx
--- a/Engine/Source/Programs/Horde/HordeDashboard/plugins/test-automation/testAutomationHealth.tsx
+++ b/Engine/Source/Programs/Horde/HordeDashboard/plugins/test-automation/testAutomationHealth.tsx
@@ -29,7 +29,8 @@ const healthStyles = {
     userTag: {
         backgroundColor: booleanColors.get(false),
         padding: '2px 9px',
-        borderRadius: 11
+        borderRadius: 11,
+        height: 23
     },
     badge: {
         padding: '3px 12px',
@@ -204,7 +205,7 @@ const getHealthColumns = (items: TestHealthItem[], columns: TestHealthColumn[],
             isSorted: false,
             isSortedDescending: true,
             onRender: (item: TestHealthItem) => {
-                return <Stack style={{...healthStyles.userTag, backgroundColor: booleanColors.get(!item.owner), cursor: 'default'}} horizontalAlign="center" horizontal>
+                return <Stack style={{...healthStyles.userTag, backgroundColor: booleanColors.get(!item.owner), cursor: 'default', height: '100%'}} horizontalAlign="center" verticalAlign="center" horizontal>
                             <Text title={usersInfoCache.getInfo(item.owner)?.name} style={{textOverflow: 'ellipsis', overflow: 'hidden', whiteSpace: 'nowrap'}}>
                                 {!!item.owner && (usersInfoCache.getInfo(item.owner)?.name ?? <Spinner size={SpinnerSize.small}/>)}
                                 {!item.owner && "N/A"}
@@ -697,19 +698,21 @@ export const TestHealthView: React.FC<{ test: TestHealthItem, handler: TestDataH
                                 <Stack horizontal tokens={{childrenGap: 12}}>
                                     <Stack tokens={{childrenGap: 6}} style={{borderRadius: 4}} verticalAlign="start">
                                         <Label>Owner: </Label>
-                                        <Stack style={{...healthStyles.userTag, width: !test.owner? 46:undefined}} horizontal verticalAlign="center" horizontalAlign="center">
+                                        <Stack horizontal tokens={{childrenGap: 6}} wrap>
                                             {!!test.owner &&
-                                                <div style={{whiteSpace: 'nowrap'}}>
-                                                    <Text style={{cursor: 'pointer', color:textEmphasisColor()}} title="change owner"
+                                                <Stack style={healthStyles.userTag} horizontal verticalAlign="center">
+                                                    <Text style={{cursor: 'pointer', color:textEmphasisColor(), whiteSpace: 'nowrap'}} title="change owner"
                                                         onClick={() => setUserPicker({target: UserType.Owner, defaultUser: usersInfoCache.getInfo(test.owner)})}
                                                     >{usersInfoCache.getInfo(test.owner)?.name ?? <Spinner size={SpinnerSize.xSmall} />}</Text>
                                                     <FontIcon iconName="Cancel" style={{cursor: 'pointer', fontSize: 11, paddingLeft: 6}} title="remove" onClick={() => {test.owner = undefined}}/>
-                                                </div>
+                                                </Stack>
                                             }
                                             {!test.owner &&
-                                                <FontIcon iconName="AddTo" style={{cursor: 'pointer', fontSize: 15, padding: 2}} title="add an owner"
-                                                    onClick={() => setUserPicker({target: UserType.Owner})}
-                                                />
+                                                <Stack style={{...healthStyles.userTag, width: 46}} horizontal verticalAlign="center" horizontalAlign="center">
+                                                    <FontIcon iconName="AddTo" style={{cursor: 'pointer', fontSize: 15, padding: 2}} title="add an owner"
+                                                        onClick={() => setUserPicker({target: UserType.Owner})}
+                                                    />
+                                                </Stack>
                                             }
                                         </Stack>
                                     </Stack>
