Adjust test health notification message

#rnx
#Synced-CL 49105898

diff --git a/Engine/Source/Programs/Horde/Plugins/Build/HordeServer.Build.Tests/Jobs/TestData/TestDataTests.cs b/Engine/Source/Programs/Horde/Plugins/Build/HordeServer.Build.Tests/Jobs/TestData/TestDataTests.cs
--- a/Engine/Source/Programs/Horde/Plugins/Build/HordeServer.Build.Tests/Jobs/TestData/TestDataTests.cs
+++ b/Engine/Source/Programs/Horde/Plugins/Build/HordeServer.Build.Tests/Jobs/TestData/TestDataTests.cs
@@ -832,7 +832,7 @@ namespace HordeServer.Tests.Jobs.TestData
 			Assert.IsNull(reports[0].NotificationLastDate);
 
 			Assert.AreEqual(false, reports[1].IsHealthy);
-			Assert.AreEqual("Unstable", reports[1].State);
+			Assert.AreEqual("Unstable[3 stars]", reports[1].State);
 			Assert.AreEqual(0, reports[1].SuccessRate);
 			Assert.AreEqual(33, reports[1].CatastrophicFailureRate);
 			Assert.AreEqual(0, reports[1].RedundantErrorRate);
diff --git a/Engine/Source/Programs/Horde/Plugins/Build/HordeServer.Build/Jobs/TestData/TestHealthService.cs b/Engine/Source/Programs/Horde/Plugins/Build/HordeServer.Build/Jobs/TestData/TestHealthService.cs
--- a/Engine/Source/Programs/Horde/Plugins/Build/HordeServer.Build/Jobs/TestData/TestHealthService.cs
+++ b/Engine/Source/Programs/Horde/Plugins/Build/HordeServer.Build/Jobs/TestData/TestHealthService.cs
@@ -288,6 +288,10 @@ namespace HordeServer.Jobs.TestData
 
 			bool isHealthy = nStars >= (int)TestHealthState.Fair;
 			string state = Enum.GetName(typeof(TestHealthState), nStars) ?? "Unknown";
+			if (!isHealthy)
+			{
+				state += $"[{nStars} star{(nStars > 1 ? "s" : "")}]";
+			}
 
 			return new ComputedTestHealth()
 			{
