Add Display Meta to Test health title

#rnx

#Synced-CL 49105898

diff --git a/Engine/Source/Programs/Horde/HordeDashboard/plugins/test-automation/testAutomationHealth.tsx b/Engine/Source/Programs/Horde/HordeDashboard/plugins/test-automation/testAutomationHealth.tsx
--- a/Engine/Source/Programs/Horde/HordeDashboard/plugins/test-automation/testAutomationHealth.tsx
+++ b/Engine/Source/Programs/Horde/HordeDashboard/plugins/test-automation/testAutomationHealth.tsx
@@ -138,9 +138,6 @@ const getHealthColumns = (items: TestHealthItem[], columns: TestHealthColumn[],
             isSorted: false,
             isSortedDescending: true,
             onRender: (item: TestHealthItem) => {
-                const maskMetaString = item.displayMeta;
-                let title = item.name;
-                if (maskMetaString) title += ` - ${maskMetaString}`;
                 return <Stack onClick={() => onClickItem(item)} style={{cursor: 'pointer'}} horizontal>{renderTestTitle(item)}</Stack>
             },
             onColumnClick: (_, column) => onClickColumn(column),
@@ -848,10 +845,13 @@ const TestHealthDetailsModal: React.FC<{ test: TestHealthItem, handler: TestData
 
     test.subscribeUpdate();
 
+    let title = test.name;
+    if (test.displayMeta) title += ` - ${test.displayMeta}`;
+
     return <Modal isOpen={true} isBlocking={true} topOffsetFixed={true} styles={{ main: { padding: 8, width: 1134, hasBeenOpened: false, top: "80px", position: "absolute" } }} onDismiss={() => onDismiss()} >
             <Stack grow>
                 <Stack horizontal verticalAlign="center">
-                    <Stack><Text style={{ paddingLeft: 8, fontSize: 16, fontWeight: 600 }}>Health Details - {test.name}</Text></Stack>
+                    <Stack><Text style={{ paddingLeft: 8, fontSize: 16, fontWeight: 600 }}>Health Details - {title}</Text></Stack>
                     <Stack grow />
                     {!isCommitting && test.isTransient &&
                         <Stack style={{ paddingRight: 6 }} title="Revert changes">
